Copying //UE4/Orion-Staging to //UE4/Main (Source: //Orion/Dev-General @ 3196521)

diff --git a/Engine/Plugins/Runtime/GameplayAbilities/Source/GameplayAbilities/Public/Abilities/Tasks/AbilityTask_ApplyRootMotionConstantForce.h b/Engine/Plugins/Runtime/GameplayAbilities/Source/GameplayAbilities/Public/Abilities/Tasks/AbilityTask_ApplyRootMotionConstantForce.h
--- a/Engine/Plugins/Runtime/GameplayAbilities/Source/GameplayAbilities/Public/Abilities/Tasks/AbilityTask_ApplyRootMotionConstantForce.h
+++ b/Engine/Plugins/Runtime/GameplayAbilities/Source/GameplayAbilities/Public/Abilities/Tasks/AbilityTask_ApplyRootMotionConstantForce.h
@@ -1,6 +1,6 @@
 // Copyright 1998-2016 Epic Games, Inc. All Rights Reserved.
 #pragma once
-#include "AbilityTask.h"
+#include "AbilityTask_ApplyRootMotion_Base.h"
 #include "AbilityTask_ApplyRootMotionConstantForce.generated.h"
 
 DECLARE_DYNAMIC_MULTICAST_DELEGATE(FApplyRootMotionConstantForceDelegate);
@@ -11,21 +11,17 @@ class AActor;
  *	Applies force to character's movement
  */
 UCLASS(MinimalAPI)
-class UAbilityTask_ApplyRootMotionConstantForce : public UAbilityTask
+class UAbilityTask_ApplyRootMotionConstantForce : public UAbilityTask_ApplyRootMotion_Base
 {
 	GENERATED_UCLASS_BODY()
 
 	UPROPERTY(BlueprintAssignable)
 	FApplyRootMotionConstantForceDelegate OnFinish;
 
-	virtual void InitSimulatedTask(UGameplayTasksComponent& InGameplayTasksComponent) override;
-
 	/** Apply force to character's movement */
 	UFUNCTION(BlueprintCallable, Category = "Ability|Tasks", meta = (HidePin = "OwningAbility", DefaultToSelf = "OwningAbility", BlueprintInternalUseOnly = "TRUE"))
 	static UAbilityTask_ApplyRootMotionConstantForce* ApplyRootMotionConstantForce(UGameplayAbility* OwningAbility, FName TaskInstanceName, FVector WorldDirection, float Strength, float Duration, bool bIsAdditive, bool bDisableImpartingVelocityOnRemoval, UCurveFloat* StrengthOverTime);
 
-	virtual void Activate() override;
-
 	/** Tick function for this task, if bTickingTask == true */
 	virtual void TickTask(float DeltaTime) override;
 
@@ -34,8 +30,9 @@ class UAbilityTask_ApplyRootMotionConstantForce : public UAbilityTask
 
 protected:
 
-	UPROPERTY(Replicated)
-	FName ForceName;
+	virtual void SharedInitAndApply() override;
+
+protected:
 
 	UPROPERTY(Replicated)
 	FVector WorldDirection;
@@ -60,15 +57,4 @@ protected:
 	 */
 	UPROPERTY(Replicated)
 	UCurveFloat* StrengthOverTime;
-
-	uint16 RootMotionSourceID;
-	bool bIsFinished;
-	float StartTime;
-	float EndTime;
-
-	void SharedInitAndApply();
-
-	UPROPERTY()
-	UCharacterMovementComponent* MovementComponent;
-
 };
\ No newline at end of file
diff --git a/Engine/Plugins/Runtime/GameplayAbilities/Source/GameplayAbilities/Public/Abilities/Tasks/AbilityTask_ApplyRootMotionJumpForce.h b/Engine/Plugins/Runtime/GameplayAbilities/Source/GameplayAbilities/Public/Abilities/Tasks/AbilityTask_ApplyRootMotionJumpForce.h
--- a/Engine/Plugins/Runtime/GameplayAbilities/Source/GameplayAbilities/Public/Abilities/Tasks/AbilityTask_ApplyRootMotionJumpForce.h
+++ b/Engine/Plugins/Runtime/GameplayAbilities/Source/GameplayAbilities/Public/Abilities/Tasks/AbilityTask_ApplyRootMotionJumpForce.h
@@ -1,6 +1,6 @@
 // Copyright 1998-2016 Epic Games, Inc. All Rights Reserved.
 #pragma once
-#include "AbilityTask.h"
+#include "AbilityTask_ApplyRootMotion_Base.h"
 #include "AbilityTask_ApplyRootMotionJumpForce.generated.h"
 
 DECLARE_DYNAMIC_MULTICAST_DELEGATE(FApplyRootMotionJumpForceDelegate);
@@ -11,7 +11,7 @@ class AActor;
  *	Applies force to character's movement
  */
 UCLASS(MinimalAPI)
-class UAbilityTask_ApplyRootMotionJumpForce : public UAbilityTask
+class UAbilityTask_ApplyRootMotionJumpForce : public UAbilityTask_ApplyRootMotion_Base
 {
 	GENERATED_UCLASS_BODY()
 
@@ -27,8 +27,6 @@ class UAbilityTask_ApplyRootMotionJumpForce : public UAbilityTask
 	UFUNCTION()
 	void OnLandedCallback(const FHitResult& Hit);
 
-	virtual void InitSimulatedTask(UGameplayTasksComponent& InGameplayTasksComponent) override;
-
 	/** Apply force to character's movement */
 	UFUNCTION(BlueprintCallable, Category = "Ability|Tasks", meta = (HidePin = "OwningAbility", DefaultToSelf = "OwningAbility", BlueprintInternalUseOnly = "TRUE"))
 	static UAbilityTask_ApplyRootMotionJumpForce* ApplyRootMotionJumpForce(UGameplayAbility* OwningAbility, FName TaskInstanceName, FRotator Rotation, float Distance, float Height, float Duration, float MinimumLandedTriggerTime, bool bFinishOnLanded, UCurveVector* PathOffsetCurve, UCurveFloat* TimeMappingCurve);
@@ -43,8 +41,17 @@ class UAbilityTask_ApplyRootMotionJumpForce : public UAbilityTask
 
 protected:
 
-	UPROPERTY(Replicated)
-	FName ForceName;
+	virtual void SharedInitAndApply() override;
+
+	/**
+	* Work-around for OnLanded being called during bClientUpdating in movement replay code
+	* Don't want to trigger our Landed logic during a replay, so we wait until next frame
+	* If we don't, we end up removing root motion from a replay root motion set instead
+	* of the real one
+	*/
+	void TriggerLanded();
+
+protected:
 
 	UPROPERTY(Replicated)
 	FRotator Rotation;
@@ -77,23 +84,5 @@ protected:
 	UPROPERTY(Replicated)
 	UCurveFloat* TimeMappingCurve;
 
-	uint16 RootMotionSourceID;
-	bool bIsFinished;
-	float StartTime;
-	float EndTime;
-
-	/**
-	 * Work-around for OnLanded being called during bClientUpdating in movement replay code
-	 * Don't want to trigger our Landed logic during a replay, so we wait until next frame
-	 * If we don't, we end up removing root motion from a replay root motion set instead
-	 * of the real one
-	 */
-	void TriggerLanded();
 	bool bHasLanded;
-
-	void SharedInitAndApply();
-
-	UPROPERTY()
-	UCharacterMovementComponent* MovementComponent;
-
 };
\ No newline at end of file
diff --git a/Engine/Plugins/Runtime/GameplayAbilities/Source/GameplayAbilities/Public/Abilities/Tasks/AbilityTask_ApplyRootMotionMoveToActorForce.h b/Engine/Plugins/Runtime/GameplayAbilities/Source/GameplayAbilities/Public/Abilities/Tasks/AbilityTask_ApplyRootMotionMoveToActorForce.h
--- a/Engine/Plugins/Runtime/GameplayAbilities/Source/GameplayAbilities/Public/Abilities/Tasks/AbilityTask_ApplyRootMotionMoveToActorForce.h
+++ b/Engine/Plugins/Runtime/GameplayAbilities/Source/GameplayAbilities/Public/Abilities/Tasks/AbilityTask_ApplyRootMotionMoveToActorForce.h
@@ -1,6 +1,6 @@
 // Copyright 1998-2016 Epic Games, Inc. All Rights Reserved.
 #pragma once
-#include "AbilityTask.h"
+#include "AbilityTask_ApplyRootMotion_Base.h"
 #include "AbilityTask_ApplyRootMotionMoveToActorForce.generated.h"
 
 DECLARE_DYNAMIC_MULTICAST_DELEGATE_ThreeParams(FApplyRootMotionMoveToActorForceDelegate, bool, DestinationReached, bool, TimedOut, FVector, FinalTargetLocation);
@@ -22,20 +22,16 @@ enum class ERootMotionMoveToActorTargetOffsetType : uint8
  *	Applies force to character's movement
  */
 UCLASS(MinimalAPI)
-class UAbilityTask_ApplyRootMotionMoveToActorForce : public UAbilityTask
+class UAbilityTask_ApplyRootMotionMoveToActorForce : public UAbilityTask_ApplyRootMotion_Base
 {
 	GENERATED_UCLASS_BODY()
 
 	UPROPERTY(BlueprintAssignable)
 	FApplyRootMotionMoveToActorForceDelegate OnFinished;
 
-	virtual void InitSimulatedTask(UGameplayTasksComponent& InGameplayTasksComponent) override;
-
 	/** Apply force to character's movement */
 	UFUNCTION(BlueprintCallable, Category = "Ability|Tasks", meta = (HidePin = "OwningAbility", DefaultToSelf = "OwningAbility", BlueprintInternalUseOnly = "TRUE"))
-	static UAbilityTask_ApplyRootMotionMoveToActorForce* ApplyRootMotionMoveToActorForce(UGameplayAbility* OwningAbility, FName TaskInstanceName, AActor* TargetActor, FVector TargetLocationOffset, ERootMotionMoveToActorTargetOffsetType OffsetAlignment, float Duration, UCurveFloat* TargetLerpSpeedHorizontal, UCurveFloat* TargetLerpSpeedVertical, bool bSetNewMovementMode, EMovementMode MovementMode, bool bRestrictSpeedToExpected, UCurveVector* PathOffsetCurve, UCurveFloat* TimeMappingCurve, ERootMotionFinishVelocityMode VelocityOnFinishMode, FVector SetVelocityOnFinish, bool bDisableDestinationReachedInterrupt);
-
-	virtual void Activate() override;
+	static UAbilityTask_ApplyRootMotionMoveToActorForce* ApplyRootMotionMoveToActorForce(UGameplayAbility* OwningAbility, FName TaskInstanceName, AActor* TargetActor, FVector TargetLocationOffset, ERootMotionMoveToActorTargetOffsetType OffsetAlignment, float Duration, UCurveFloat* TargetLerpSpeedHorizontal, UCurveFloat* TargetLerpSpeedVertical, bool bSetNewMovementMode, EMovementMode MovementMode, bool bRestrictSpeedToExpected, UCurveVector* PathOffsetCurve, UCurveFloat* TimeMappingCurve, ERootMotionFinishVelocityMode VelocityOnFinishMode, FVector SetVelocityOnFinish, float ClampVelocityOnFinish, bool bDisableDestinationReachedInterrupt);
 
 	/** Tick function for this task, if bTickingTask == true */
 	virtual void TickTask(float DeltaTime) override;
@@ -45,14 +41,18 @@ class UAbilityTask_ApplyRootMotionMoveToActorForce : public UAbilityTask
 
 protected:
 
+	virtual void SharedInitAndApply() override;
+
 	bool UpdateTargetLocation(float DeltaTime);
 
 	void SetRootMotionTargetLocation(FVector NewTargetLocation);
 
 	FVector CalculateTargetOffset() const;
 
-	UPROPERTY(Replicated)
-	FName ForceName;
+	UFUNCTION()
+	void OnRep_TargetLocation();
+
+protected:
 
 	UPROPERTY(Replicated)
 	FVector StartLocation;
@@ -60,9 +60,6 @@ protected:
 	UPROPERTY(ReplicatedUsing=OnRep_TargetLocation)
 	FVector TargetLocation;
 
-	UFUNCTION()
-	void OnRep_TargetLocation();
-
 	UPROPERTY(Replicated)
 	AActor* TargetActor;
 
@@ -119,16 +116,9 @@ protected:
 	UPROPERTY(Replicated)
 	FVector SetVelocityOnFinish;
 
-	uint16 RootMotionSourceID;
-	EMovementMode PreviousMovementMode;
-
-	bool bIsFinished;
-	float StartTime;
-	float EndTime;
-
-	void SharedInitAndApply();
-
-	UPROPERTY()
-	UCharacterMovementComponent* MovementComponent;
+	/** If VelocityOnFinish mode is "ClampVelocity", character velocity is clamped to this value when root motion finishes */
+	UPROPERTY(Replicated)
+	float ClampVelocityOnFinish;
 
+	EMovementMode PreviousMovementMode;
 };
\ No newline at end of file
diff --git a/Engine/Plugins/Runtime/GameplayAbilities/Source/GameplayAbilities/Public/Abilities/Tasks/AbilityTask_ApplyRootMotionMoveToForce.h b/Engine/Plugins/Runtime/GameplayAbilities/Source/GameplayAbilities/Public/Abilities/Tasks/AbilityTask_ApplyRootMotionMoveToForce.h
--- a/Engine/Plugins/Runtime/GameplayAbilities/Source/GameplayAbilities/Public/Abilities/Tasks/AbilityTask_ApplyRootMotionMoveToForce.h
+++ b/Engine/Plugins/Runtime/GameplayAbilities/Source/GameplayAbilities/Public/Abilities/Tasks/AbilityTask_ApplyRootMotionMoveToForce.h
@@ -1,6 +1,6 @@
 // Copyright 1998-2016 Epic Games, Inc. All Rights Reserved.
 #pragma once
-#include "AbilityTask.h"
+#include "AbilityTask_ApplyRootMotion_Base.h"
 #include "AbilityTask_ApplyRootMotionMoveToForce.generated.h"
 
 DECLARE_DYNAMIC_MULTICAST_DELEGATE(FApplyRootMotionMoveToForceDelegate);
@@ -11,7 +11,7 @@ class AActor;
  *	Applies force to character's movement
  */
 UCLASS(MinimalAPI)
-class UAbilityTask_ApplyRootMotionMoveToForce : public UAbilityTask
+class UAbilityTask_ApplyRootMotionMoveToForce : public UAbilityTask_ApplyRootMotion_Base
 {
 	GENERATED_UCLASS_BODY()
 
@@ -21,13 +21,9 @@ class UAbilityTask_ApplyRootMotionMoveToForce : public UAbilityTask
 	UPROPERTY(BlueprintAssignable)
 	FApplyRootMotionMoveToForceDelegate OnTimedOutAndDestinationReached;
 
-	virtual void InitSimulatedTask(UGameplayTasksComponent& InGameplayTasksComponent) override;
-
 	/** Apply force to character's movement */
 	UFUNCTION(BlueprintCallable, Category = "Ability|Tasks", meta = (HidePin = "OwningAbility", DefaultToSelf = "OwningAbility", BlueprintInternalUseOnly = "TRUE"))
-	static UAbilityTask_ApplyRootMotionMoveToForce* ApplyRootMotionMoveToForce(UGameplayAbility* OwningAbility, FName TaskInstanceName, FVector TargetLocation, float Duration, bool bSetNewMovementMode, EMovementMode MovementMode, bool bRestrictSpeedToExpected, UCurveVector* PathOffsetCurve, ERootMotionFinishVelocityMode VelocityOnFinishMode, FVector SetVelocityOnFinish);
-
-	virtual void Activate() override;
+	static UAbilityTask_ApplyRootMotionMoveToForce* ApplyRootMotionMoveToForce(UGameplayAbility* OwningAbility, FName TaskInstanceName, FVector TargetLocation, float Duration, bool bSetNewMovementMode, EMovementMode MovementMode, bool bRestrictSpeedToExpected, UCurveVector* PathOffsetCurve, ERootMotionFinishVelocityMode VelocityOnFinishMode, FVector SetVelocityOnFinish, float ClampVelocityOnFinish);
 
 	/** Tick function for this task, if bTickingTask == true */
 	virtual void TickTask(float DeltaTime) override;
@@ -37,8 +33,9 @@ class UAbilityTask_ApplyRootMotionMoveToForce : public UAbilityTask
 
 protected:
 
-	UPROPERTY(Replicated)
-	FName ForceName;
+	virtual void SharedInitAndApply() override;
+
+protected:
 
 	UPROPERTY(Replicated)
 	FVector StartLocation;
@@ -72,16 +69,9 @@ protected:
 	UPROPERTY(Replicated)
 	FVector SetVelocityOnFinish;
 
-	uint16 RootMotionSourceID;
-	EMovementMode PreviousMovementMode;
-
-	bool bIsFinished;
-	float StartTime;
-	float EndTime;
-
-	void SharedInitAndApply();
-
-	UPROPERTY()
-	UCharacterMovementComponent* MovementComponent;
+	/** If VelocityOnFinish mode is "ClampVelocity", character velocity is clamped to this value when root motion finishes */
+	UPROPERTY(Replicated)
+	float ClampVelocityOnFinish;
 
+	EMovementMode PreviousMovementMode;
 };
\ No newline at end of file
diff --git a/Engine/Plugins/Runtime/GameplayAbilities/Source/GameplayAbilities/Public/Abilities/Tasks/AbilityTask_ApplyRootMotionRadialForce.h b/Engine/Plugins/Runtime/GameplayAbilities/Source/GameplayAbilities/Public/Abilities/Tasks/AbilityTask_ApplyRootMotionRadialForce.h
--- a/Engine/Plugins/Runtime/GameplayAbilities/Source/GameplayAbilities/Public/Abilities/Tasks/AbilityTask_ApplyRootMotionRadialForce.h
+++ b/Engine/Plugins/Runtime/GameplayAbilities/Source/GameplayAbilities/Public/Abilities/Tasks/AbilityTask_ApplyRootMotionRadialForce.h
@@ -1,6 +1,6 @@
 // Copyright 1998-2016 Epic Games, Inc. All Rights Reserved.
 #pragma once
-#include "AbilityTask.h"
+#include "AbilityTask_ApplyRootMotion_Base.h"
 #include "AbilityTask_ApplyRootMotionRadialForce.generated.h"
 
 DECLARE_DYNAMIC_MULTICAST_DELEGATE(FApplyRootMotionRadialForceDelegate);
@@ -11,20 +11,16 @@ class AActor;
  *	Applies force to character's movement
  */
 UCLASS(MinimalAPI)
-class UAbilityTask_ApplyRootMotionRadialForce : public UAbilityTask
+class UAbilityTask_ApplyRootMotionRadialForce : public UAbilityTask_ApplyRootMotion_Base
 {
 	GENERATED_UCLASS_BODY()
 
 	UPROPERTY(BlueprintAssignable)
 	FApplyRootMotionRadialForceDelegate OnFinish;
 
-	virtual void InitSimulatedTask(UGameplayTasksComponent& InGameplayTasksComponent) override;
-
 	/** Apply force to character's movement */
 	UFUNCTION(BlueprintCallable, Category = "Ability|Tasks", meta = (HidePin = "OwningAbility", DefaultToSelf = "OwningAbility", BlueprintInternalUseOnly = "TRUE"))
-	static UAbilityTask_ApplyRootMotionRadialForce* ApplyRootMotionRadialForce(UGameplayAbility* OwningAbility, FName TaskInstanceName, FVector Location, AActor* LocationActor, float Strength, float Duration, float Radius, bool bIsPush, bool bIsAdditive, bool bNoZForce, UCurveFloat* StrengthDistanceFalloff, UCurveFloat* StrengthOverTime, bool bUseFixedWorldDirection, FRotator FixedWorldDirection);
-
-	virtual void Activate() override;
+	static UAbilityTask_ApplyRootMotionRadialForce* ApplyRootMotionRadialForce(UGameplayAbility* OwningAbility, FName TaskInstanceName, FVector Location, AActor* LocationActor, float Strength, float Duration, float Radius, bool bIsPush, bool bIsAdditive, bool bNoZForce, UCurveFloat* StrengthDistanceFalloff, UCurveFloat* StrengthOverTime, bool bUseFixedWorldDirection, FRotator FixedWorldDirection, ERootMotionFinishVelocityMode VelocityOnFinishMode, FVector SetVelocityOnFinish, float ClampVelocityOnFinish);
 
 	/** Tick function for this task, if bTickingTask == true */
 	virtual void TickTask(float DeltaTime) override;
@@ -34,8 +30,9 @@ class UAbilityTask_ApplyRootMotionRadialForce : public UAbilityTask
 
 protected:
 
-	UPROPERTY(Replicated)
-	FName ForceName;
+	virtual void SharedInitAndApply() override;
+
+protected:
 
 	UPROPERTY(Replicated)
 	FVector Location;
@@ -84,14 +81,15 @@ protected:
 	UPROPERTY(Replicated)
 	FRotator FixedWorldDirection;
 
-	uint16 RootMotionSourceID;
-	bool bIsFinished;
-	float StartTime;
-	float EndTime;
-
-	void SharedInitAndApply();
+	/** What to do with character's Velocity when root motion finishes */
+	UPROPERTY(Replicated)
+	ERootMotionFinishVelocityMode VelocityOnFinishMode;
 
-	UPROPERTY()
-	UCharacterMovementComponent* MovementComponent;
+	/** If VelocityOnFinish mode is "SetVelocity", character velocity is set to this value when root motion finishes */
+	UPROPERTY(Replicated)
+	FVector SetVelocityOnFinish;
 
+	/** If VelocityOnFinish mode is "ClampVelocity", character velocity is clamped to this value when root motion finishes */
+	UPROPERTY(Replicated)
+	float ClampVelocityOnFinish;
 };
\ No newline at end of file
diff --git a/Engine/Plugins/Runtime/GameplayAbilities/Source/GameplayAbilities/Public/Abilities/Tasks/AbilityTask_ApplyRootMotion_Base.h b/Engine/Plugins/Runtime/GameplayAbilities/Source/GameplayAbilities/Public/Abilities/Tasks/AbilityTask_ApplyRootMotion_Base.h
new file mode 100644
--- /dev/null
+++ b/Engine/Plugins/Runtime/GameplayAbilities/Source/GameplayAbilities/Public/Abilities/Tasks/AbilityTask_ApplyRootMotion_Base.h
@@ -0,0 +1,35 @@
+// Copyright 1998-2016 Epic Games, Inc. All Rights Reserved.
+#pragma once
+#include "AbilityTask.h"
+#include "AbilityTask_ApplyRootMotion_Base.generated.h"
+
+
+UCLASS(MinimalAPI)
+class UAbilityTask_ApplyRootMotion_Base : public UAbilityTask
+{
+	GENERATED_UCLASS_BODY()
+
+	virtual void InitSimulatedTask(UGameplayTasksComponent& InGameplayTasksComponent) override;
+
+protected:
+
+	virtual void SharedInitAndApply() {};
+
+	void SetFinishVelocity(FName RootMotionSourceName, FVector FinishVelocity);
+	void ClampFinishVelocity(FName RootMotionSourceName, float VelocityClamp);
+
+protected:
+
+	UPROPERTY(Replicated)
+	FName ForceName;
+
+	UPROPERTY()
+	UCharacterMovementComponent* MovementComponent; 
+	
+	uint16 RootMotionSourceID;
+
+	bool bIsFinished;
+
+	float StartTime;
+	float EndTime;
+};
\ No newline at end of file
diff --git a/Engine/Plugins/Runtime/GameplayAbilities/Source/GameplayAbilities/Public/Abilities/Tasks/AbilityTask_WaitGameplayEffectApplied.h b/Engine/Plugins/Runtime/GameplayAbilities/Source/GameplayAbilities/Public/Abilities/Tasks/AbilityTask_WaitGameplayEffectApplied.h
--- a/Engine/Plugins/Runtime/GameplayAbilities/Source/GameplayAbilities/Public/Abilities/Tasks/AbilityTask_WaitGameplayEffectApplied.h
+++ b/Engine/Plugins/Runtime/GameplayAbilities/Source/GameplayAbilities/Public/Abilities/Tasks/AbilityTask_WaitGameplayEffectApplied.h
@@ -39,4 +39,7 @@ protected:
 
 	UPROPERTY()
 	UAbilitySystemComponent* ExternalOwner;
+
+	// If we are in the process of broadcasting and should not accept additional GE callbacks
+	bool Locked;
 };
\ No newline at end of file
diff --git a/Engine/Plugins/Runtime/GameplayAbilities/Source/GameplayAbilities/Public/AbilitySystemComponent.h b/Engine/Plugins/Runtime/GameplayAbilities/Source/GameplayAbilities/Public/AbilitySystemComponent.h
--- a/Engine/Plugins/Runtime/GameplayAbilities/Source/GameplayAbilities/Public/AbilitySystemComponent.h
+++ b/Engine/Plugins/Runtime/GameplayAbilities/Source/GameplayAbilities/Public/AbilitySystemComponent.h
@@ -348,6 +348,10 @@ class GAMEPLAYABILITIES_API UAbilitySystemComponent : public UGameplayTasksCompo
 	UFUNCTION(BlueprintCallable, BlueprintAuthorityOnly, Category = GameplayEffects)
 	void SetActiveGameplayEffectLevel(FActiveGameplayEffectHandle ActiveHandle, int32 NewLevel);
 
+	/** Updates the level of an already applied gameplay effect. The intention is that this is 'seemless' and doesnt behave like removing/reapplying */
+	UFUNCTION(BlueprintCallable, BlueprintAuthorityOnly, Category = GameplayEffects)
+	void SetActiveGameplayEffectLevelUsingQuery(FGameplayEffectQuery Query, int32 NewLevel);
+
 	// Not happy with this interface but don't see a better way yet. How should outside code (UI, etc) ask things like 'how much is this gameplay effect modifying my damage by'
 	// (most likely we want to catch this on the backend - when damage is applied we can get a full dump/history of how the number got to where it is. But still we may need polling methods like below (how much would my damage be)
 	UFUNCTION(BlueprintCallable, Category = GameplayEffects)
@@ -755,6 +759,9 @@ class GAMEPLAYABILITIES_API UAbilitySystemComponent : public UGameplayTasksCompo
 		return HasMatchingGameplayTag(GameplayCueTag);
 	}
 
+	/** Will initialize gameplay cue parameters with this ASC's Owner (Instigator) and AvatarActor (EffectCauser) */
+	virtual void InitDefaultGameplayCueParameters(FGameplayCueParameters& Parameters);
+
 	// ----------------------------------------------------------------------------------------------------------------
 
 	/**
@@ -1404,7 +1411,7 @@ public:
 protected:
 
 	/** Actually pushes the final attribute value to the attribute set's property. Should not be called by outside code since this does not go through the attribute aggregator system. */
-	void SetNumericAttribute_Internal(const FGameplayAttribute &Attribute, float NewFloatValue);
+	void SetNumericAttribute_Internal(const FGameplayAttribute &Attribute, float& NewFloatValue);
 
 	bool HasNetworkAuthorityToApplyGameplayEffect(FPredictionKey PredictionKey) const;
 
diff --git a/Engine/Plugins/Runtime/GameplayAbilities/Source/GameplayAbilities/Public/AbilitySystemGlobals.h b/Engine/Plugins/Runtime/GameplayAbilities/Source/GameplayAbilities/Public/AbilitySystemGlobals.h
--- a/Engine/Plugins/Runtime/GameplayAbilities/Source/GameplayAbilities/Public/AbilitySystemGlobals.h
+++ b/Engine/Plugins/Runtime/GameplayAbilities/Source/GameplayAbilities/Public/AbilitySystemGlobals.h
@@ -149,6 +149,11 @@ class GAMEPLAYABILITIES_API UAbilitySystemGlobals : public UObject
 	/** Global place to accumulate debug strings for ability system component. Used when we fill up client side debug string immediately, and then wait for server to send server strings */
 	TArray<FString>	AbilitySystemDebugStrings;
 
+
+	// Helper functions for applying global scaling to various ability system tasks. This isn't meant to be a shipping feature, but to help with debugging and interation via cvar AbilitySystem.GlobalAbilityScale.
+	static void NonShipping_ApplyGlobalAbilityScaler_Rate(float& Rate);
+	static void NonShipping_ApplyGlobalAbilityScaler_Duration(float& Duration);
+
 	// Global Tags
 
 	UPROPERTY()
@@ -233,6 +238,8 @@ class GAMEPLAYABILITIES_API UAbilitySystemGlobals : public UObject
 	/** Simple index-based accessor to the alias name for the specified gameplay mod evaluation channel, if any */
 	const FName& GetGameplayModEvaluationChannelAlias(int32 Index) const;
 
+	virtual TArray<FString> GetGameplayCueNotifyPaths() { return GameplayCueNotifyPaths; }
+
 protected:
 
 	virtual void InitAttributeDefaults();
@@ -285,7 +292,7 @@ protected:
 	UPROPERTY(config)
 	FStringAssetReference GlobalGameplayCueManagerName;
 
-	/** Look in these paths for GameplayCueNotifies */
+	/** Look in these paths for GameplayCueNotifies. These are your "always loaded" set. */
 	UPROPERTY(config)
 	TArray<FString>	GameplayCueNotifyPaths;
 
diff --git a/Engine/Plugins/Runtime/GameplayAbilities/Source/GameplayAbilities/Public/AttributeSet.h b/Engine/Plugins/Runtime/GameplayAbilities/Source/GameplayAbilities/Public/AttributeSet.h
--- a/Engine/Plugins/Runtime/GameplayAbilities/Source/GameplayAbilities/Public/AttributeSet.h
+++ b/Engine/Plugins/Runtime/GameplayAbilities/Source/GameplayAbilities/Public/AttributeSet.h
@@ -89,7 +89,7 @@ struct GAMEPLAYABILITIES_API FGameplayAttribute
 	/** Returns true if the variable associated with Property is of type FGameplayAttributeData or one of its subclasses */
 	static bool IsGameplayAttributeDataProperty(const UProperty* Property);
 
-	void SetNumericValueChecked(float NewValue, class UAttributeSet* Dest) const;
+	void SetNumericValueChecked(float& NewValue, class UAttributeSet* Dest) const;
 
 	float GetNumericValue(const UAttributeSet* Src) const;
 	float GetNumericValueChecked(const UAttributeSet* Src) const;
@@ -227,18 +227,19 @@ struct GAMEPLAYABILITIES_API FScalableFloat
 	FScalableFloat()
 		: Value(0.f)
 		, FinalCurve(nullptr)
+		, LocalCachedCurveID(INDEX_NONE)
 	{
 	}
 
 	FScalableFloat(float InInitialValue)
 		: Value(InInitialValue)
 		, FinalCurve(nullptr)
+		, LocalCachedCurveID(INDEX_NONE)
 	{
 	}
 
 	~FScalableFloat()
 	{
-		UnRegisterOnCurveTablePostReimport();
 	}
 
 public:
@@ -299,24 +300,15 @@ public:
 	/* Used to upgrade a float or int8/int16/int32 property into an FScalableFloat */
 	bool SerializeFromMismatchedTag(const FPropertyTag& Tag, FArchive& Ar);
 
-private:
-
-	/** Conditionally register the OnCurveTablePostReimport function with the re-import manager */
-	void RegisterOnCurveTablePostReimport() const;
-
-	/** Conditionally unregister the OnCurveTablePostReimport function from the re-import manager */
-	void UnRegisterOnCurveTablePostReimport() const;
+	static void InvalidateAllCachedCurves();
 
-#if WITH_EDITOR
-	/** Function to call when the curve table has been re-imported */
-	void OnCurveTablePostReimport(UObject* InObject, bool);
-
-	/** Handle to the delegate called when the curve table has been re-imported */
-	mutable FDelegateHandle OnCurveTablePostReimportHandle;
-#endif // WITH_EDITOR
+private:
 
 	// Cached direct pointer to RichCurve we should evaluate
 	mutable FRichCurve* FinalCurve;
+	mutable int32 LocalCachedCurveID;
+
+	static int32 GlobalCachedCurveID;
 };
 
 template<>
diff --git a/Engine/Plugins/Runtime/GameplayAbilities/Source/GameplayAbilities/Public/GameplayCueManager.h b/Engine/Plugins/Runtime/GameplayAbilities/Source/GameplayAbilities/Public/GameplayCueManager.h
--- a/Engine/Plugins/Runtime/GameplayAbilities/Source/GameplayAbilities/Public/GameplayCueManager.h
+++ b/Engine/Plugins/Runtime/GameplayAbilities/Source/GameplayAbilities/Public/GameplayCueManager.h
@@ -12,50 +12,63 @@
 #include "GameplayCueTranslator.h"
 #include "GameplayCueManager.generated.h"
 
-class UGameplayCueSet;
-
-/**
- *	
- *	Major TODO:
- *	
- *	-Remove LoadObjectLibraryFromPaths from base implementation. Just have project pass in ObjectLibrary, let it figure out how to make it.
- *		-On Bob's recommendation: in hopes of removing object libraries coupling with directory structure.
- *		-This becomes trickier with smaller projects/licensees. It would be nice to offer a path of least resistence for using gameplaycues (without having to implement a load object library function).
- *		-This gets even trickier though when dealing with having to reload that objecet library when assets are added/deleted/renamed. Ideally the GameplaycueManager handles all of this, but if we are getting
- *		an objectlibrary passed in, we don't know how to recreate it. Could add a delegate to invoke when we need the object library reloaded, but then ownership of that library feels pretty weird.
- *	
- *	-Async loading of all required GameplayCueNotifies
- *		-Currently async loaded on demand (first GameplayCueNotify is dropped)
- *		-Need way to enumerate all required GameplayCue tags (Interface? GameplayCueTags can come from GameplayEffects and GameplayAbilities, possibly more?)
- *		-Implemented UGameplayCueManager::BeginLoadingGameplayCueNotify
- *		
- *	-Better figure out how to handle instancing of GameplayCueNotifies
- *		-How to handle status type GameplayCueNotifies? (Code/blueprint will spawn stuff, needs to keep track and destroy them when GameplayCue is removed)
- *		-Instanced InstantiatedObjects is growing unbounded!
- *		 
- *	
- *	-Editor Workflow:
- *		-Make way to create new GameplayCueNotifies from UGameplayCueManager (details customization)
- *		-Jump to GameplayCueManager entry for GameplayCueTag (details customization)
- *		-Implement HandleAssetAdded and HandleAssetDeleted
- *			-Must make sure we update GameplayCueData/GameplayCueDataMap at appropriate times
- *				-On startup/begin PIE or try to do it as things change?
- *		
- *	-Overriding/forwarding: are we doing this right?
- *		-When can things override, when can they call parent implementations, etc
- *		-Do GameplayCueNotifies ever override GameplayCue Events or vice versa?
- *		
- *	-Take a pass on destruction
- *		-Make sure we are cleaning up GameplayCues when actors are destroyed
- *			(Register with Actor delegate or force game code to call EndGameplayCuesFor?)
- *	
- */
-
 DECLARE_DELEGATE_OneParam(FOnGameplayCueNotifySetLoaded, TArray<FStringAssetReference>);
-DECLARE_DELEGATE_RetVal_OneParam(bool, FShouldLoadGCNotifyDelegate, const FAssetData&);
+DECLARE_DELEGATE_OneParam(FGameplayCueProxyTick, float);
+DECLARE_DELEGATE_RetVal_TwoParams(bool, FShouldLoadGCNotifyDelegate, const FAssetData&, FName);
 
 class UObjectLibrary;
 
+/** An ObjectLibrary for the GameplayCue Notifies. Wraps 2 underlying UObjectLibraries plus options/delegates for how they are loaded */ 
+USTRUCT()
+struct FGameplayCueObjectLibrary
+{
+	GENERATED_BODY()
+	FGameplayCueObjectLibrary()
+	{
+		bHasBeenInitialized = false;
+	}
+
+	// Paths to search for
+	UPROPERTY()
+	TArray<FString> Paths;
+
+	// Callback for when load finishes
+	FOnGameplayCueNotifySetLoaded OnLoaded;
+
+	// Callback for "should I add this FAssetData to the set"
+	FShouldLoadGCNotifyDelegate ShouldLoad;
+
+	// Object library for actor based notifies
+	UPROPERTY()
+	UObjectLibrary* ActorObjectLibrary;
+
+	// Object library for object based notifies
+	UPROPERTY()
+	UObjectLibrary* StaticObjectLibrary;
+
+	// Priority to use if async loading
+	TAsyncLoadPriority AsyncPriority;
+
+	// Should we force a sync scan on the asset registry in order to discover asset data, or just use what is there?
+	UPROPERTY()
+	bool bShouldSyncScan;
+
+	// Should we start async loading everything that we find (that passes ShouldLoad delegate check)
+	UPROPERTY()
+	bool bShouldAsyncLoad;
+
+	// Should we sync load everything that we find (that passes ShouldLoad delegate check)
+	UPROPERTY()
+	bool bShouldSyncLoad;
+
+	// Set to put the loaded asset data into. If null we will use the global set (RuntimeGameplayCueObjectLibrary.CueSet)
+	UPROPERTY()
+	UGameplayCueSet* CueSet;
+
+	UPROPERTY()
+	bool bHasBeenInitialized;
+};
+
 UCLASS()
 class GAMEPLAYABILITIES_API UGameplayCueManager : public UDataAsset
 {
@@ -80,6 +93,8 @@ class GAMEPLAYABILITIES_API UGameplayCueManager : public UDataAsset
 
 	virtual void OnCreated();
 
+	virtual void OnEngineInitComplete();
+
 	/** Process a pending cue, return false if the cue should be rejected. */
 	virtual bool ProcessPendingCueExecute(FGameplayCuePendingExecute& PendingCue);
 
@@ -132,34 +147,80 @@ class GAMEPLAYABILITIES_API UGameplayCueManager : public UDataAsset
 
 	// -------------------------------------------------------------
 	//  Loading GameplayCueNotifies from ObjectLibraries
+	//
+	//	There are two libraries in the GameplayCueManager:
+	//	1. The RunTime object library, which is initialized with the "always loaded" paths, via UAbilitySystemGlobals::GetGameplayCueNotifyPaths()
+	//		-GC Notifies in this path are loaded at startup
+	//		-Everything loaded will go into the global gameplaycue set, which is where GC events will be routed to by default
+	//
+	//	2. The Editor object library, which is initialized with the "all valid" paths, via UGameplayCueManager::GetValidGameplayCuePaths()
+	//		-Only used in the editor.
+	//		-Never loads clases or scans directly itself. 
+	//		-Just reflect asset registry to show about "all gameplay cue notifies the game could know about"
+	//
 	// -------------------------------------------------------------
 
-	/** Loading soft refs to all GameplayCueNotifies */
-	void LoadObjectLibraryFromPaths(const TArray<FString>& Paths);
+	/** Returns the Ryntime cueset, which is the global cueset used in runtime, as opposed to the editor cue set which is used only when running the editor */
+	UGameplayCueSet* GetRuntimeCueSet();
 
-	UPROPERTY(transient)
-	class UGameplayCueSet* GlobalCueSet;
+	/** Called to setup and initialize the runtime library. The passed in paths will be scanned and added to the global gameplay cue set where appropriate */
+	void InitializeRuntimeObjectLibrary();
+
+	// Will return the Runtime cue set and the EditorCueSet, if the EditorCueSet is available. This is used mainly for handling asset created/deleted in the editor
+	TArray<UGameplayCueSet*> GetGlobalCueSets();
 	
+
+#if WITH_EDITOR
+	/** Called from editor to soft load all gameplay cue notifies for the GameplayCueEditor */
+	void InitializeEditorObjectLibrary();
+
+	/** Calling this will make the GC manager periodically refresh the EditorObjectLibrary until the asset registry is finished scanning */
+	void RequestPeriodicUpdateOfEditorObjectLibraryWhileWaitingOnAssetRegistry();
+
+	/** Get filenames of all GC notifies we know about (loaded or not). Useful for cooking */
+	void GetEditorObjectLibraryGameplayCueNotifyFilenames(TArray<FString>& Filenames) const;
+
+	/** Looks in the EditorObjectLibrary for a notify for this tag, if it finds it, it loads it and puts it in the RuntimeObjectLibrary so that it can be previewed in the editor */
+	void LoadNotifyForEditorPreview(FGameplayTag GameplayCueTag);
+
+	UGameplayCueSet* GetEditorCueSet();
+
+	FSimpleMulticastDelegate OnEditorObjectLibraryUpdated;
+	bool EditorObjectLibraryFullyInitialized;
+
+	FTimerHandle EditorPeriodicUpdateHandle;
+#endif
+
+protected:
+
+	virtual bool ShouldSyncScanRuntimeObjectLibraries() const;
+	virtual bool ShouldSyncLoadRuntimeObjectLibraries() const;
+	virtual bool ShouldAsyncLoadRuntimeObjectLibraries() const;
+
+	/** Refreshes the existing, already initialized, object libraries. */
+	void RefreshObjectLibraries();
+
+	/** Internal function to actually init the FGameplayCueObjectLibrary */
+	void InitObjectLibrary(FGameplayCueObjectLibrary& Library);
+
+	virtual TArray<FString> GetAlwaysLoadedGameplayCuePaths();
+
+	/** returns list of valid gameplay cue paths. Subclasses may override this to specify locations that aren't part of the "always loaded" LoadedPaths array */
+	virtual TArray<FString>	GetValidGameplayCuePaths() { return GetAlwaysLoadedGameplayCuePaths(); }
+
 	UPROPERTY(transient)
-	UObjectLibrary* GameplayCueNotifyActorObjectLibrary;
+	FGameplayCueObjectLibrary RuntimeGameplayCueObjectLibrary;
 
 	UPROPERTY(transient)
-	UObjectLibrary* GameplayCueNotifyStaticObjectLibrary;
+	FGameplayCueObjectLibrary EditorGameplayCueObjectLibrary;
+
+public:		
 
 	/** Called before loading any gameplay cue notifies from object libraries. Allows subclasses to skip notifies. */
 	virtual bool ShouldLoadGameplayCueAssetData(const FAssetData& Data) const { return true; }
-
-	// -------------------------------------------------------------
-	// Preload GameplayCue tags that we think we will need:
-	// -------------------------------------------------------------
-
-	void BeginLoadingGameplayCueNotify(FGameplayTag GameplayCueTag);
-
+	
 	int32 FinishLoadingGameplayCueNotifies();
 
-	/** Get filenames of all GC notifies we know about (loaded or not). Useful for cooking */
-	void GetGameplayCueNotifyFilenames(TArray<FString>& Filenames) const;
-
 	UPROPERTY(transient)
 	FStreamableManager	StreamableManager;
 	
@@ -172,8 +233,6 @@ class GAMEPLAYABILITIES_API UGameplayCueManager : public UDataAsset
 	virtual class UWorld* GetWorld() const override;
 
 #if WITH_EDITOR
-	/** Called from editor to soft load all gameplay cue notifies for the GameplayCueEditor */
-	void LoadAllGameplayCueNotifiesForEditor();
 
 	/** Handles updating an object library when a new asset is created */
 	void HandleAssetAdded(UObject *Object);
@@ -186,18 +245,14 @@ class GAMEPLAYABILITIES_API UGameplayCueManager : public UDataAsset
 
 	bool VerifyNotifyAssetIsInValidPath(FString Path);
 
-	/** returns list of valid gameplay cue paths. Subclasses may override this to specify locations that aren't part of the "always loaded" LoadedPaths array */
-	virtual TArray<FString>	GetValidGameplayCuePaths() { return LoadedPaths; }
-
-	bool RegisteredEditorCallbacks;
-
 	bool bAccelerationMapOutdated;
 
 	FOnGameplayCueNotifyChange	OnGameplayCueNotifyAddOrRemove;
 
+	/** Animation Preview Hacks */
 	static class USceneComponent* PreviewComponent;
-
 	static UWorld* PreviewWorld;
+	static FGameplayCueProxyTick PreviewProxyTick;
 #endif
 
 	static bool IsGameplayCueRecylingEnabled();
@@ -206,6 +261,14 @@ class GAMEPLAYABILITIES_API UGameplayCueManager : public UDataAsset
 
 	FGameplayCueTranslationManager	TranslationManager;
 
+	// -------------------------------------------------------------
+	//  Debugging Help
+	// -------------------------------------------------------------
+
+#if GAMEPLAYCUE_DEBUG
+	virtual FGameplayCueDebugInfo* GetDebugInfo(int32 Handle, bool Reset=false);
+#endif
+
 protected:
 
 #if WITH_EDITOR
@@ -213,16 +276,6 @@ protected:
 	virtual void ReloadObjectLibrary(UWorld* World, const UWorld::InitializationValues IVS);
 #endif
 
-	void LoadObjectLibrary_Internal();
-
-	void InitObjectLibraries(TArray<FString> Paths, UObjectLibrary* ActorObjectLibrary, UObjectLibrary* StaticObjectLibrary, FOnGameplayCueNotifySetLoaded OnLoaded, FShouldLoadGCNotifyDelegate ShouldLoad = FShouldLoadGCNotifyDelegate());
-
-	/** Async load all gameplay cue notify classes that we find in ::InitObjectLibraries */
-	virtual bool ShouldAsyncLoadAtStartup() const { return true; }
-
-	/** Sync load all gameplay cue notify classes that we find in ::InitObjectLibraries */
-	virtual bool ShouldSyncLoadAtStartup() const { return IsRunningCommandlet(); }
-
 	void BuildCuesToAddToGlobalSet(const TArray<FAssetData>& AssetDataList, FName TagPropertyName, TArray<struct FGameplayCueReferencePair>& OutCuesToAdd, TArray<FStringAssetReference>& OutAssetsToLoad, FShouldLoadGCNotifyDelegate = FShouldLoadGCNotifyDelegate());
 
 	/** The cue manager has a tendency to produce a lot of RPCs. This logs out when we are attempting to fire more RPCs than will actually go off */
@@ -240,8 +293,6 @@ protected:
 	UPROPERTY(transient)
 	TArray<AGameplayCueNotify_Actor*> GameplayCueClassesForPreallocation;
 
-	TArray<FString>	LoadedPaths;
-
 	/** List of gameplay cue executes that haven't been processed yet */
 	UPROPERTY(transient)
 	TArray<FGameplayCuePendingExecute> PendingExecuteCues;
diff --git a/Engine/Plugins/Runtime/GameplayAbilities/Source/GameplayAbilities/Public/GameplayCueSet.h b/Engine/Plugins/Runtime/GameplayAbilities/Source/GameplayAbilities/Public/GameplayCueSet.h
--- a/Engine/Plugins/Runtime/GameplayAbilities/Source/GameplayAbilities/Public/GameplayCueSet.h
+++ b/Engine/Plugins/Runtime/GameplayAbilities/Source/GameplayAbilities/Public/GameplayCueSet.h
@@ -65,8 +65,12 @@ class GAMEPLAYABILITIES_API UGameplayCueSet : public UDataAsset
 	/** Returns filenames of everything we know about (loaded or not) */
 	virtual void GetFilenames(TArray<FString>& Filenames) const;
 
+	virtual void GetStringAssetReferences(TArray<FStringAssetReference>& List) const;
+
 #if WITH_EDITOR
 
+	void CopyCueDataToSetForEditorPreview(FGameplayTag Tag, UGameplayCueSet* DestinationSet);
+
 	/** Updates an existing cue */
 	virtual void UpdateCueByStringRefs(const FStringAssetReference& CueToRemove, FString NewPath);
 
diff --git a/Engine/Plugins/Runtime/GameplayAbilities/Source/GameplayAbilities/Public/GameplayEffect.h b/Engine/Plugins/Runtime/GameplayAbilities/Source/GameplayAbilities/Public/GameplayEffect.h
--- a/Engine/Plugins/Runtime/GameplayAbilities/Source/GameplayAbilities/Public/GameplayEffect.h
+++ b/Engine/Plugins/Runtime/GameplayAbilities/Source/GameplayAbilities/Public/GameplayEffect.h
@@ -856,6 +856,9 @@ struct GAMEPLAYABILITIES_API FGameplayEffectSpec
 
 	FGameplayEffectSpec& operator=(const FGameplayEffectSpec& Other);
 
+	// Can be called manually but it is  preferred to use the 3 parameter constructor
+	void Initialize(const UGameplayEffect* InDef, const FGameplayEffectContextHandle& InEffectContext, float Level = FGameplayEffectConstants::INVALID_LEVEL);
+
 	/**
 	 * Determines if the spec has capture specs with valid captures for all of the specified definitions.
 	 * 
diff --git a/Engine/Plugins/Runtime/GameplayAbilities/Source/GameplayAbilities/Public/GameplayEffectExecutionCalculation.h b/Engine/Plugins/Runtime/GameplayAbilities/Source/GameplayAbilities/Public/GameplayEffectExecutionCalculation.h
--- a/Engine/Plugins/Runtime/GameplayAbilities/Source/GameplayAbilities/Public/GameplayEffectExecutionCalculation.h
+++ b/Engine/Plugins/Runtime/GameplayAbilities/Source/GameplayAbilities/Public/GameplayEffectExecutionCalculation.h
@@ -20,8 +20,8 @@ public:
 
 	// Constructors
 	FGameplayEffectCustomExecutionParameters();
-	FGameplayEffectCustomExecutionParameters(FGameplayEffectSpec& InOwningSpec, const TArray<FGameplayEffectExecutionScopedModifierInfo>& InScopedMods, UAbilitySystemComponent* InTargetAbilityComponent, const FGameplayTagContainer& InPassedIntags);
-	FGameplayEffectCustomExecutionParameters(FGameplayEffectSpec& InOwningSpec, const TArray<FGameplayEffectExecutionScopedModifierInfo>& InScopedMods, UAbilitySystemComponent* InTargetAbilityComponent, const FGameplayTagContainer& InPassedIntags, const TArray<FActiveGameplayEffectHandle>& InIgnoreHandles);
+	FGameplayEffectCustomExecutionParameters(FGameplayEffectSpec& InOwningSpec, const TArray<FGameplayEffectExecutionScopedModifierInfo>& InScopedMods, UAbilitySystemComponent* InTargetAbilityComponent, const FGameplayTagContainer& InPassedIntags, const FPredictionKey& InPredictionKey);
+	FGameplayEffectCustomExecutionParameters(FGameplayEffectSpec& InOwningSpec, const TArray<FGameplayEffectExecutionScopedModifierInfo>& InScopedMods, UAbilitySystemComponent* InTargetAbilityComponent, const FGameplayTagContainer& InPassedIntags, const FPredictionKey& InPredictionKey, const TArray<FActiveGameplayEffectHandle>& InIgnoreHandles);
 
 	/** Simple accessor to owning gameplay spec */
 	const FGameplayEffectSpec& GetOwningSpec() const;
@@ -38,8 +38,9 @@ public:
 	/** Simple accessor to the Passed In Tags to this execution */
 	const FGameplayTagContainer& GetPassedInTags() const;
 
-	/** Simple accessor to the GE handles to ignore in this execution. */
 	TArray<FActiveGameplayEffectHandle> GetIgnoreHandles() const;
+	
+	FPredictionKey GetPredictionKey() const;
 
 	/**
 	 * Attempts to calculate the magnitude of a captured attribute given the specified parameters. Can fail if the gameplay spec doesn't have
@@ -114,8 +115,9 @@ private:
 	/** The extra tags that were passed in to this execution */
 	FGameplayTagContainer PassedInTags;
 
-	/** Any mods with one of these handles will be ignored during evaluation */
 	TArray<FActiveGameplayEffectHandle> IgnoreHandles;
+	
+	FPredictionKey PredictionKey;
 };
 
 /** Struct representing the output of a custom gameplay effect execution. */
diff --git a/Engine/Plugins/Runtime/GameplayAbilities/Source/GameplayAbilities/Public/GameplayTagResponseTable.h b/Engine/Plugins/Runtime/GameplayAbilities/Source/GameplayAbilities/Public/GameplayTagResponseTable.h
--- a/Engine/Plugins/Runtime/GameplayAbilities/Source/GameplayAbilities/Public/GameplayTagResponseTable.h
+++ b/Engine/Plugins/Runtime/GameplayAbilities/Source/GameplayAbilities/Public/GameplayTagResponseTable.h
@@ -92,6 +92,8 @@ protected:
 
 	TMap< TWeakObjectPtr<UAbilitySystemComponent>, TArray< FGameplayTagResponseAppliedInfo> > RegisteredASCs;
 
+	float LastASCPurgeTime;
+
 	void Remove(UAbilitySystemComponent* ASC, TArray<FActiveGameplayEffectHandle>& Handles);
 
 	void AddOrUpdate(UAbilitySystemComponent* ASC, const TArray<TSubclassOf<UGameplayEffect> >& ResponseGameplayEffects, int32 TotalCount, TArray<FActiveGameplayEffectHandle>& Handles);
diff --git a/Engine/Plugins/Runtime/GameplayAbilities/Source/GameplayAbilitiesEditor/Private/AttributeDetails.cpp b/Engine/Plugins/Runtime/GameplayAbilities/Source/GameplayAbilitiesEditor/Private/AttributeDetails.cpp
--- a/Engine/Plugins/Runtime/GameplayAbilities/Source/GameplayAbilitiesEditor/Private/AttributeDetails.cpp
+++ b/Engine/Plugins/Runtime/GameplayAbilities/Source/GameplayAbilitiesEditor/Private/AttributeDetails.cpp
@@ -355,7 +355,7 @@ void FScalableFloatDetails::CustomizeHeader( TSharedRef<class IPropertyHandle> S
 				.HAlign(HAlign_Fill)
 				.Padding(2.f, 0.f, 2.f, 0.f)
 				[
-					CurveTableProperty->CreatePropertyValueWidget()
+					CreateCurveTableWidget()
 				]
 
 				+SHorizontalBox::Slot()
@@ -371,6 +371,7 @@ void FScalableFloatDetails::CustomizeHeader( TSharedRef<class IPropertyHandle> S
 					[
 						SNew(STextBlock)
 						.Text(this, &FScalableFloatDetails::GetRowNameComboBoxContentText)
+						.ToolTipText( this, &FScalableFloatDetails::GetRowNameComboBoxContentText)
 					]
 				]
 
@@ -414,6 +415,10 @@ void FScalableFloatDetails::CustomizeHeader( TSharedRef<class IPropertyHandle> S
 
 END_SLATE_FUNCTION_BUILD_OPTIMIZATION
 
+TSharedRef<SWidget> FScalableFloatDetails::CreateCurveTableWidget()
+{
+	return CurveTableProperty->CreatePropertyValueWidget();
+}
 
 void FScalableFloatDetails::OnCurveTableChanged()
 {
diff --git a/Engine/Plugins/Runtime/GameplayAbilities/Source/GameplayAbilitiesEditor/Private/GameplayAbilitiesEditorModule.cpp b/Engine/Plugins/Runtime/GameplayAbilities/Source/GameplayAbilitiesEditor/Private/GameplayAbilitiesEditorModule.cpp
--- a/Engine/Plugins/Runtime/GameplayAbilities/Source/GameplayAbilitiesEditor/Private/GameplayAbilitiesEditorModule.cpp
+++ b/Engine/Plugins/Runtime/GameplayAbilities/Source/GameplayAbilitiesEditor/Private/GameplayAbilitiesEditorModule.cpp
@@ -156,6 +156,9 @@ void FGameplayAbilitiesEditorModule::StartupModule()
 	IGameplayAbilitiesModule::Get().CallOrRegister_OnAbilitySystemGlobalsReady(FSimpleMulticastDelegate::FDelegate::CreateLambda([] {
 		FGameplayAbilitiesEditorModule::RegisterDebuggingCallbacks();
 	}));
+	
+	// Invalidate all internal cacheing of FRichCurve* in FScalableFlaots when a UCurveTable is reimported
+	FReimportManager::Instance()->OnPostReimport().AddLambda([](UObject* InObject, bool b){ FScalableFloat::InvalidateAllCachedCurves(); });
 }
 
 void FGameplayAbilitiesEditorModule::HandleNotify_OpenAssetInEditor(FString AssetName, int AssetType)
diff --git a/Engine/Plugins/Runtime/GameplayAbilities/Source/GameplayAbilitiesEditor/Private/GameplayCueTagDetails.cpp b/Engine/Plugins/Runtime/GameplayAbilities/Source/GameplayAbilitiesEditor/Private/GameplayCueTagDetails.cpp
--- a/Engine/Plugins/Runtime/GameplayAbilities/Source/GameplayAbilitiesEditor/Private/GameplayCueTagDetails.cpp
+++ b/Engine/Plugins/Runtime/GameplayAbilities/Source/GameplayAbilitiesEditor/Private/GameplayCueTagDetails.cpp
@@ -174,19 +174,19 @@ bool FGameplayCueTagDetails::UpdateNotifyList()
 	{
 		uint8 EnumVal;
 		GameplayTagProperty->GetValue(EnumVal);
-
-		UGameplayCueManager* CueManager = UAbilitySystemGlobals::Get().GetGameplayCueManager();
-		if (CueManager && CueManager->GlobalCueSet)
+		if (UGameplayCueManager* CueManager = UAbilitySystemGlobals::Get().GetGameplayCueManager())
 		{
-			if (int32* idxPtr = CueManager->GlobalCueSet->GameplayCueDataMap.Find(*Tag))
+			if (UGameplayCueSet* CueSet = CueManager->GetEditorCueSet())
 			{
-				int32 idx = *idxPtr;
-				if (CueManager->GlobalCueSet->GameplayCueData.IsValidIndex(idx))
+				if (int32* idxPtr = CueSet->GameplayCueDataMap.Find(*Tag))
 				{
-					FGameplayCueNotifyData& Data = CueManager->GlobalCueSet->GameplayCueData[*idxPtr];
-
-					TSharedRef<FStringAssetReference> Item(MakeShareable(new FStringAssetReference(Data.GameplayCueNotifyObj)));
-					NotifyList.Add(Item);
+					int32 idx = *idxPtr;
+					if (CueSet->GameplayCueData.IsValidIndex(idx))
+					{
+						FGameplayCueNotifyData& Data = CueSet->GameplayCueData[*idxPtr];
+						TSharedRef<FStringAssetReference> Item(MakeShareable(new FStringAssetReference(Data.GameplayCueNotifyObj)));
+						NotifyList.Add(Item);
+					}
 				}
 			}
 		}
@@ -202,7 +202,10 @@ FGameplayTag* FGameplayCueTagDetails::GetTag() const
 	if (GameplayTagProperty.IsValid())
 	{
 		GameplayTagProperty->AccessRawData(RawStructData);
-		Tag = (FGameplayTag*)(RawStructData[0]);
+		if (RawStructData.Num() > 0)
+		{
+			Tag = (FGameplayTag*)(RawStructData[0]);
+		}
 	}
 	return Tag;
 }
diff --git a/Engine/Plugins/Runtime/GameplayAbilities/Source/GameplayAbilitiesEditor/Private/SGameplayCueEditor.cpp b/Engine/Plugins/Runtime/GameplayAbilities/Source/GameplayAbilitiesEditor/Private/SGameplayCueEditor.cpp
--- a/Engine/Plugins/Runtime/GameplayAbilities/Source/GameplayAbilitiesEditor/Private/SGameplayCueEditor.cpp
+++ b/Engine/Plugins/Runtime/GameplayAbilities/Source/GameplayAbilitiesEditor/Private/SGameplayCueEditor.cpp
@@ -22,6 +22,7 @@
 #include "AssetToolsModule.h"
 #include "GameplayTagsEditorModule.h"
 #include "Editor/LevelEditor/Public/LevelEditor.h"
+#include "AssetRegistryModule.h"
 
 #define LOCTEXT_NAMESPACE "SGameplayCueEditor"
 #include "SGameplayCueEditor_Picker.h"
@@ -95,9 +96,6 @@ private:
 	/** Show only GC Tags that explicitly exist. If a.b.c is in the dictionary, don't show a.b as a distinct tag */
 	bool bShowOnlyLeafTags;
 
-	/** Flag to avoid loading all gameplaycues multiple times (wasted time creating and checking object libraries that are up to date) */
-	bool bHasLoadedAllGameplayCues;
-
 	/** Track when filter state is dirty, so that we only rebuilt view when it has changed, once menu is closed. */
 	bool bFilterIDsDirty;
 
@@ -643,14 +641,20 @@ BEGIN_SLATE_FUNCTION_BUILD_OPTIMIZATION
 
 	bool FindGameplayCueNotifyObj(UGameplayCueManager* CueManager, TSharedPtr< FGCTreeItem >& Item)
 	{
-		if (CueManager && CueManager->GlobalCueSet && Item->GameplayCueTag.IsValid())
+		if (CueManager && Item->GameplayCueTag.IsValid())
 		{
-			if (int32* idxPtr = CueManager->GlobalCueSet->GameplayCueDataMap.Find(Item->GameplayCueTag))
+			UGameplayCueSet* EditorSet = CueManager->GetEditorCueSet();
+			if (EditorSet == nullptr)
+			{
+				return false;
+			}
+
+			if (int32* idxPtr = EditorSet->GameplayCueDataMap.Find(Item->GameplayCueTag))
 			{
 				int32 idx = *idxPtr;
-				if (CueManager->GlobalCueSet->GameplayCueData.IsValidIndex(idx))
+				if (EditorSet->GameplayCueData.IsValidIndex(idx))
 				{
-					FGameplayCueNotifyData& Data = CueManager->GlobalCueSet->GameplayCueData[idx];
+					FGameplayCueNotifyData& Data = EditorSet->GameplayCueData[idx];
 					Item->GameplayCueNotifyObj = Data.GameplayCueNotifyObj;
 					return true;
 				}
@@ -676,8 +680,8 @@ BEGIN_SLATE_FUNCTION_BUILD_OPTIMIZATION
 		}
 
 #if STATS
-		FString PerfMessage = FString::Printf(TEXT("SGameplayCueEditor::UpdateGameplayCueListItems"));
-		SCOPE_LOG_TIME_IN_SECONDS(*PerfMessage, nullptr)
+		//FString PerfMessage = FString::Printf(TEXT("SGameplayCueEditor::UpdateGameplayCueListItems"));
+		//SCOPE_LOG_TIME_IN_SECONDS(*PerfMessage, nullptr)
 #endif
 		double FindGameplayCueNotifyObjTime = 0.f;
 		double AddTranslationTagsTime = 0.f;
@@ -698,13 +702,6 @@ BEGIN_SLATE_FUNCTION_BUILD_OPTIMIZATION
 		TArray<FString> SearchStrings;
 		FullSearchString.ParseIntoArrayWS(SearchStrings);
 
-		if (bHasLoadedAllGameplayCues == false)
-		{
-			SCOPE_LOG_TIME_IN_SECONDS(TEXT(" LoadAllGameplayCueNotifiesForEditor"), nullptr)
-			bHasLoadedAllGameplayCues = true;
-			CueManager->LoadAllGameplayCueNotifiesForEditor();
-		}
-
 		// ------------------------------------------------------
 		if (bShowAllOverrides)
 		{
@@ -722,7 +719,7 @@ BEGIN_SLATE_FUNCTION_BUILD_OPTIMIZATION
 		FGameplayTagContainer AllGameplayCueTags;
 		{
 			FString RequestGameplayTagChildrenPerfMessage = FString::Printf(TEXT(" RequestGameplayTagChildren"));
-			SCOPE_LOG_TIME_IN_SECONDS(*RequestGameplayTagChildrenPerfMessage, nullptr)
+			//SCOPE_LOG_TIME_IN_SECONDS(*RequestGameplayTagChildrenPerfMessage, nullptr)
 
 			if (bShowOnlyLeafTags)
 			{
@@ -735,7 +732,7 @@ BEGIN_SLATE_FUNCTION_BUILD_OPTIMIZATION
 		}
 
 		// Create data structs for widgets
-		for (FGameplayTag ThisGameplayCueTag : AllGameplayCueTags)
+		for (const FGameplayTag& ThisGameplayCueTag : AllGameplayCueTags)
 		{
 			if (SearchStrings.Num() > 0)
 			{
@@ -812,7 +809,7 @@ BEGIN_SLATE_FUNCTION_BUILD_OPTIMIZATION
 
 		{
 			FString RequestTreeRefreshMessage = FString::Printf(TEXT("  RequestTreeRefresh"));
-			SCOPE_LOG_TIME_IN_SECONDS(*RequestTreeRefreshMessage, nullptr)
+			//SCOPE_LOG_TIME_IN_SECONDS(*RequestTreeRefreshMessage, nullptr)
 
 			if (GameplayCueTreeView.IsValid())
 			{
@@ -1005,6 +1002,49 @@ BEGIN_SLATE_FUNCTION_BUILD_OPTIMIZATION
 			SelectedUniqueID = INDEX_NONE;
 		}
 	}
+
+	void HandleOverrideTypeChange(bool NewValue)
+	{
+		bShowAllOverrides = NewValue;
+		UpdateGameplayCueListItems();
+	}
+	
+	TSharedRef<SWidget> OnGetShowOverrideTypeMenu()
+	{
+		FMenuBuilder MenuBuilder( true, NULL );
+
+		FUIAction YesAction( FExecuteAction::CreateSP( this, &SGameplayCueEditorImpl::HandleOverrideTypeChange, true ) );
+		MenuBuilder.AddMenuEntry( GetOverrideTypeDropDownText_Explicit(true), LOCTEXT("GameplayCueEditor", "Show ALL POSSIBLE tags for overrides: including Tags that could exist but currently dont"), FSlateIcon(), YesAction );
+
+		FUIAction NoAction( FExecuteAction::CreateSP( this, &SGameplayCueEditorImpl::HandleOverrideTypeChange, false ) );
+		MenuBuilder.AddMenuEntry( GetOverrideTypeDropDownText_Explicit(false), LOCTEXT("GameplayCueEditor", "ONLY show tags for overrides that exist/have been setup."), FSlateIcon(), NoAction );
+
+		return MenuBuilder.MakeWidget();
+	}
+
+	FText GetOverrideTypeDropDownText() const
+	{
+		return GetOverrideTypeDropDownText_Explicit(bShowAllOverrides);
+	}
+	
+	FText GetOverrideTypeDropDownText_Explicit(bool ShowAll) const
+	{
+		if (ShowAll)
+		{
+			return LOCTEXT("ShowAllOverrides_Tooltip_Yes", "Show all possible overrides");
+		}
+
+		return LOCTEXT("ShowAllOverrides_Tooltip_No", "Show only existing overrides");
+	}
+
+	EVisibility GetWaitingOnAssetRegistryVisiblity() const
+	{
+		if (UGameplayCueManager* CueManager = UAbilitySystemGlobals::Get().GetGameplayCueManager())
+		{
+			return CueManager->EditorObjectLibraryFullyInitialized ? EVisibility::Collapsed : EVisibility::Visible;
+		}
+		return EVisibility::Visible;
+	}
 };
 
 bool SGameplayCueEditorImpl::bSuppressCueViewUpdate = false;
@@ -1021,12 +1061,13 @@ void SGameplayCueEditorImpl::Construct(const FArguments& InArgs)
 	if (CueManager)
 	{
 		CueManager->OnGameplayCueNotifyAddOrRemove.AddSP(this, &SGameplayCueEditorImpl::OnPropertyValueChanged);
+		CueManager->OnEditorObjectLibraryUpdated.AddSP(this, &SGameplayCueEditorImpl::UpdateGameplayCueListItems);
+		CueManager->RequestPeriodicUpdateOfEditorObjectLibraryWhileWaitingOnAssetRegistry();
 	}
 
 	bShowAll = true;
 	bShowAllOverrides = false;
 	bShowOnlyLeafTags = true;
-	bHasLoadedAllGameplayCues = false;
 	bFilterIDsDirty = false;
 
 	bool CanAddFromINI = UGameplayTagsManager::Get().ShouldImportTagsFromINI(); // We only support adding new tags to the ini files.
@@ -1155,6 +1196,18 @@ void SGameplayCueEditorImpl::Construct(const FArguments& InArgs)
 			.Padding(2.0f, 2.0f)
 			.AutoWidth()
 			[
+				SNew( SComboButton )
+				.OnGetMenuContent( this, &SGameplayCueEditorImpl::OnGetShowOverrideTypeMenu )
+				.VAlign( VAlign_Center )
+				.ContentPadding(2)
+				.ButtonContent()
+				[
+					SNew( STextBlock )
+					.ToolTipText( LOCTEXT("ShowOverrideType", "Toggles how we display overrides. Either show the existing overrides, or show possible overrides") )
+					.Text( this, &SGameplayCueEditorImpl::GetOverrideTypeDropDownText )
+				]
+
+				/*
 				SNew(SCheckBox)
 				.IsChecked(this, &SGameplayCueEditorImpl::HandleShowAllOverridesCheckBoxIsChecked)
 				.OnCheckStateChanged(this, &SGameplayCueEditorImpl::HandleShowAllOverridesCheckedStateChanged)
@@ -1162,6 +1215,9 @@ void SGameplayCueEditorImpl::Construct(const FArguments& InArgs)
 					SNew(STextBlock)
 					.Text(LOCTEXT("ShowAllOverrides", "Show all possible overrides"))
 				]
+				*/
+
+
 			]
 
 #if GAMEPLAYCUEEDITOR_SHOW_ONLY_LEAFTAGS_OPTION
@@ -1177,9 +1233,23 @@ void SGameplayCueEditorImpl::Construct(const FArguments& InArgs)
 					.Text(LOCTEXT("ShowLeafTagsOnly", "Show leaf tags only"))
 				]
 			]
-#endif
+#endif			
+		]
+		
+		// ---------------------------------------------------------------
 
-			
+		+ SVerticalBox::Slot()
+		.AutoHeight()
+		[
+			SNew(SHorizontalBox)
+			+ SHorizontalBox::Slot()
+			.Padding(2.0f)
+			.AutoWidth()
+			[
+				SNew(STextBlock)
+				.Text(LOCTEXT("WaitingOnAssetRegister", "Waiting on Asset Registry to finish loading (all tags are present but some GC Notifies may not yet be discovered)"))
+				.Visibility(this, &SGameplayCueEditorImpl::GetWaitingOnAssetRegistryVisiblity)
+			]
 		]
 		
 		// ---------------------------------------------------------------
diff --git a/Engine/Plugins/Runtime/GameplayAbilities/Source/GameplayAbilitiesEditor/Public/AttributeDetails.h b/Engine/Plugins/Runtime/GameplayAbilities/Source/GameplayAbilitiesEditor/Public/AttributeDetails.h
new file mode 100644
--- /dev/null
+++ b/Engine/Plugins/Runtime/GameplayAbilities/Source/GameplayAbilitiesEditor/Public/AttributeDetails.h
@@ -0,0 +1,110 @@
+// Copyright 1998-2016 Epic Games, Inc. All Rights Reserved.
+
+#pragma once
+
+#include "Editor/PropertyEditor/Public/PropertyEditorModule.h"
+
+class IPropertyHandle;
+class FDetailWidgetRow;
+class IDetailLayoutBuilder;
+
+DECLARE_LOG_CATEGORY_EXTERN(LogAttributeDetails, Log, All);
+
+class FAttributeDetails : public IDetailCustomization
+{
+public:
+	/** Makes a new instance of this detail layout class for a specific detail view requesting it */
+	static TSharedRef<IDetailCustomization> MakeInstance();
+
+private:
+	/** IDetailCustomization interface */
+	virtual void CustomizeDetails(IDetailLayoutBuilder& DetailLayout) override;
+
+	TSharedPtr<FString> GetPropertyType() const;
+
+	TArray<TSharedPtr<FString>> PropertyOptions;
+
+	TSharedPtr<IPropertyHandle> MyProperty;
+
+	void OnChangeProperty(TSharedPtr<FString> ItemSelected, ESelectInfo::Type SelectInfo);	
+};
+
+
+class FAttributePropertyDetails : public IPropertyTypeCustomization
+{
+public:
+	static TSharedRef<IPropertyTypeCustomization> MakeInstance();
+
+	/** IPropertyTypeCustomization interface */
+	virtual void CustomizeHeader( TSharedRef<class IPropertyHandle> StructPropertyHandle, class FDetailWidgetRow& HeaderRow, IPropertyTypeCustomizationUtils& StructCustomizationUtils ) override;
+	virtual void CustomizeChildren( TSharedRef<class IPropertyHandle> StructPropertyHandle, class IDetailChildrenBuilder& StructBuilder, IPropertyTypeCustomizationUtils& StructCustomizationUtils ) override;              
+
+private:
+
+	// the attribute property
+	TSharedPtr<IPropertyHandle> MyProperty;
+	// the owner property
+	TSharedPtr<IPropertyHandle> OwnerProperty;
+	// the name property
+	TSharedPtr<IPropertyHandle> NameProperty;
+
+	TArray<TSharedPtr<FString>> PropertyOptions;
+
+	TSharedPtr<FString>	GetPropertyType() const;
+
+	void OnChangeProperty(TSharedPtr<FString> ItemSelected, ESelectInfo::Type SelectInfo);
+	void OnAttributeChanged(UProperty* SelectedAttribute);
+};
+
+class GAMEPLAYABILITIESEDITOR_API FScalableFloatDetails : public IPropertyTypeCustomization
+{
+public:
+	static TSharedRef<IPropertyTypeCustomization> MakeInstance();
+
+	FScalableFloatDetails()
+		: PreviewLevel(0.f)
+		, MaxPreviewLevel(30.f) // This should perhaps be configurable per scalable float somehow
+	{
+	}
+
+protected:
+
+	virtual void CustomizeHeader( TSharedRef<class IPropertyHandle> StructPropertyHandle, class FDetailWidgetRow& HeaderRow, IPropertyTypeCustomizationUtils& StructCustomizationUtils ) override;
+	virtual void CustomizeChildren( TSharedRef<class IPropertyHandle> StructPropertyHandle, class IDetailChildrenBuilder& StructBuilder, IPropertyTypeCustomizationUtils& StructCustomizationUtils ) override;              
+
+	bool IsEditable( ) const;
+
+	virtual TSharedRef<SWidget> CreateCurveTableWidget();
+
+	EVisibility GetRowNameVisibility() const;
+	EVisibility GetPreviewVisibility() const;
+	float GetPreviewLevel() const;
+	void SetPreviewLevel(float NewLevel);
+
+	TSharedRef<SWidget> GetListContent();
+	void OnSelectionChanged(TSharedPtr<FString> SelectedItem, ESelectInfo::Type SelectInfo);
+	FText GetRowNameComboBoxContentText() const;
+	FText GetRowValuePreviewLabel() const;
+	FText GetRowValuePreviewText() const;
+	TSharedRef<ITableRow> HandleRowNameComboBoxGenarateWidget(TSharedPtr<FString> InItem, const TSharedRef<STableViewBase>& OwnerTable);
+	void OnFilterTextChanged(const FText& InFilterText);
+	
+	class UCurveTable * GetCurveTable();
+
+	TSharedPtr<FString> InitWidgetContent();
+
+	void OnCurveTableChanged();
+
+	TSharedPtr<SComboButton> RowNameComboButton;
+	TSharedPtr<SListView<TSharedPtr<FString> > > RowNameComboListView;
+	TSharedPtr<FString> CurrentSelectedItem;
+	TArray<TSharedPtr<FString> > RowNames;
+
+	TSharedPtr<IPropertyHandle> ValueProperty;
+	TSharedPtr<IPropertyHandle> CurveTableHandleProperty;
+	TSharedPtr<IPropertyHandle> CurveTableProperty;
+	TSharedPtr<IPropertyHandle> RowNameProperty;
+
+	float PreviewLevel;
+	float MaxPreviewLevel;
+};
\ No newline at end of file
diff --git a/Engine/Shaders/BasePassCommon.usf b/Engine/Shaders/BasePassCommon.usf
--- a/Engine/Shaders/BasePassCommon.usf
+++ b/Engine/Shaders/BasePassCommon.usf
@@ -11,7 +11,7 @@
 // Materials that read from scene color skip getting fogged, because the contents of the scene color lookup have already been fogged
 // This is not foolproof, as any additional color the material adds will then not be fogged correctly
 #define TRANSLUCENCY_NEEDS_BASEPASS_FOGGING	(ENABLE_TRANSLUCENCY_VERTEX_FOG && MATERIALBLENDING_ANY_TRANSLUCENT && !MATERIAL_USES_SCENE_COLOR_COPY)
-#define OPAQUE_NEEDS_BASEPASS_FOGGING		(!MATERIALBLENDING_ANY_TRANSLUCENT && VERTEX_FOGGING_FOR_OPAQUE)
+#define OPAQUE_NEEDS_BASEPASS_FOGGING		(!MATERIALBLENDING_ANY_TRANSLUCENT && VERTEX_FOGGING_FOR_OPAQUE && FORWARD_SHADING)
 #define NEEDS_BASEPASS_FOGGING				(TRANSLUCENCY_NEEDS_BASEPASS_FOGGING || OPAQUE_NEEDS_BASEPASS_FOGGING)
 #define NEEDS_LIGHTMAP						(NEEDS_LIGHTMAP_COORDINATE)
 
diff --git a/Engine/Shaders/PostProcessAmbientOcclusion.usf b/Engine/Shaders/PostProcessAmbientOcclusion.usf
--- a/Engine/Shaders/PostProcessAmbientOcclusion.usf
+++ b/Engine/Shaders/PostProcessAmbientOcclusion.usf
@@ -322,11 +322,9 @@ float3 GetWorldSpaceNormalFromAOInput(float2 UV, float4 SvPosition)
 		WorldNormal = GetGBufferData(UV, false).WorldNormal;
 
 		FGBufferData GBuffer = GetGBufferData(UV);
-		BRANCH if( GBuffer.ShadingModelID == SHADINGMODELID_HAIR )
+		if( GBuffer.ShadingModelID == SHADINGMODELID_HAIR )
 		{
-			float3 V = -normalize(SvPositionToTranslatedWorld(float4(SvPosition.xy, 1, 0)));
-//				float3 V = -normalize(mul(float4(UVAndScreenPos.zw, 1, 0), View.ScreenToWorld).xyz);
-			WorldNormal = normalize( V - WorldNormal * dot(V, WorldNormal) );
+			WorldNormal = 0;
 		}
 	}
 	else if(SSAO_NORMAL_MODE == 2)
diff --git a/Engine/Shaders/ShaderVersion.usf b/Engine/Shaders/ShaderVersion.usf
--- a/Engine/Shaders/ShaderVersion.usf
+++ b/Engine/Shaders/ShaderVersion.usf
@@ -2,4 +2,4 @@
 // Each time the console command is executed it generates a new GUID. As this file is included
 // in common.usf (which should be included in any shader) it allows to invalidate the shader DDC.
 // 
-// GUID = 92A46499-DEF3-4304-B7FB-8FFD1BB7AF9E
\ No newline at end of file
+// GUID = 04221832B892468D9DE5276777772205
\ No newline at end of file
diff --git a/Engine/Shaders/SlatePostProcessPixelShader.usf b/Engine/Shaders/SlatePostProcessPixelShader.usf
new file mode 100644
--- /dev/null
+++ b/Engine/Shaders/SlatePostProcessPixelShader.usf
@@ -0,0 +1,110 @@
+// Copyright 1998-2016 Epic Games, Inc. All Rights Reserved.
+
+#include "Common.usf"
+#include "SlateShaderCommon.usf"
+
+
+#define BILINEAR_FILTER_METHOD 1
+#define MAX_SAMPLES 127
+
+#if BILINEAR_FILTER_METHOD
+
+// Weigts and offsets are packed into 4 floats (Weight, Offset, Weight, Offset)
+float4 WeightAndOffsets[MAX_SAMPLES/2];
+#else
+
+#define MAX_SAMPLES 127
+// Weigts and offsets are packed into 4 floats (Weight, Offset, Weight, Offset)
+float4 WeightAndOffsets[MAX_SAMPLES];
+
+#endif
+
+/** Blur sample count */
+int SampleCount;
+
+
+Texture2D ElementTexture;
+SamplerState ElementTextureSampler;
+
+float4 BufferSizeAndDirection;
+float4 UVBounds;
+float4 ShaderParams;
+
+float4 GetSample(float Weight, float Offset, float2 UV)
+{
+	const float2 MinUV = UVBounds.xy;
+	const float2 MaxUV = UVBounds.zw;
+	const float2 Direction = BufferSizeAndDirection.zw;
+	const float2 BufferSize = BufferSizeAndDirection.xy;
+	const float2 UVOffset = float2(Offset*BufferSize.x*Direction.x, Offset*BufferSize.y*Direction.y);
+
+	return 
+		Texture2DSample(ElementTexture, ElementTextureSampler, clamp(UV + UVOffset, MinUV, MaxUV)) * Weight
+		+ Texture2DSample(ElementTexture, ElementTextureSampler, clamp(UV - UVOffset, MinUV, MaxUV)) * Weight;
+}
+
+float4 GaussianBlurMain( FScreenVertexOutput Input ) : SV_Target0
+{
+#if 0
+	float4 OutColor = Texture2DSample(ElementTexture, ElementTextureSampler, Input.UV);
+#else
+
+	float4 OutColor = Texture2DSample(ElementTexture, ElementTextureSampler, clamp(Input.UV, UVBounds.xy, UVBounds.zw)) * WeightAndOffsets[0].x;
+	// First offset is in zw
+	{
+		float Weight = WeightAndOffsets[0].z;
+		float Offset = WeightAndOffsets[0].w;
+
+		OutColor += GetSample(Weight, Offset, Input.UV);
+	}
+
+	for (int i = 2; i<SampleCount; i+=2)
+	{
+		int Index = i/2;
+		{
+			float Weight = WeightAndOffsets[Index].x;
+			float Offset = WeightAndOffsets[Index].y;
+
+			OutColor += GetSample(Weight, Offset, Input.UV);
+		}
+
+		{
+			float Weight = WeightAndOffsets[Index].z;
+			float Offset = WeightAndOffsets[Index].w;
+
+			OutColor += GetSample(Weight, Offset, Input.UV);
+		}
+
+	}
+#endif
+
+	return float4(OutColor.rgb, 1);
+}
+
+float4 DownsampleMain(FScreenVertexOutput Input) : SV_Target0
+{
+#if 0
+	return Texture2DSample(ElementTexture, ElementTextureSampler, Input.UV);
+#else
+
+	float2 UV[4];
+
+	float2 MinUV = UVBounds.xy;
+	float2 MaxUV = UVBounds.zw;
+
+	// Shader params X/Y stores the UV offset in each direction
+	UV[0] = clamp(Input.UV + ShaderParams.xy * float2(-1, -1), MinUV, MaxUV);
+	UV[1] = clamp(Input.UV + ShaderParams.xy * float2(1, -1), MinUV, MaxUV);
+	UV[2] = clamp(Input.UV + ShaderParams.xy * float2(-1, 1), MinUV, MaxUV);
+	UV[3] = clamp(Input.UV + ShaderParams.xy * float2(1, 1), MinUV, MaxUV);
+
+	float4 Sample[4];
+
+	UNROLL for(int i = 0; i < 4; ++i)
+	{
+		Sample[i] = Texture2DSample(ElementTexture, ElementTextureSampler, UV[i]);
+	}
+
+	return float4(Sample[0] + Sample[1] + Sample[2] + Sample[3]) * 0.25f;
+#endif
+}
\ No newline at end of file
diff --git a/Engine/Source/Developer/AutomationController/Private/AutomationCommandline.cpp b/Engine/Source/Developer/AutomationController/Private/AutomationCommandline.cpp
--- a/Engine/Source/Developer/AutomationController/Private/AutomationCommandline.cpp
+++ b/Engine/Source/Developer/AutomationController/Private/AutomationCommandline.cpp
@@ -335,7 +335,7 @@ public:
 				{
 					AutomationCommandQueue.Add(EAutomationCommand::ListAllTests);
 				}
-				else if (FParse::Command(&TempCmd, TEXT("RunTests")))
+				else if (FParse::Command(&TempCmd, TEXT("RunTests")) || FParse::Command(&TempCmd, TEXT("RunTest")))
 				{
 					if ( FParse::Command(&TempCmd, TEXT("Now")) )
 					{
@@ -344,6 +344,7 @@ public:
 
 					//only one of these should be used
 					StringCommand = TempCmd;
+					Ar.Logf(TEXT("Running all tests matching substring: %s"), *StringCommand);
 					AutomationCommandQueue.Add(EAutomationCommand::RunCommandLineTests);
 				}
 				else if (FParse::Command(&TempCmd, TEXT("RunFilter")))
@@ -354,23 +355,30 @@ public:
 					if (FilterMaps.Contains(FlagToUse))
 					{
 						AutomationController->SetRequestedTestFlags(FilterMaps[FlagToUse]);
+						Ar.Logf(TEXT("Running all tests for filter: %s"), *FlagToUse);
 					}
 					AutomationCommandQueue.Add(EAutomationCommand::RunFilter);
 				}
 				else if (FParse::Command(&TempCmd, TEXT("RunAll")))
 				{
 					AutomationCommandQueue.Add(EAutomationCommand::RunAll);
+					Ar.Logf(TEXT("Running all available automated tests for this program. NOTE: This may take a while."));
 				}
 				else if (FParse::Command(&TempCmd, TEXT("Quit")))
 				{
 					AutomationCommandQueue.Add(EAutomationCommand::Quit);
+					Ar.Logf(TEXT("Exiting: Received quit command."));
 				}
-				// let user know he mis-typed a param/command
-				else if (FParse::Command(&TempCmd, TEXT("RunTest")))
+
+				else
 				{
-					Ar.Logf(TEXT("Unhandled token \"RunTest\". You probably meant \"RunTests\". Bailing out."));
+					Ar.Logf(TEXT("Incorrect automation command syntax! Supported commands are: "));
+					Ar.Logf(TEXT("\tAutomation List"));
+					Ar.Logf(TEXT("\tAutomation RunTests <test string>"));
+					Ar.Logf(TEXT("\tAutomation RunAll "));
+					Ar.Logf(TEXT("\tAutomation RunFilter <filter name>"));
+					Ar.Logf(TEXT("\tAutomation Quit"));
 					bHandled = false;
-					break;
 				}
 			}
 		}
diff --git a/Engine/Source/Developer/FunctionalTesting/Classes/FunctionalTest.h b/Engine/Source/Developer/FunctionalTesting/Classes/FunctionalTest.h
--- a/Engine/Source/Developer/FunctionalTesting/Classes/FunctionalTest.h
+++ b/Engine/Source/Developer/FunctionalTesting/Classes/FunctionalTest.h
@@ -561,7 +561,8 @@ protected:
 public:
 	FFunctionalTestDoneSignature TestFinishedObserver;
 
-protected:
+	// AG TEMP - solving a compile issue in a temp way to unblock the bui.d
+	UPROPERTY(Transient)
 	bool bIsRunning;
 
 	TArray<FString> Steps;
diff --git a/Engine/Source/Developer/FunctionalTesting/Classes/FunctionalTestingManager.h b/Engine/Source/Developer/FunctionalTesting/Classes/FunctionalTestingManager.h
--- a/Engine/Source/Developer/FunctionalTesting/Classes/FunctionalTestingManager.h
+++ b/Engine/Source/Developer/FunctionalTesting/Classes/FunctionalTestingManager.h
@@ -3,10 +3,12 @@
 #pragma once
 
 #include "Misc/AutomationTest.h"
+#include "FunctionalTest.h"
 #include "Kismet/BlueprintFunctionLibrary.h"
 #include "FunctionalTestingManager.generated.h"
 
 class UWorld;
+class IMessageLogListing;
 
 namespace FFunctionalTesting
 {
diff --git a/Engine/Source/Developer/FunctionalTesting/Private/AutomationBlueprintFunctionLibrary.cpp b/Engine/Source/Developer/FunctionalTesting/Private/AutomationBlueprintFunctionLibrary.cpp
--- a/Engine/Source/Developer/FunctionalTesting/Private/AutomationBlueprintFunctionLibrary.cpp
+++ b/Engine/Source/Developer/FunctionalTesting/Private/AutomationBlueprintFunctionLibrary.cpp
@@ -263,7 +263,9 @@ bool UAutomationBlueprintFunctionLibrary::TakeAutomationScreenshotInternal(const
 			if ( !GEngine->GameViewport->GetGameViewport()->TakeHighResScreenShot() )
 			{
 				// If we failed to take the screenshot, we're going to need to cleanup the automation screenshot taker.
+#if (WITH_DEV_AUTOMATION_TESTS || WITH_PERF_AUTOMATION_TESTS)
 				delete TempObject;
+#endif 
 				return false;
 			}
 
@@ -365,7 +367,6 @@ void UAutomationBlueprintFunctionLibrary::TakeAutomationScreenshotOfUI(UObject*
 				{
 #if (WITH_DEV_AUTOMATION_TESTS || WITH_PERF_AUTOMATION_TESTS)
 					FAutomationScreenshotTaker* TempObject = new FAutomationScreenshotTaker(Name, Options);
-#endif
 
 					FAutomationScreenshotData Data = AutomationCommon::BuildScreenshotData(GWorld->GetName(), Name, OutSize.X, OutSize.Y);
 
@@ -383,6 +384,7 @@ void UAutomationBlueprintFunctionLibrary::TakeAutomationScreenshotOfUI(UObject*
 					Data.MaximumGlobalError = Options.MaximumGlobalError;
 
 					GEngine->GameViewport->OnScreenshotCaptured().Broadcast(OutSize.X, OutSize.Y, OutColorData);
+#endif
 				}
 
 				FLatentActionManager& LatentActionManager = World->GetLatentActionManager();
diff --git a/Engine/Source/Developer/FunctionalTesting/Private/FunctionalTest.cpp b/Engine/Source/Developer/FunctionalTesting/Private/FunctionalTest.cpp
--- a/Engine/Source/Developer/FunctionalTesting/Private/FunctionalTest.cpp
+++ b/Engine/Source/Developer/FunctionalTesting/Private/FunctionalTest.cpp
@@ -289,7 +289,7 @@ void AFunctionalTest::FinishTest(EFunctionalTestResult TestResult, const FString
 	}
 
 	const FText ResultText = FTestResultTypeEnum->GetEnumText( (int32)TestResult );
-	const FString OutMessage = FString::Printf(TEXT("%s %s: \"%s\'")
+	const FString OutMessage = FString::Printf(TEXT("%s %s: \"%s\"")
 		, *GetName()
 		, *ResultText.ToString()
 		, Message.IsEmpty() == false ? *Message : TEXT("Test finished") );
diff --git a/Engine/Source/Developer/GameplayDebugger/Private/GameplayDebuggerAddonManager.h b/Engine/Source/Developer/GameplayDebugger/Private/GameplayDebuggerAddonManager.h
--- a/Engine/Source/Developer/GameplayDebugger/Private/GameplayDebuggerAddonManager.h
+++ b/Engine/Source/Developer/GameplayDebugger/Private/GameplayDebuggerAddonManager.h
@@ -61,8 +61,8 @@ public:
 	/** get slot-Name map */
 	const TArray<FString> GetSlotNames() const { return SlotNames; }
 
-	/** get number of known categories */
-	int32 GetNumCategories() const { return CategoryMap.Num(); }
+	/** get number of visible categories */
+	int32 GetNumVisibleCategories() const { return NumVisibleCategories; }
 
 	/** singleton accessor */
 	static FGameplayDebuggerAddonManager& GetCurrent();
@@ -85,4 +85,7 @@ private:
 
 	/** list of slot names */
 	TArray<FString> SlotNames;
+
+	/** number of categories, excluding hidden ones */
+	int32 NumVisibleCategories;
 };
diff --git a/Engine/Source/Developer/GameplayDebugger/Private/GameplayDebuggerCategoryReplicator.cpp b/Engine/Source/Developer/GameplayDebugger/Private/GameplayDebuggerCategoryReplicator.cpp
--- a/Engine/Source/Developer/GameplayDebugger/Private/GameplayDebuggerCategoryReplicator.cpp
+++ b/Engine/Source/Developer/GameplayDebugger/Private/GameplayDebuggerCategoryReplicator.cpp
@@ -183,7 +183,7 @@ bool FGameplayDebuggerNetPack::NetDeltaSerialize(FNetDeltaSerializeInfo& DeltaPa
 			}
 		}
 
-		if (CVarGameplayDebuggerRepDetails.GetValueOnGameThread())
+		if (CVarGameplayDebuggerRepDetails.GetValueOnAnyThread())
 		{
 			if (OldState)
 			{
diff --git a/Engine/Source/Developer/GameplayDebugger/Private/GameplayDebuggerLocalController.cpp b/Engine/Source/Developer/GameplayDebugger/Private/GameplayDebuggerLocalController.cpp
--- a/Engine/Source/Developer/GameplayDebugger/Private/GameplayDebuggerLocalController.cpp
+++ b/Engine/Source/Developer/GameplayDebugger/Private/GameplayDebuggerLocalController.cpp
@@ -713,7 +713,7 @@ void UGameplayDebuggerLocalController::OnCategoriesChanged()
 	}
 
 	NumCategorySlots = SlotCategoryIds.Num();
-	NumCategories = AddonManager.GetNumCategories();
+	NumCategories = AddonManager.GetNumVisibleCategories();
 
 	DataPackMap.Reset();
 }
diff --git a/Engine/Source/Developer/GameplayDebugger/Public/GameplayDebugger.h b/Engine/Source/Developer/GameplayDebugger/Public/GameplayDebugger.h
--- a/Engine/Source/Developer/GameplayDebugger/Public/GameplayDebugger.h
+++ b/Engine/Source/Developer/GameplayDebugger/Public/GameplayDebugger.h
@@ -48,6 +48,7 @@ enum class EGameplayDebuggerCategoryState : uint8
 	EnabledInGame,
 	EnabledInSimulate,
 	Disabled,
+	Hidden,
 };
 
 class IGameplayDebugger : public IModuleInterface
diff --git a/Engine/Source/Developer/GameplayDebugger/Public/GameplayDebuggerConfig.h b/Engine/Source/Developer/GameplayDebugger/Public/GameplayDebuggerConfig.h
--- a/Engine/Source/Developer/GameplayDebugger/Public/GameplayDebuggerConfig.h
+++ b/Engine/Source/Developer/GameplayDebugger/Public/GameplayDebuggerConfig.h
@@ -54,13 +54,19 @@ struct FGameplayDebuggerCategoryConfig
 	UPROPERTY(EditAnywhere, Category = Settings)
 	EGameplayDebuggerOverrideMode ActiveInSimulate;
 
+	UPROPERTY(EditAnywhere, Category = Settings)
+	EGameplayDebuggerOverrideMode Hidden;
+
 	UPROPERTY(EditAnywhere, Category = Settings, AdvancedDisplay)
 	uint32 bOverrideSlotIdx : 1;
 
 	UPROPERTY(EditAnywhere, Category = Settings, EditFixedSize)
 	TArray<FGameplayDebuggerInputConfig> InputHandlers;
 
-	FGameplayDebuggerCategoryConfig() : ActiveInGame(EGameplayDebuggerOverrideMode::UseDefault), ActiveInSimulate(EGameplayDebuggerOverrideMode::UseDefault), bOverrideSlotIdx(false) {}
+	FGameplayDebuggerCategoryConfig() :
+		ActiveInGame(EGameplayDebuggerOverrideMode::UseDefault), ActiveInSimulate(EGameplayDebuggerOverrideMode::UseDefault),
+		Hidden(EGameplayDebuggerOverrideMode::UseDefault), bOverrideSlotIdx(false)
+	{}
 };
 
 USTRUCT()
diff --git a/Engine/Source/Developer/MeshUtilities/Private/MeshUtilities.cpp b/Engine/Source/Developer/MeshUtilities/Private/MeshUtilities.cpp
--- a/Engine/Source/Developer/MeshUtilities/Private/MeshUtilities.cpp
+++ b/Engine/Source/Developer/MeshUtilities/Private/MeshUtilities.cpp
@@ -8308,6 +8308,8 @@ void FMeshUtilities::StartupModule()
 		TArray<FName> SwarmModuleNames;
 		FModuleManager::Get().FindModules(TEXT("*SimplygonSwarm"), SwarmModuleNames);
 
+			// Look for MeshReduction interface
+
 
 		for (int32 Index = 0; Index < SwarmModuleNames.Num(); Index++)
 		{
diff --git a/Engine/Source/Editor/BlueprintGraph/Private/CallFunctionHandler.cpp b/Engine/Source/Editor/BlueprintGraph/Private/CallFunctionHandler.cpp
--- a/Engine/Source/Editor/BlueprintGraph/Private/CallFunctionHandler.cpp
+++ b/Engine/Source/Editor/BlueprintGraph/Private/CallFunctionHandler.cpp
@@ -374,7 +374,7 @@ void FKCHandler_CallFunction::CreateFunctionCallStatement(FKismetFunctionContext
 					for(int32 i = 0; i < SelfPin->LinkedTo.Num(); i++)
 					{
 						FBPTerminal** pContextTerm = Context.NetMap.Find(SelfPin->LinkedTo[i]);
-						if(ensure(pContextTerm != nullptr))
+						if(ensureMsgf(pContextTerm != nullptr, TEXT("'%s' is missing a target input - if this is a server build, the input may be a cosmetic only property which was discarded (if this is the case, and this is expecting component variable try resaving.)"), *Node->GetPathName()))
 						{
 							CheckAndAddSelfTermLambda(*pContextTerm);
 						}
diff --git a/Engine/Source/Editor/Kismet/Private/SBlueprintEditorSelectedDebugObjectWidget.cpp b/Engine/Source/Editor/Kismet/Private/SBlueprintEditorSelectedDebugObjectWidget.cpp
--- a/Engine/Source/Editor/Kismet/Private/SBlueprintEditorSelectedDebugObjectWidget.cpp
+++ b/Engine/Source/Editor/Kismet/Private/SBlueprintEditorSelectedDebugObjectWidget.cpp
@@ -564,19 +564,21 @@ TSharedPtr<FString> SBlueprintEditorSelectedDebugObjectWidget::GetDebugObjectNam
 TSharedPtr<FString> SBlueprintEditorSelectedDebugObjectWidget::GetDebugWorldName() const
 {
 	check(GetBlueprintObj());
-	check(DebugWorlds.Num() == DebugWorldNames.Num());
-	if (UObject* DebugObj = GetBlueprintObj()->GetObjectBeingDebugged())
+	if (ensure(DebugWorlds.Num() == DebugWorldNames.Num()))
 	{
-		for (int32 WorldIdx = 0; WorldIdx < DebugWorlds.Num(); ++WorldIdx)
+		if (UObject* DebugObj = GetBlueprintObj()->GetObjectBeingDebugged())
 		{
-			if (DebugObj->IsIn(DebugWorlds[WorldIdx].Get()))
+			for (int32 WorldIdx = 0; WorldIdx < DebugWorlds.Num(); ++WorldIdx)
 			{
-				return DebugWorldNames[WorldIdx];
+				if (DebugObj->IsIn(DebugWorlds[WorldIdx].Get()))
+				{
+					return DebugWorldNames[WorldIdx];
+				}
 			}
 		}
 	}
-	return DebugWorldNames[0];
 
+	return DebugWorldNames[0];
 }
 
 void SBlueprintEditorSelectedDebugObjectWidget::DebugWorldSelectionChanged(TSharedPtr<FString> NewSelection, ESelectInfo::Type SelectInfo)
diff --git a/Engine/Source/Editor/TranslationEditor/Private/TranslationPickerFloatingWindow.cpp b/Engine/Source/Editor/TranslationEditor/Private/TranslationPickerFloatingWindow.cpp
--- a/Engine/Source/Editor/TranslationEditor/Private/TranslationPickerFloatingWindow.cpp
+++ b/Engine/Source/Editor/TranslationEditor/Private/TranslationPickerFloatingWindow.cpp
@@ -236,25 +236,41 @@ FText STranslationPickerFloatingWindow::GetTextFromWidget(TSharedRef<SWidget> Wi
 	if (!OriginalText.IsEmpty())
 	{
 		// Search the text from this widget's FText::Format history to find any source text
-		TArray<FText> FormattingSourceTexts;
-		OriginalText.GetSourceTextsFromFormatHistory(FormattingSourceTexts);
+		TArray<FHistoricTextFormatData> HistoricFormatData;
+		FTextInspector::GetHistoricFormatData(OriginalText, HistoricFormatData);
 
-		for (FText FormattingSourceText : FormattingSourceTexts)
+		auto PickedTextsContains = [&](const FText& InTextToTest)
 		{
-			// Don't show the same text twice
-			bool bAlreadyPicked = false;
-			for (FText PickedText : PickedTexts)
+			return PickedTexts.ContainsByPredicate([&](const FText& InPickedText)
 			{
-				if (FormattingSourceText.EqualTo(PickedText))
+				return InTextToTest.EqualTo(InPickedText);
+			});
+		};
+
+		if (HistoricFormatData.Num() > 0)
+		{
+			for (const FHistoricTextFormatData& HistoricFormatDataItem : HistoricFormatData)
+			{
+				if (!PickedTextsContains(HistoricFormatDataItem.SourceFmt.GetSourceText()))
 				{
-					bAlreadyPicked = true;
-					break;
+					PickedTexts.Add(HistoricFormatDataItem.SourceFmt.GetSourceText());
 				}
-			}
 
-			if (!bAlreadyPicked)
+				for (auto It = HistoricFormatDataItem.Arguments.CreateConstIterator(); It; ++It)
+				{
+					const FFormatArgumentValue& ArgumentValue = It.Value();
+					if (ArgumentValue.GetType() == EFormatArgumentType::Text && !PickedTextsContains(ArgumentValue.GetTextValue()))
+					{
+						PickedTexts.Add(ArgumentValue.GetTextValue());
+					}
+				}
+			}
+		}
+		else
+		{
+			if (!PickedTextsContains(OriginalText))
 			{
-				PickedTexts.Add(FormattingSourceText);
+				PickedTexts.Add(OriginalText);
 			}
 		}
 	}
diff --git a/Engine/Source/Editor/UMGEditor/Private/WidgetBlueprintEditor.cpp b/Engine/Source/Editor/UMGEditor/Private/WidgetBlueprintEditor.cpp
--- a/Engine/Source/Editor/UMGEditor/Private/WidgetBlueprintEditor.cpp
+++ b/Engine/Source/Editor/UMGEditor/Private/WidgetBlueprintEditor.cpp
@@ -102,6 +102,7 @@ void FWidgetBlueprintEditor::InitWidgetBlueprintEditor(const EToolkitMode::Type
 	// register for any objects replaced
 	GEditor->OnObjectsReplaced().AddSP(this, &FWidgetBlueprintEditor::OnObjectsReplaced);
 
+	// for change selected widgets on sequencer tree view
 	UWidgetBlueprint* Blueprint = GetWidgetBlueprintObj();
 
 	// If this blueprint is empty, add a canvas panel as the root widget.
diff --git a/Engine/Source/Editor/UnrealEd/Classes/CookOnTheSide/CookOnTheFlyServer.h b/Engine/Source/Editor/UnrealEd/Classes/CookOnTheSide/CookOnTheFlyServer.h
--- a/Engine/Source/Editor/UnrealEd/Classes/CookOnTheSide/CookOnTheFlyServer.h
+++ b/Engine/Source/Editor/UnrealEd/Classes/CookOnTheSide/CookOnTheFlyServer.h
@@ -1285,24 +1285,6 @@ private:
 	 * @param FoundPackages list of packages which were found
 	 */
 	void GetDependentPackages(const TSet<FName>& RootPackages, TSet<FName>& FoundPackages);
-	/**
-	 * GenerateManifestInfo
-	 * generate the manifest information for a given package
-	 *
-	 * @param Package package to generate manifest information for
-	 */
-	void GenerateManifestInfo( UPackage* Package, const TArray<FName>& TargetPlatformNames );
-
-
-	/**
-	 * GenerateManifestInfo
-	 * generate manfiest information for the given package and add it to the manifest
-	 * this version is teh same as GenerateManifestInfo which takes in a UPackageObject except that doesn't require the package to be loaded
-	 *
-	 * @param Package name of the package to generate the manifest information for
-	 * @param TargetPlatformNames to add the manifest information to
-	 */
-	void GenerateManifestInfo(const FName& Package, const TArray<FName>& TargetPlatformNames);
 
 	/**
 	 * ContainsWorld
diff --git a/Engine/Source/Editor/UnrealEd/Classes/Settings/LevelEditorPlaySettings.h b/Engine/Source/Editor/UnrealEd/Classes/Settings/LevelEditorPlaySettings.h
--- a/Engine/Source/Editor/UnrealEd/Classes/Settings/LevelEditorPlaySettings.h
+++ b/Engine/Source/Editor/UnrealEd/Classes/Settings/LevelEditorPlaySettings.h
@@ -415,10 +415,7 @@ public:
 	UPROPERTY(config)
 	TArray<FPlayScreenResolution> TelevisionScreenResolutions;
 
-protected:
-
-	// UObject overrides
-
+	// UObject interface
 	virtual void PostEditChangeProperty(struct FPropertyChangedEvent& PropertyChangedEvent) override;
-
+	// End of UObject interface
 };
diff --git a/Engine/Source/Editor/UnrealEd/Private/Commandlets/ContentCommandlets.cpp b/Engine/Source/Editor/UnrealEd/Private/Commandlets/ContentCommandlets.cpp
--- a/Engine/Source/Editor/UnrealEd/Private/Commandlets/ContentCommandlets.cpp
+++ b/Engine/Source/Editor/UnrealEd/Private/Commandlets/ContentCommandlets.cpp
@@ -143,6 +143,22 @@ int32 UResavePackagesCommandlet::InitializeResaveParameters( const TArray<FStrin
 
 	}
 
+	if (bShouldBuildLighting && !bExplicitPackages)
+	{
+		UE_LOG(LogContentCommandlet, Display, TEXT("No maps found to save when building lighting, checking CommandletSettings:ResavePackages in EditorIni"));
+		// if we haven't specified any maps and we are building lighting check if there are packages setup in the ini file to build
+		TArray<FString> ResavePackages;
+		GConfig->GetArray(TEXT("CommandletSettings"), TEXT("ResavePackages"), ResavePackages, GEditorIni);
+		for ( const auto& ResavePackage : ResavePackages )
+		{
+			FString PackageFile;
+			FPackageName::SearchForPackageOnDisk(ResavePackage, NULL, &PackageFile);
+			UE_LOG(LogContentCommandlet, Display, TEXT("Rebuilding lighting for package %s"), *PackageFile);
+			PackageNames.Add(*PackageFile);
+			bExplicitPackages = true;
+		}
+	}
+
 	// ... if not, load in all packages
 	if( !bExplicitPackages )
 	{
diff --git a/Engine/Source/Editor/UnrealEd/Private/CookOnTheFlyServer.cpp b/Engine/Source/Editor/UnrealEd/Private/CookOnTheFlyServer.cpp
--- a/Engine/Source/Editor/UnrealEd/Private/CookOnTheFlyServer.cpp
+++ b/Engine/Source/Editor/UnrealEd/Private/CookOnTheFlyServer.cpp
@@ -1163,111 +1163,6 @@ bool UCookOnTheFlyServer::ContainsMap(const FName& PackageName) const
 	return false;
 }
 
-void UCookOnTheFlyServer::GenerateManifestInfo(UPackage* Package, const TArray<FName>& TargetPlatformNames)
-{
-	if (!CookByTheBookOptions)
-		return;
-
-	check(IsChildCooker() == false);
-
-	// generate dependency information for this package
-
-	TSet<UPackage*> RootPackages;
-	RootPackages.Add(Package);
-
-	FString LastLoadedMapName;
-
-	// load sublevels
-	UWorld* World = UWorld::FindWorldInPackage(Package);
-
-#define VERIFY_LEVELS_IN_DEPENDENCIES_LIST 1
-#if VERIFY_LEVELS_IN_DEPENDENCIES_LIST
-	TArray<FString> LevelPackageNames;
-#endif
-
-	if (World)
-	{
-		for (const auto& StreamingLevel : World->StreamingLevels)
-		{
-			if (StreamingLevel->GetLoadedLevel())
-			{
-				RootPackages.Add(StreamingLevel->GetLoadedLevel()->GetOutermost());
-#if VERIFY_LEVELS_IN_DEPENDENCIES_LIST
-				LevelPackageNames.AddUnique(StreamingLevel->GetLoadedLevel()->GetOutermost()->GetName());
-#endif
-			}
-		}
-
-		TArray<FString> NewPackagesToCook;
-
-		// Collect world composition tile packages to cook
-		if (World->WorldComposition)
-		{
-			World->WorldComposition->CollectTilesToCook(NewPackagesToCook);
-		}
-
-#if VERIFY_LEVELS_IN_DEPENDENCIES_LIST
-		for (const auto& PackageToCook : NewPackagesToCook)
-		{
-			LevelPackageNames.AddUnique(PackageToCook);
-		}
-#endif
-
-		for (const auto& PackageName : NewPackagesToCook)
-		{
-			UPackage* PackageToCook = LoadPackage(NULL, *PackageName, LOAD_None);
-
-			RootPackages.Add(PackageToCook);
-			//GetDependencies( PackageToCook, Dependencies );
-
-			// Dependencies.Add(PackageToCook);
-		}
-
-		LastLoadedMapName = Package->GetName();
-	}
-#if USEASSETREGISTRYFORDEPENDENTPACKAGES
-
-	static const FName AssetRegistryName("AssetRegistry");
-	FAssetRegistryModule& AssetRegistryModule = FModuleManager::LoadModuleChecked<FAssetRegistryModule>(AssetRegistryName);
-	IAssetRegistry& AssetRegistry = AssetRegistryModule.Get();
-
-	TSet<FName> Packages;
-
-	GetDependentPackages(RootPackages, Packages);
-
-	auto GetPackageFName = [](const FName& Name) { return Name; };
-
-
-#if VERIFY_LEVELS_IN_DEPENDENCIES_LIST
-	if (World)
-	{
-		TSet<UPackage*> WorldPackages;
-		WorldPackages.Add(Package);
-		TSet<FName> DependentWorldPackages;
-		GetDependentPackages(WorldPackages, DependentWorldPackages);
-
-		for (const auto& LevelPackageName : LevelPackageNames)
-		{
-			check(DependentWorldPackages.Contains(FName(*LevelPackageName)));
-		}
-	}
-#endif
-
-#else
-	TSet<UObject*> Dependencies;
-	GetDependencies(RootPackages, Dependencies);
-
-	TSet<UPackage*> Packages;
-	for (const auto& Object : Dependencies)
-	{
-		Packages.Add(Object->GetOutermost());
-	}
-
-	auto GetPackageFName = [](const UPackage* Package) { return Package->GetFName(); };
-#endif
-
-
-}
 bool UCookOnTheFlyServer::IsCookingInEditor() const
 {
 	return CurrentCookMode == ECookMode::CookByTheBookFromTheEditor || CurrentCookMode == ECookMode::CookOnTheFlyFromTheEditor;
@@ -1802,13 +1697,13 @@ uint32 UCookOnTheFlyServer::TickCookOnTheSide( const float TimeSlice, uint32 &Co
 #if DEBUG_COOKONTHEFLY
 						UE_LOG(LogCook, Display, TEXT("Object %s isn't cached yet"), *Obj->GetFullName());
 #endif
-						if ( Obj->IsA(UMaterial::StaticClass()) )
+						/*if ( Obj->IsA(UMaterial::StaticClass()) )
 						{
 							if (GShaderCompilingManager->HasShaderJobs() == false)
 							{
 								UE_LOG(LogCook, Warning, TEXT("Shader compiler is in a bad state!  Shader %s is finished compile but shader compiling manager did not notify shader.  "), *Obj->GetPathName());
 							}
-						}
+						}*/
 						return false;
 					}
 				}
diff --git a/Engine/Source/Editor/UnrealEd/Private/PlayLevel.cpp b/Engine/Source/Editor/UnrealEd/Private/PlayLevel.cpp
--- a/Engine/Source/Editor/UnrealEd/Private/PlayLevel.cpp
+++ b/Engine/Source/Editor/UnrealEd/Private/PlayLevel.cpp
@@ -264,8 +264,11 @@ void UEditorEngine::EndPlayMap()
 				bSeamlessTravelActive = true;
 			}
 
-			TeardownPlaySession(ThisContext);
-			
+			if (ThisContext.World())
+			{
+				TeardownPlaySession(ThisContext);
+			}
+
 			// Cleanup online subsystems instantiated during PIE
 			FName OnlineIdentifier = UOnlineEngineInterface::Get()->GetOnlineIdentifier(ThisContext);
 			if (UOnlineEngineInterface::Get()->DoesInstanceExist(OnlineIdentifier))
@@ -506,7 +509,7 @@ void UEditorEngine::CleanupPIEOnlineSessions(TArray<FName> OnlineIdentifiers)
 	NumOnlinePIEInstances = 0;
 }
 
-void UEditorEngine::TeardownPlaySession(FWorldContext &PieWorldContext)
+void UEditorEngine::TeardownPlaySession(FWorldContext& PieWorldContext)
 {
 	check(PieWorldContext.WorldType == EWorldType::PIE);
 	PlayWorld = PieWorldContext.World();
@@ -2988,9 +2991,18 @@ UGameInstance* UEditorEngine::CreatePIEGameInstance(int32 InPIEInstance, bool bI
 	// We need to temporarily add the GameInstance to the root because the InitPIE call can do garbage collection wiping out the GameInstance
 	GameInstance->AddToRoot();
 
-	bool bSuccess = GameInstance->InitializePIE(bAnyBlueprintErrors, InPIEInstance, bRunAsDedicated);
-	if (!bSuccess)
+	FGameInstancePIEParameters GameInstanceParams;
+	GameInstanceParams.bAnyBlueprintErrors = bAnyBlueprintErrors;
+	GameInstanceParams.bSimulateInEditor = bInSimulateInEditor;
+	GameInstanceParams.bStartInSpectatorMode = bStartInSpectatorMode;
+	GameInstanceParams.bRunAsDedicated = bRunAsDedicated;
+
+	
+	const FGameInstancePIEResult InitializeResult = GameInstance->InitializeForPlayInEditor(InPIEInstance, GameInstanceParams);
+	if (!InitializeResult.IsSuccess())
 	{
+		FMessageDialog::Open(EAppMsgType::Ok, InitializeResult.FailureReason);
+
 		FEditorDelegates::EndPIE.Broadcast(bInSimulateInEditor);
 
 		if (EditorWorld->GetNavigationSystem())
@@ -3354,10 +3366,11 @@ UGameInstance* UEditorEngine::CreatePIEGameInstance(int32 InPIEInstance, bool bI
 	// By this point it is safe to remove the GameInstance from the root and allow it to garbage collected as per usual
 	GameInstance->RemoveFromRoot();
 
-	bSuccess = GameInstance->StartPIEGameInstance(NewLocalPlayer, bInSimulateInEditor, bAnyBlueprintErrors, bStartInSpectatorMode);
-	if (!bSuccess)
+	// Start the game instance
+	const FGameInstancePIEResult StartResult = GameInstance->StartPlayInEditorGameInstance(NewLocalPlayer, GameInstanceParams);
+	if (!StartResult.IsSuccess())
 	{
-		FMessageDialog::Open(EAppMsgType::Ok, NSLOCTEXT("UnrealEd", "Error_CouldntStartInstance", "Failed to start PIE game instance"));
+		FMessageDialog::Open(EAppMsgType::Ok, StartResult.FailureReason);
 		RestoreEditorWorld( EditorWorld );
 		EndPlayMap();
 		return nullptr;
diff --git a/Engine/Source/Editor/UnrealEd/Public/EditorAnalytics.h b/Engine/Source/Editor/UnrealEd/Public/EditorAnalytics.h
--- a/Engine/Source/Editor/UnrealEd/Public/EditorAnalytics.h
+++ b/Engine/Source/Editor/UnrealEd/Public/EditorAnalytics.h
@@ -60,6 +60,7 @@ namespace EAnalyticsErrorCodes
 		MissingExecutable = 103,
 		DeviceNotSetupForDevelopment = 150,
 		DeviceOSNewerThanSDK = 151,
+		TestFailure = 152,
 	};
 };
 
diff --git a/Engine/Source/Programs/AutomationTool/AutomationUtils/AutomationException.cs b/Engine/Source/Programs/AutomationTool/AutomationUtils/AutomationException.cs
--- a/Engine/Source/Programs/AutomationTool/AutomationUtils/AutomationException.cs
+++ b/Engine/Source/Programs/AutomationTool/AutomationUtils/AutomationException.cs
@@ -57,6 +57,7 @@ namespace AutomationTool
         Error_MissingExecutable = 103,
         Error_DeviceNotSetupForDevelopment = 150,
         Error_DeviceOSNewerThanSDK = 151,
+		Error_TestFailure = 152,
     };
 
     /// <summary>
diff --git a/Engine/Source/Programs/AutomationTool/AutomationUtils/MacHostPlatform.cs b/Engine/Source/Programs/AutomationTool/AutomationUtils/MacHostPlatform.cs
--- a/Engine/Source/Programs/AutomationTool/AutomationUtils/MacHostPlatform.cs
+++ b/Engine/Source/Programs/AutomationTool/AutomationUtils/MacHostPlatform.cs
@@ -135,6 +135,10 @@ namespace AutomationTool
 		public override bool IsScriptModuleSupported(string ModuleName)
 		{
 			// @todo: add more unsupported modules here
+			if (String.Equals(ModuleName, "OrionGame", StringComparison.InvariantCultureIgnoreCase))
+			{
+				return false;
+			}
 			return true;
 		}
 
diff --git a/Engine/Source/Programs/AutomationTool/AutomationUtils/ProjectParams.cs b/Engine/Source/Programs/AutomationTool/AutomationUtils/ProjectParams.cs
--- a/Engine/Source/Programs/AutomationTool/AutomationUtils/ProjectParams.cs
+++ b/Engine/Source/Programs/AutomationTool/AutomationUtils/ProjectParams.cs
@@ -280,6 +280,7 @@ namespace AutomationTool
 			this.ServerCommandline = InParams.ServerCommandline;
             this.Package = InParams.Package;
 			this.Deploy = InParams.Deploy;
+			this.DeployFolder = InParams.DeployFolder;
 			this.IterativeDeploy = InParams.IterativeDeploy;
 			this.IgnoreCookErrors = InParams.IgnoreCookErrors;
 			this.FastCook = InParams.FastCook;
@@ -389,6 +390,7 @@ namespace AutomationTool
 			bool? DedicatedServer = null,
 			bool? Client = null,
 			bool? Deploy = null,
+			string DeployFolder = null,
 			bool? FileServer = null,
 			bool? Foreign = null,
 			bool? ForeignCode = null,
@@ -619,7 +621,23 @@ namespace AutomationTool
 			this.ServerCommandline = ParseParamValueIfNotSpecified(Command, ServerCommandline, "servercmdline");
 			this.ServerCommandline = this.ServerCommandline.Replace('\'', '\"'); // replace any single quotes with double quotes
             this.Package = GetParamValueIfNotSpecified(Command, Package, this.Package, "package");
+
 			this.Deploy = GetParamValueIfNotSpecified(Command, Deploy, this.Deploy, "deploy");
+			this.DeployFolder = ParseParamValueIfNotSpecified(Command, DeployFolder, "deploy", null);
+
+			// if the user specified -deploy but no folder, set the default
+			if (this.Deploy && string.IsNullOrEmpty(this.DeployFolder))
+			{
+				this.DeployFolder = this.ShortProjectName;
+			}
+			else if (string.IsNullOrEmpty(this.DeployFolder) == false)
+			{
+				// if the user specified a folder set deploy to true.
+				//@todo - remove 'deploy' var and check deployfolder != null?
+				this.Deploy = true;
+			}
+
+			if (string.IsNullOrEmpty(this.DeployFolder))
 			this.IterativeDeploy = GetParamValueIfNotSpecified(Command, IterativeDeploy, this.IterativeDeploy, new string[] {"iterativedeploy", "iterate" } );
 			this.FastCook = GetParamValueIfNotSpecified(Command, FastCook, this.FastCook, "FastCook");
 			this.IgnoreCookErrors = GetParamValueIfNotSpecified(Command, IgnoreCookErrors, this.IgnoreCookErrors, "IgnoreCookErrors");
@@ -1605,10 +1623,13 @@ namespace AutomationTool
 		[Help("deploy", "deploy the project for the target platform")]
 		public bool Deploy { get; set; }
 
+		[Help("deploy", "Location to deploy to on the target platform")]
+		public string DeployFolder { get; set; }
+
 		#endregion
 
 		#region Misc
-		
+
 		[Help("MapsToRebuildLightMaps", "List of maps that need light maps rebuilding")]
 		public ParamList<string> MapsToRebuildLightMaps = new ParamList<string>();
 
diff --git a/Engine/Source/Programs/AutomationTool/AutomationUtils/ScriptCompiler.cs b/Engine/Source/Programs/AutomationTool/AutomationUtils/ScriptCompiler.cs
--- a/Engine/Source/Programs/AutomationTool/AutomationUtils/ScriptCompiler.cs
+++ b/Engine/Source/Programs/AutomationTool/AutomationUtils/ScriptCompiler.cs
@@ -93,21 +93,28 @@ namespace AutomationTool
 			ScriptCommands = new CaselessDictionary<Type>();
 			foreach (var CompiledScripts in ScriptAssemblies)
 			{
-				foreach (var ClassType in CompiledScripts.GetTypes())
+				try
 				{
-					if (ClassType.IsSubclassOf(typeof(BuildCommand)) && ClassType.IsAbstract == false)
+					foreach (var ClassType in CompiledScripts.GetTypes())
 					{
-						if (ScriptCommands.ContainsKey(ClassType.Name) == false)
+						if (ClassType.IsSubclassOf(typeof(BuildCommand)) && ClassType.IsAbstract == false)
 						{
-							ScriptCommands.Add(ClassType.Name, ClassType);
-						}
-						else
-						{
-							Log.TraceWarning("Unable to add command {0} twice. Previous: {1}, Current: {2}", ClassType.Name,
-								ClassType.AssemblyQualifiedName, ScriptCommands[ClassType.Name].AssemblyQualifiedName);
+							if (ScriptCommands.ContainsKey(ClassType.Name) == false)
+							{
+								ScriptCommands.Add(ClassType.Name, ClassType);
+							}
+							else
+							{
+								Log.TraceWarning("Unable to add command {0} twice. Previous: {1}, Current: {2}", ClassType.Name,
+									ClassType.AssemblyQualifiedName, ScriptCommands[ClassType.Name].AssemblyQualifiedName);
+							}
 						}
 					}
 				}
+				catch (Exception Ex)
+				{
+					throw new AutomationException("Failed to add commands from {0}. {1}", CompiledScripts, Ex);
+				}
 			}
 		}
 
diff --git a/Engine/Source/Programs/AutomationTool/BuildGraph/Tasks/StageTask.cs b/Engine/Source/Programs/AutomationTool/BuildGraph/Tasks/StageTask.cs
--- a/Engine/Source/Programs/AutomationTool/BuildGraph/Tasks/StageTask.cs
+++ b/Engine/Source/Programs/AutomationTool/BuildGraph/Tasks/StageTask.cs
@@ -153,6 +153,12 @@ namespace BuildGraph.Tasks
 					TargetFile = FileReference.Combine(TargetProjectDir, SourceFile.MakeRelativeTo(SourceProjectDir));
 				}
 
+				// Fixup the case of the output file. Would expect Platform.DeployLowerCaseFilenames() to return true here, but seems not to be the case.
+				if(Parameters.Platform == UnrealTargetPlatform.PS4)
+				{
+					TargetFile = FileReference.Combine(TargetDir, TargetFile.MakeRelativeTo(TargetDir).ToLowerInvariant());
+				}
+
 				// Only copy the output file if it doesn't already exist. We can stage multiple targets to the same output directory.
 				if(Parameters.Overwrite || !TargetFile.Exists())
 				{
