Copying //UE4/Orion-Staging to //UE4/Main (Source: //Orion/Dev-General @ 3564337)

Change 3564610 on 2017/07/31 by [REDACTED_PERSON]

	Integrated CL 3543210 : Fixed an issue when computing material scales where the default material ends up being used instead of the required material.
	Deprecated previous material data as it was causing some waste.
	Integrated CL 3526859 : Texture mip bias is now reset whenever the streaming budget increases
	#!rb none
	#!tests played monolith2 on PS4

Change 3564585 on 2017/07/31 by robomerge

	#!ROBOMERGE-AUTHOR: ben.salem
	Merging using Dev-Gen_->_Release-42
	Killing old useless nodes, fixing perf reporting, turning on shallow tests, killing non-focus in-depth perf tests for now
	#!rb various people in devgen
	#!tests Ran a shallow test map.

	#!ROBOMERGE-SOURCE: CL 3564579 in //Orion/Release-42/... via CL 3564580
	#!ROBOMERGE-BOT: ORION (Main -> Dev-Cinematics)

Change 3564584 on 2017/07/31 by robomerge

	#!ROBOMERGE-AUTHOR: ben.salem
	Merging using Dev-Gen_->_Release-42
	Killing old useless nodes, fixing perf reporting, turning on shallow tests, killing non-focus in-depth perf tests for now
	#!rb various people in devgen
	#!tests Ran a shallow test map.

	#!ROBOMERGE-SOURCE: CL 3564579 in //Orion/Release-42/... via CL 3564580
	#!ROBOMERGE-BOT: ORION (Main -> Dev-Balance)

Change 3564583 on 2017/07/31 by robomerge

	#!ROBOMERGE-AUTHOR: ben.salem
	Merging using Dev-Gen_->_Release-42
	Killing old useless nodes, fixing perf reporting, turning on shallow tests, killing non-focus in-depth perf tests for now
	#!rb various people in devgen
	#!tests Ran a shallow test map.

	#!ROBOMERGE-SOURCE: CL 3564579 in //Orion/Release-42/... via CL 3564580
	#!ROBOMERGE-BOT: ORION (Main -> Dev-REGS)

Change 3564582 on 2017/07/31 by robomerge

	#!ROBOMERGE-AUTHOR: ben.salem
	Merging using Dev-Gen_->_Release-42
	Killing old useless nodes, fixing perf reporting, turning on shallow tests, killing non-focus in-depth perf tests for now
	#!rb various people in devgen
	#!tests Ran a shallow test map.

	#!ROBOMERGE-SOURCE: CL 3564579 in //Orion/Release-42/... via CL 3564580
	#!ROBOMERGE-BOT: ORION (Main -> Dev-UI)

Change 3564580 on 2017/07/31 by robomerge

	#!ROBOMERGE-AUTHOR: ben.salem
	Merging using Dev-Gen_->_Release-42
	Killing old useless nodes, fixing perf reporting, turning on shallow tests, killing non-focus in-depth perf tests for now
	#!rb various people in devgen
	#!tests Ran a shallow test map.

	#!ROBOMERGE-SOURCE: CL 3564579 in //Orion/Release-42/...
	#!ROBOMERGE-BOT: ORION (Release-42 -> Main)

Change 3564579 on 2017/07/31 by [REDACTED_PERSON]

	Merging using Dev-Gen_->_Release-42
	Killing old useless nodes, fixing perf reporting, turning on shallow tests, killing non-focus in-depth perf tests for now
	#!rb various people in devgen
	#!tests Ran a shallow test map.

Change 3564513 on 2017/07/31 by robomerge

	#!ROBOMERGE-AUTHOR: laurent.delayen
	Fixed warning when opening Kwang AnimBP the first time, due to missing virtual bone name. (When Skeleton doesn't have PostLoad() called on it yet - happens only the first time the AnimBP is opened).

	#!rb none
	#!tests Kwang AnimBP opens without a warning.

	#!ROBOMERGE-SOURCE: CL 3564506 in //Orion/Release-42/... via CL 3564507
	#!ROBOMERGE-BOT: ORION (Main -> Dev-Cinematics)

Change 3564512 on 2017/07/31 by robomerge

	#!ROBOMERGE-AUTHOR: laurent.delayen
	Fixed warning when opening Kwang AnimBP the first time, due to missing virtual bone name. (When Skeleton doesn't have PostLoad() called on it yet - happens only the first time the AnimBP is opened).

	#!rb none
	#!tests Kwang AnimBP opens without a warning.

	#!ROBOMERGE-SOURCE: CL 3564506 in //Orion/Release-42/... via CL 3564507
	#!ROBOMERGE-BOT: ORION (Main -> Dev-Balance)

Change 3564511 on 2017/07/31 by robomerge

	#!ROBOMERGE-AUTHOR: laurent.delayen
	Fixed warning when opening Kwang AnimBP the first time, due to missing virtual bone name. (When Skeleton doesn't have PostLoad() called on it yet - happens only the first time the AnimBP is opened).

	#!rb none
	#!tests Kwang AnimBP opens without a warning.

	#!ROBOMERGE-SOURCE: CL 3564506 in //Orion/Release-42/... via CL 3564507
	#!ROBOMERGE-BOT: ORION (Main -> Dev-REGS)

Change 3564510 on 2017/07/31 by robomerge

	#!ROBOMERGE-AUTHOR: laurent.delayen
	Fixed warning when opening Kwang AnimBP the first time, due to missing virtual bone name. (When Skeleton doesn't have PostLoad() called on it yet - happens only the first time the AnimBP is opened).

	#!rb none
	#!tests Kwang AnimBP opens without a warning.

	#!ROBOMERGE-SOURCE: CL 3564506 in //Orion/Release-42/... via CL 3564507
	#!ROBOMERGE-BOT: ORION (Main -> Dev-UI)

Change 3564509 on 2017/07/31 by robomerge

	#!ROBOMERGE-AUTHOR: laurent.delayen
	Fixed warning when opening Kwang AnimBP the first time, due to missing virtual bone name. (When Skeleton doesn't have PostLoad() called on it yet - happens only the first time the AnimBP is opened).

	#!rb none
	#!tests Kwang AnimBP opens without a warning.

	#!ROBOMERGE-SOURCE: CL 3564506 in //Orion/Release-42/... via CL 3564507
	#!ROBOMERGE-BOT: ORION (Main -> Dev-General)

Change 3564507 on 2017/07/31 by robomerge

	#!ROBOMERGE-AUTHOR: laurent.delayen
	Fixed warning when opening Kwang AnimBP the first time, due to missing virtual bone name. (When Skeleton doesn't have PostLoad() called on it yet - happens only the first time the AnimBP is opened).

	#!rb none
	#!tests Kwang AnimBP opens without a warning.

	#!ROBOMERGE-SOURCE: CL 3564506 in //Orion/Release-42/...
	#!ROBOMERGE-BOT: ORION (Release-42 -> Main)

Change 3564506 on 2017/07/31 by [REDACTED_PERSON]

	Fixed warning when opening Kwang AnimBP the first time, due to missing virtual bone name. (When Skeleton doesn't have PostLoad() called on it yet - happens only the first time the AnimBP is opened).

	#!codereview jurre.debaare, dwayne.martin, lina.halper, martin.wilson
	#!rb none
	#!tests Kwang AnimBP opens without a warning.

Change 3564384 on 2017/07/31 by [REDACTED_PERSON]

	Now have a System Life Cycle module that looks for all the emitters being dead and then disables itself. This also triggers the reset of the simulation.
	GPU particles seems to have degraded after the spawn rate.

	Emitters now reset when there are no particles.
	Systems now reset when the state is Dead or Disabled, so you'll need to add a System Life Cycle component to have proper looping behavior for a system.

	#!rb none
	#!tests updated hypnotizer and other scripts

Change 3564012 on 2017/07/31 by robomerge

	#!ROBOMERGE-AUTHOR: laurent.delayen
	Added console command to disable URO interpolation.

	#!rb none
	#!tests ghost networked, simulated proxy.

	#!ROBOMERGE-SOURCE: CL 3564003 in //Orion/Release-42/... via CL 3564005
	#!ROBOMERGE-BOT: ORION (Main -> Dev-Cinematics)

Change 3564009 on 2017/07/31 by robomerge

	#!ROBOMERGE-AUTHOR: laurent.delayen
	Added console command to disable URO interpolation.

	#!rb none
	#!tests ghost networked, simulated proxy.

	#!ROBOMERGE-SOURCE: CL 3564003 in //Orion/Release-42/... via CL 3564005
	#!ROBOMERGE-BOT: ORION (Main -> Dev-Balance)

Change 3564008 on 2017/07/31 by robomerge

	#!ROBOMERGE-AUTHOR: laurent.delayen
	Added console command to disable URO interpolation.

	#!rb none
	#!tests ghost networked, simulated proxy.

	#!ROBOMERGE-SOURCE: CL 3564003 in //Orion/Release-42/... via CL 3564005
	#!ROBOMERGE-BOT: ORION (Main -> Dev-REGS)

Change 3564007 on 2017/07/31 by robomerge

	#!ROBOMERGE-AUTHOR: laurent.delayen
	Added console command to disable URO interpolation.

	#!rb none
	#!tests ghost networked, simulated proxy.

	#!ROBOMERGE-SOURCE: CL 3564003 in //Orion/Release-42/... via CL 3564005
	#!ROBOMERGE-BOT: ORION (Main -> Dev-UI)

Change 3564006 on 2017/07/31 by robomerge

	#!ROBOMERGE-AUTHOR: laurent.delayen
	Added console command to disable URO interpolation.

	#!rb none
	#!tests ghost networked, simulated proxy.

	#!ROBOMERGE-SOURCE: CL 3564003 in //Orion/Release-42/... via CL 3564005
	#!ROBOMERGE-BOT: ORION (Main -> Dev-General)

Change 3564005 on 2017/07/31 by robomerge

	#!ROBOMERGE-AUTHOR: laurent.delayen
	Added console command to disable URO interpolation.

	#!rb none
	#!tests ghost networked, simulated proxy.

	#!ROBOMERGE-SOURCE: CL 3564003 in //Orion/Release-42/...
	#!ROBOMERGE-BOT: ORION (Release-42 -> Main)

Change 3564003 on 2017/07/31 by [REDACTED_PERSON]

	Added console command to disable URO interpolation.

	#!codereview martin.wilson, lina.halper
	#!rb none
	#!tests ghost networked, simulated proxy.

Change 3563538 on 2017/07/30 by [REDACTED_PERSON]

	Niagara - Stack data interface editing fixes
	+ When a data interface object is modified by the stack, refresh the curves UI and re-initialize the simulation.
	+ Generate better names for the inputs used by data interfaces.

	#!Tests The curve UI and simulation update correctly when modifying the curve data interfaces in the stack and the generated inputs for data interfaces have better names.
	#!rb none

	#!codereview Olaf.Piesche,Simon.Tovey,Shaun.Kime

Change 3563537 on 2017/07/30 by [REDACTED_PERSON]

	Niagara - Fix the background color for stack errors.

	#!Tests Stack errors are no longer white.
	#!rb none

Change 3563531 on 2017/07/30 by [REDACTED_PERSON]

	Niagara - Generate stack spacer keys more safely to prevent list view crashes.

	#!Tests adding an emitter spawn module no longer crashes.
	#!rb none

	#!codereview Olaf.Piesche,Simon.Tovey,Shaun.Kime

Change 3563518 on 2017/07/30 by [REDACTED_PERSON]

	Niagara - Give parameter map error log message more context

	#!Tests none
	#!rb none

	#!codereview Olaf.Piesche,Simon.Tovey,Shaun.Kime

Change 3563384 on 2017/07/29 by robomerge

	#!ROBOMERGE-AUTHOR: andrew.grant
	Fixed some commandline logic issues in LoadTest

	#!tests ran locally
	#!rb none

	#!ROBOMERGE-SOURCE: CL 3563375 in //Orion/Release-42/... via CL 3563379
	#!ROBOMERGE-BOT: ORION (Main -> Dev-Cinematics)

Change 3563383 on 2017/07/29 by robomerge

	#!ROBOMERGE-AUTHOR: andrew.grant
	Fixed some commandline logic issues in LoadTest

	#!tests ran locally
	#!rb none

	#!ROBOMERGE-SOURCE: CL 3563375 in //Orion/Release-42/... via CL 3563379
	#!ROBOMERGE-BOT: ORION (Main -> Dev-Balance)

Change 3563382 on 2017/07/29 by robomerge

	#!ROBOMERGE-AUTHOR: andrew.grant
	Fixed some commandline logic issues in LoadTest

	#!tests ran locally
	#!rb none

	#!ROBOMERGE-SOURCE: CL 3563375 in //Orion/Release-42/... via CL 3563379
	#!ROBOMERGE-BOT: ORION (Main -> Dev-REGS)

Change 3563381 on 2017/07/29 by robomerge

	#!ROBOMERGE-AUTHOR: andrew.grant
	Fixed some commandline logic issues in LoadTest

	#!tests ran locally
	#!rb none

	#!ROBOMERGE-SOURCE: CL 3563375 in //Orion/Release-42/... via CL 3563379
	#!ROBOMERGE-BOT: ORION (Main -> Dev-UI)

Change 3563380 on 2017/07/29 by robomerge

	#!ROBOMERGE-AUTHOR: andrew.grant
	Fixed some commandline logic issues in LoadTest

	#!tests ran locally
	#!rb none

	#!ROBOMERGE-SOURCE: CL 3563375 in //Orion/Release-42/... via CL 3563379
	#!ROBOMERGE-BOT: ORION (Main -> Dev-General)

Change 3563379 on 2017/07/29 by robomerge

	#!ROBOMERGE-AUTHOR: andrew.grant
	Fixed some commandline logic issues in LoadTest

	#!tests ran locally
	#!rb none

	#!ROBOMERGE-SOURCE: CL 3563375 in //Orion/Release-42/...
	#!ROBOMERGE-BOT: ORION (Release-42 -> Main)

Change 3563375 on 2017/07/29 by [REDACTED_PERSON]

	Fixed some commandline logic issues in LoadTest

	#!tests ran locally
	#!rb none

Change 3563307 on 2017/07/29 by [REDACTED_PERSON]

	Niagara - Stack UI Rework
	+ Refactor most of the stack layout code to make things more consistent and to make future features possible.
	+ Add a hover cue for item rows.
	+ Add icons for the different types of inputs.
	+ Make inputs collapsible.
	+ Move the pin buttons to the right side of the name column to prevent visual clutter with the expanders.
	+ Make the module splitter visible and add a correct hover cue.

	#!Tests Stack functions correctly.
	#!rb none

	#!codereview Olaf.Piesche,Simon.Tovey,Shaun.Kime

Change 3563305 on 2017/07/29 by [REDACTED_PERSON]

	Add Shallow FX Test node to gauntlet and to orionbuild. Also switched Dev-Gen to being the Deep Test branch instead of dev-ui.
	#!rb none
	#!tests Ran a test of the new node, preflighted orionbuild.xml changes.

Change 3563205 on 2017/07/29 by robomerge

	#!ROBOMERGE-AUTHOR: andrew.grant
	Add an exception handler around post-test Gif creation.

	Added -attended option to tests.

	#!tests compiled
	#!rb none

	#!ROBOMERGE-SOURCE: CL 3563199 in //Orion/Release-42/... via CL 3563200
	#!ROBOMERGE-BOT: ORION (Main -> Dev-Cinematics)

Change 3563204 on 2017/07/29 by robomerge

	#!ROBOMERGE-AUTHOR: andrew.grant
	Add an exception handler around post-test Gif creation.

	Added -attended option to tests.

	#!tests compiled
	#!rb none

	#!ROBOMERGE-SOURCE: CL 3563199 in //Orion/Release-42/... via CL 3563200
	#!ROBOMERGE-BOT: ORION (Main -> Dev-Balance)

Change 3563203 on 2017/07/29 by robomerge

	#!ROBOMERGE-AUTHOR: andrew.grant
	Add an exception handler around post-test Gif creation.

	Added -attended option to tests.

	#!tests compiled
	#!rb none

	#!ROBOMERGE-SOURCE: CL 3563199 in //Orion/Release-42/... via CL 3563200
	#!ROBOMERGE-BOT: ORION (Main -> Dev-REGS)

Change 3563202 on 2017/07/29 by robomerge

	#!ROBOMERGE-AUTHOR: andrew.grant
	Add an exception handler around post-test Gif creation.

	Added -attended option to tests.

	#!tests compiled
	#!rb none

	#!ROBOMERGE-SOURCE: CL 3563199 in //Orion/Release-42/... via CL 3563200
	#!ROBOMERGE-BOT: ORION (Main -> Dev-UI)

Change 3563201 on 2017/07/29 by robomerge

	#!ROBOMERGE-AUTHOR: andrew.grant
	Add an exception handler around post-test Gif creation.

	Added -attended option to tests.

	#!tests compiled
	#!rb none

	#!ROBOMERGE-SOURCE: CL 3563199 in //Orion/Release-42/... via CL 3563200
	#!ROBOMERGE-BOT: ORION (Main -> Dev-General)

Change 3563200 on 2017/07/29 by robomerge

	#!ROBOMERGE-AUTHOR: andrew.grant
	Add an exception handler around post-test Gif creation.

	Added -attended option to tests.

	#!tests compiled
	#!rb none

	#!ROBOMERGE-SOURCE: CL 3563199 in //Orion/Release-42/...
	#!ROBOMERGE-BOT: ORION (Release-42 -> Main)

Change 3563199 on 2017/07/29 by [REDACTED_PERSON]

	Add an exception handler around post-test Gif creation.

	Added -attended option to tests.

	#!tests compiled
	#!rb none

Change 3563187 on 2017/07/29 by robomerge

	#!ROBOMERGE-AUTHOR: andrew.grant
	Fix an issue where we'd try to set a file attriute before copying it (!)

	Turn failure of handling loadorder file into a warning

	#!tests compiled.
	#!rb none

	#!ROBOMERGE-SOURCE: CL 3563181 in //Orion/Release-42/... via CL 3563182
	#!ROBOMERGE-BOT: ORION (Main -> Dev-Cinematics)

Change 3563186 on 2017/07/29 by robomerge

	#!ROBOMERGE-AUTHOR: andrew.grant
	Fix an issue where we'd try to set a file attriute before copying it (!)

	Turn failure of handling loadorder file into a warning

	#!tests compiled.
	#!rb none

	#!ROBOMERGE-SOURCE: CL 3563181 in //Orion/Release-42/... via CL 3563182
	#!ROBOMERGE-BOT: ORION (Main -> Dev-Balance)

Change 3563185 on 2017/07/29 by robomerge

	#!ROBOMERGE-AUTHOR: andrew.grant
	Fix an issue where we'd try to set a file attriute before copying it (!)

	Turn failure of handling loadorder file into a warning

	#!tests compiled.
	#!rb none

	#!ROBOMERGE-SOURCE: CL 3563181 in //Orion/Release-42/... via CL 3563182
	#!ROBOMERGE-BOT: ORION (Main -> Dev-REGS)

Change 3563184 on 2017/07/29 by robomerge

	#!ROBOMERGE-AUTHOR: andrew.grant
	Fix an issue where we'd try to set a file attriute before copying it (!)

	Turn failure of handling loadorder file into a warning

	#!tests compiled.
	#!rb none

	#!ROBOMERGE-SOURCE: CL 3563181 in //Orion/Release-42/... via CL 3563182
	#!ROBOMERGE-BOT: ORION (Main -> Dev-UI)

Change 3563183 on 2017/07/29 by robomerge

	#!ROBOMERGE-AUTHOR: andrew.grant
	Fix an issue where we'd try to set a file attriute before copying it (!)

	Turn failure of handling loadorder file into a warning

	#!tests compiled.
	#!rb none

	#!ROBOMERGE-SOURCE: CL 3563181 in //Orion/Release-42/... via CL 3563182
	#!ROBOMERGE-BOT: ORION (Main -> Dev-General)

Change 3563182 on 2017/07/29 by robomerge

	#!ROBOMERGE-AUTHOR: andrew.grant
	Fix an issue where we'd try to set a file attriute before copying it (!)

	Turn failure of handling loadorder file into a warning

	#!tests compiled.
	#!rb none

	#!ROBOMERGE-SOURCE: CL 3563181 in //Orion/Release-42/...
	#!ROBOMERGE-BOT: ORION (Release-42 -> Main)

Change 3563181 on 2017/07/29 by [REDACTED_PERSON]

	Fix an issue where we'd try to set a file attriute before copying it (!)

	Turn failure of handling loadorder file into a warning

	#!tests compiled.
	#!rb none

Change 3562983 on 2017/07/28 by robomerge

	#!ROBOMERGE-AUTHOR: dan.hertzka
	Fixed all orion tooltip scaling & placement issues
	- Tooltip itself is always drawn unscaled, regardless of the anchor's layout or render scaling
	- Tooltip is anchored based on both the layout and render scaling, so the unscaled tooltip still appears in the correct spot relative to the scaled anchor
	- Finally, all tooltips are always drawn at full opacity and with no tint, regardless of the tint/alpha on the anchor
	- Unfortunately this couldn't all just be added direcly to SMenuAnchor. It's in proper Slate land and unable to access the game viewport's DPI scale.

	Made a few small engine-level changes to SMenuAnchor:
	- Added bApplyWidgetStyleToMenu - if false, the popup is given a default FWidgetStyle when it's painted
	- Moved the FPopupPlacement declaration to SMenuAnchor.h, but it's a protected declaration within the widget

	[at Nick.Darnell,] [at Don.Eubanks] [FYI] Matt.Schembari, Philip.Buuck, Stephan.Jiang
	#!rb none
	#!tests Editor tooltips are fine; PIE Frontend - checked that both the deck builder gem tree gems and the side entries in the chest selection screen appear properly (good examples of layout scaling and pure render scaling)

	#!ROBOMERGE-SOURCE: CL 3562969 in //Orion/Release-42/... via CL 3562978
	#!ROBOMERGE-BOT: ORION (Main -> Dev-Cinematics)

Change 3562982 on 2017/07/28 by robomerge

	#!ROBOMERGE-AUTHOR: dan.hertzka
	Fixed all orion tooltip scaling & placement issues
	- Tooltip itself is always drawn unscaled, regardless of the anchor's layout or render scaling
	- Tooltip is anchored based on both the layout and render scaling, so the unscaled tooltip still appears in the correct spot relative to the scaled anchor
	- Finally, all tooltips are always drawn at full opacity and with no tint, regardless of the tint/alpha on the anchor
	- Unfortunately this couldn't all just be added direcly to SMenuAnchor. It's in proper Slate land and unable to access the game viewport's DPI scale.

	Made a few small engine-level changes to SMenuAnchor:
	- Added bApplyWidgetStyleToMenu - if false, the popup is given a default FWidgetStyle when it's painted
	- Moved the FPopupPlacement declaration to SMenuAnchor.h, but it's a protected declaration within the widget

	[at Nick.Darnell,] [at Don.Eubanks] [FYI] Matt.Schembari, Philip.Buuck, Stephan.Jiang
	#!rb none
	#!tests Editor tooltips are fine; PIE Frontend - checked that both the deck builder gem tree gems and the side entries in the chest selection screen appear properly (good examples of layout scaling and pure render scaling)

	#!ROBOMERGE-SOURCE: CL 3562969 in //Orion/Release-42/... via CL 3562978
	#!ROBOMERGE-BOT: ORION (Main -> Dev-Balance)

Change 3562981 on 2017/07/28 by robomerge

	#!ROBOMERGE-AUTHOR: dan.hertzka
	Fixed all orion tooltip scaling & placement issues
	- Tooltip itself is always drawn unscaled, regardless of the anchor's layout or render scaling
	- Tooltip is anchored based on both the layout and render scaling, so the unscaled tooltip still appears in the correct spot relative to the scaled anchor
	- Finally, all tooltips are always drawn at full opacity and with no tint, regardless of the tint/alpha on the anchor
	- Unfortunately this couldn't all just be added direcly to SMenuAnchor. It's in proper Slate land and unable to access the game viewport's DPI scale.

	Made a few small engine-level changes to SMenuAnchor:
	- Added bApplyWidgetStyleToMenu - if false, the popup is given a default FWidgetStyle when it's painted
	- Moved the FPopupPlacement declaration to SMenuAnchor.h, but it's a protected declaration within the widget

	[at Nick.Darnell,] [at Don.Eubanks] [FYI] Matt.Schembari, Philip.Buuck, Stephan.Jiang
	#!rb none
	#!tests Editor tooltips are fine; PIE Frontend - checked that both the deck builder gem tree gems and the side entries in the chest selection screen appear properly (good examples of layout scaling and pure render scaling)

	#!ROBOMERGE-SOURCE: CL 3562969 in //Orion/Release-42/... via CL 3562978
	#!ROBOMERGE-BOT: ORION (Main -> Dev-REGS)

Change 3562980 on 2017/07/28 by robomerge

	#!ROBOMERGE-AUTHOR: dan.hertzka
	Fixed all orion tooltip scaling & placement issues
	- Tooltip itself is always drawn unscaled, regardless of the anchor's layout or render scaling
	- Tooltip is anchored based on both the layout and render scaling, so the unscaled tooltip still appears in the correct spot relative to the scaled anchor
	- Finally, all tooltips are always drawn at full opacity and with no tint, regardless of the tint/alpha on the anchor
	- Unfortunately this couldn't all just be added direcly to SMenuAnchor. It's in proper Slate land and unable to access the game viewport's DPI scale.

	Made a few small engine-level changes to SMenuAnchor:
	- Added bApplyWidgetStyleToMenu - if false, the popup is given a default FWidgetStyle when it's painted
	- Moved the FPopupPlacement declaration to SMenuAnchor.h, but it's a protected declaration within the widget

	[at Nick.Darnell,] [at Don.Eubanks] [FYI] Matt.Schembari, Philip.Buuck, Stephan.Jiang
	#!rb none
	#!tests Editor tooltips are fine; PIE Frontend - checked that both the deck builder gem tree gems and the side entries in the chest selection screen appear properly (good examples of layout scaling and pure render scaling)

	#!ROBOMERGE-SOURCE: CL 3562969 in //Orion/Release-42/... via CL 3562978
	#!ROBOMERGE-BOT: ORION (Main -> Dev-UI)

Change 3562979 on 2017/07/28 by robomerge

	#!ROBOMERGE-AUTHOR: dan.hertzka
	Fixed all orion tooltip scaling & placement issues
	- Tooltip itself is always drawn unscaled, regardless of the anchor's layout or render scaling
	- Tooltip is anchored based on both the layout and render scaling, so the unscaled tooltip still appears in the correct spot relative to the scaled anchor
	- Finally, all tooltips are always drawn at full opacity and with no tint, regardless of the tint/alpha on the anchor
	- Unfortunately this couldn't all just be added direcly to SMenuAnchor. It's in proper Slate land and unable to access the game viewport's DPI scale.

	Made a few small engine-level changes to SMenuAnchor:
	- Added bApplyWidgetStyleToMenu - if false, the popup is given a default FWidgetStyle when it's painted
	- Moved the FPopupPlacement declaration to SMenuAnchor.h, but it's a protected declaration within the widget

	[at Nick.Darnell,] [at Don.Eubanks] [FYI] Matt.Schembari, Philip.Buuck, Stephan.Jiang
	#!rb none
	#!tests Editor tooltips are fine; PIE Frontend - checked that both the deck builder gem tree gems and the side entries in the chest selection screen appear properly (good examples of layout scaling and pure render scaling)

	#!ROBOMERGE-SOURCE: CL 3562969 in //Orion/Release-42/... via CL 3562978
	#!ROBOMERGE-BOT: ORION (Main -> Dev-General)

Change 3562978 on 2017/07/28 by robomerge

	#!ROBOMERGE-AUTHOR: dan.hertzka
	Fixed all orion tooltip scaling & placement issues
	- Tooltip itself is always drawn unscaled, regardless of the anchor's layout or render scaling
	- Tooltip is anchored based on both the layout and render scaling, so the unscaled tooltip still appears in the correct spot relative to the scaled anchor
	- Finally, all tooltips are always drawn at full opacity and with no tint, regardless of the tint/alpha on the anchor
	- Unfortunately this couldn't all just be added direcly to SMenuAnchor. It's in proper Slate land and unable to access the game viewport's DPI scale.

	Made a few small engine-level changes to SMenuAnchor:
	- Added bApplyWidgetStyleToMenu - if false, the popup is given a default FWidgetStyle when it's painted
	- Moved the FPopupPlacement declaration to SMenuAnchor.h, but it's a protected declaration within the widget

	[at Nick.Darnell,] [at Don.Eubanks] [FYI] Matt.Schembari, Philip.Buuck, Stephan.Jiang
	#!rb none
	#!tests Editor tooltips are fine; PIE Frontend - checked that both the deck builder gem tree gems and the side entries in the chest selection screen appear properly (good examples of layout scaling and pure render scaling)

	#!ROBOMERGE-SOURCE: CL 3562969 in //Orion/Release-42/...
	#!ROBOMERGE-BOT: ORION (Release-42 -> Main)

Change 3562977 on 2017/07/28 by robomerge

	#!ROBOMERGE-AUTHOR: andrew.grant
	Editgration of 3437205 from Dev-Framework to address issues with Blueprint references being incorrectly collected

	#!rb none
	[at daniel.lamb] #!tests LoadTest locally on cooked data on PS4/Win64

	#!ROBOMERGE-SOURCE: CL 3562966 in //Orion/Release-42/... via CL 3562970
	#!ROBOMERGE-BOT: ORION (Main -> Dev-Cinematics)

Change 3562976 on 2017/07/28 by robomerge

	#!ROBOMERGE-AUTHOR: andrew.grant
	Editgration of 3437205 from Dev-Framework to address issues with Blueprint references being incorrectly collected

	#!rb none
	[at daniel.lamb] #!tests LoadTest locally on cooked data on PS4/Win64

	#!ROBOMERGE-SOURCE: CL 3562966 in //Orion/Release-42/... via CL 3562970
	#!ROBOMERGE-BOT: ORION (Main -> Dev-Balance)

Change 3562975 on 2017/07/28 by robomerge

	#!ROBOMERGE-AUTHOR: andrew.grant
	Editgration of 3437205 from Dev-Framework to address issues with Blueprint references being incorrectly collected

	#!rb none
	[at daniel.lamb] #!tests LoadTest locally on cooked data on PS4/Win64

	#!ROBOMERGE-SOURCE: CL 3562966 in //Orion/Release-42/... via CL 3562970
	#!ROBOMERGE-BOT: ORION (Main -> Dev-REGS)

Change 3562974 on 2017/07/28 by robomerge

	#!ROBOMERGE-AUTHOR: andrew.grant
	Editgration of 3437205 from Dev-Framework to address issues with Blueprint references being incorrectly collected

	#!rb none
	[at daniel.lamb] #!tests LoadTest locally on cooked data on PS4/Win64

	#!ROBOMERGE-SOURCE: CL 3562966 in //Orion/Release-42/... via CL 3562970
	#!ROBOMERGE-BOT: ORION (Main -> Dev-UI)

Change 3562973 on 2017/07/28 by robomerge

	#!ROBOMERGE-AUTHOR: andrew.grant
	Editgration of 3437205 from Dev-Framework to address issues with Blueprint references being incorrectly collected

	#!rb none
	[at daniel.lamb] #!tests LoadTest locally on cooked data on PS4/Win64

	#!ROBOMERGE-SOURCE: CL 3562966 in //Orion/Release-42/... via CL 3562970
	#!ROBOMERGE-BOT: ORION (Main -> Dev-General)

Change 3562970 on 2017/07/28 by robomerge

	#!ROBOMERGE-AUTHOR: andrew.grant
	Editgration of 3437205 from Dev-Framework to address issues with Blueprint references being incorrectly collected

	#!rb none
	[at daniel.lamb] #!tests LoadTest locally on cooked data on PS4/Win64

	#!ROBOMERGE-SOURCE: CL 3562966 in //Orion/Release-42/...
	#!ROBOMERGE-BOT: ORION (Release-42 -> Main)

Change 3562969 on 2017/07/28 by [REDACTED_PERSON]

	Fixed all orion tooltip scaling & placement issues
	- Tooltip itself is always drawn unscaled, regardless of the anchor's layout or render scaling
	- Tooltip is anchored based on both the layout and render scaling, so the unscaled tooltip still appears in the correct spot relative to the scaled anchor
	- Finally, all tooltips are always drawn at full opacity and with no tint, regardless of the tint/alpha on the anchor
	- Unfortunately this couldn't all just be added direcly to SMenuAnchor. It's in proper Slate land and unable to access the game viewport's DPI scale.

	Made a few small engine-level changes to SMenuAnchor:
	- Added bApplyWidgetStyleToMenu - if false, the popup is given a default FWidgetStyle when it's painted
	- Moved the FPopupPlacement declaration to SMenuAnchor.h, but it's a protected declaration within the widget

	[OR-41642] - Alpha is no longer applied to the chest tooltips. Also, the chests on the edge won't have their tooltip clip off the screen.

	#!review-3562971 [REDACTED_HANDLE], [REDACTED_HANDLE]
	#!fyi Matt.Schembari, Philip.Buuck, Stephan.Jiang
	#!rb none
	#!tests Editor tooltips are fine; PIE Frontend - checked that both the deck builder gem tree gems and the side entries in the chest selection screen appear properly (good examples of layout scaling and pure render scaling)

	#!QAReview Let me know if you come across any tooltips that are blatantly huge, tiny, or in an incorrect place

Change 3562966 on 2017/07/28 by [REDACTED_PERSON]

	Editgration of 3437205 from Dev-Framework to address issues with Blueprint references being incorrectly collected

	#!rb none
	#!review-3562967 [REDACTED_HANDLE]
	#!tests LoadTest locally on cooked data on PS4/Win64

Change 3562965 on 2017/07/28 by robomerge

	#!ROBOMERGE-AUTHOR: andrew.grant
	Changed to LoadTest to prevent it timing out on PS4

	#!tests tested locally
	#!rb none

	#!ROBOMERGE-SOURCE: CL 3562959 in //Orion/Release-42/... via CL 3562960
	#!ROBOMERGE-BOT: ORION (Main -> Dev-Cinematics)

Change 3562964 on 2017/07/28 by robomerge

	#!ROBOMERGE-AUTHOR: andrew.grant
	Changed to LoadTest to prevent it timing out on PS4

	#!tests tested locally
	#!rb none

	#!ROBOMERGE-SOURCE: CL 3562959 in //Orion/Release-42/... via CL 3562960
	#!ROBOMERGE-BOT: ORION (Main -> Dev-Balance)

Change 3562963 on 2017/07/28 by robomerge

	#!ROBOMERGE-AUTHOR: andrew.grant
	Changed to LoadTest to prevent it timing out on PS4

	#!tests tested locally
	#!rb none

	#!ROBOMERGE-SOURCE: CL 3562959 in //Orion/Release-42/... via CL 3562960
	#!ROBOMERGE-BOT: ORION (Main -> Dev-REGS)

Change 3562962 on 2017/07/28 by robomerge

	#!ROBOMERGE-AUTHOR: andrew.grant
	Changed to LoadTest to prevent it timing out on PS4

	#!tests tested locally
	#!rb none

	#!ROBOMERGE-SOURCE: CL 3562959 in //Orion/Release-42/... via CL 3562960
	#!ROBOMERGE-BOT: ORION (Main -> Dev-UI)

Change 3562961 on 2017/07/28 by robomerge

	#!ROBOMERGE-AUTHOR: andrew.grant
	Changed to LoadTest to prevent it timing out on PS4

	#!tests tested locally
	#!rb none

	#!ROBOMERGE-SOURCE: CL 3562959 in //Orion/Release-42/... via CL 3562960
	#!ROBOMERGE-BOT: ORION (Main -> Dev-General)

Change 3562960 on 2017/07/28 by robomerge

	#!ROBOMERGE-AUTHOR: andrew.grant
	Changed to LoadTest to prevent it timing out on PS4

	#!tests tested locally
	#!rb none

	#!ROBOMERGE-SOURCE: CL 3562959 in //Orion/Release-42/...
	#!ROBOMERGE-BOT: ORION (Release-42 -> Main)

Change 3562959 on 2017/07/28 by [REDACTED_PERSON]

	Changed to LoadTest to prevent it timing out on PS4

	#!tests tested locally
	#!rb none

Change 3562136 on 2017/07/28 by [REDACTED_PERSON]

	Changing the version so that old assets will recompile and regenerate their spawn attribute table

	#!rb none
	#!code.review simon.tovey
	#!tests opened asset and made sure it compiled on load

Change 3560805 on 2017/07/28 by [REDACTED_PERSON]

	- Programmable spawning
	All spawning controlled by creating a FNiagaraSpawnInfo attribute. Any of these attributes in an emitter will feed one spawn script run.
	- Fixed issue with HLSL and register table layout not matching for structs correctly.
	- Removed some vestigial code.
	- Temporarily commenting out references to burst in the UI until we can hook them back up.
	- Removed direct ref to emitter handle in emitter instances with an EmitterIndex in their parent. More broadly useful and can be used to access emitter handle.
	- Fixed a couple of issues breaking interpolated spawning.
	- Updated default emitter and the hypnotiser to new spawning method.

	#!rb none
	#!tests Tested new default emitter and a few others.
	#!codereview Olaf.Piesche, Frank.Fella, Shaun.Kime

Change 3560376 on 2017/07/27 by robomerge

	#!ROBOMERGE-AUTHOR: stephan.jiang
	OrionEditableTextBox max count

	  -- This way there is a max count for Deck names so they won't go over above 50 characters.

	#!rb Dan.Hertzka
	#!test PIE

	#!ROBOMERGE-SOURCE: CL 3560367 in //Orion/Release-42/... via CL 3560370
	#!ROBOMERGE-BOT: ORION (Main -> Dev-Cinematics)

Change 3560375 on 2017/07/27 by robomerge

	#!ROBOMERGE-AUTHOR: stephan.jiang
	OrionEditableTextBox max count

	  -- This way there is a max count for Deck names so they won't go over above 50 characters.

	#!rb Dan.Hertzka
	#!test PIE

	#!ROBOMERGE-SOURCE: CL 3560367 in //Orion/Release-42/... via CL 3560370
	#!ROBOMERGE-BOT: ORION (Main -> Dev-Balance)

Change 3560374 on 2017/07/27 by robomerge

	#!ROBOMERGE-AUTHOR: stephan.jiang
	OrionEditableTextBox max count

	  -- This way there is a max count for Deck names so they won't go over above 50 characters.

	#!rb Dan.Hertzka
	#!test PIE

	#!ROBOMERGE-SOURCE: CL 3560367 in //Orion/Release-42/... via CL 3560370
	#!ROBOMERGE-BOT: ORION (Main -> Dev-REGS)

Change 3560373 on 2017/07/27 by robomerge

	#!ROBOMERGE-AUTHOR: stephan.jiang
	OrionEditableTextBox max count

	  -- This way there is a max count for Deck names so they won't go over above 50 characters.

	#!rb Dan.Hertzka
	#!test PIE

	#!ROBOMERGE-SOURCE: CL 3560367 in //Orion/Release-42/... via CL 3560370
	#!ROBOMERGE-BOT: ORION (Main -> Dev-UI)

Change 3560372 on 2017/07/27 by robomerge

	#!ROBOMERGE-AUTHOR: stephan.jiang
	OrionEditableTextBox max count

	  -- This way there is a max count for Deck names so they won't go over above 50 characters.

	#!rb Dan.Hertzka
	#!test PIE

	#!ROBOMERGE-SOURCE: CL 3560367 in //Orion/Release-42/... via CL 3560370
	#!ROBOMERGE-BOT: ORION (Main -> Dev-General)

Change 3560370 on 2017/07/27 by robomerge

	#!ROBOMERGE-AUTHOR: stephan.jiang
	OrionEditableTextBox max count

	  -- This way there is a max count for Deck names so they won't go over above 50 characters.

	#!rb Dan.Hertzka
	#!test PIE

	#!ROBOMERGE-SOURCE: CL 3560367 in //Orion/Release-42/...
	#!ROBOMERGE-BOT: ORION (Release-42 -> Main)

Change 3560367 on 2017/07/27 by [REDACTED_PERSON]

	OrionEditableTextBox max count

	  -- This way there is a max count for Deck names so they won't go over above 50 characters.

	#!rb Dan.Hertzka
	#!test PIE

Change 3560196 on 2017/07/27 by robomerge

	#!ROBOMERGE-AUTHOR: daniel.lamb
	Added more information to the logging output for OR40458.
	#!rb Trivial
	#!test Compile and run orion server / ps4 client

	#!ROBOMERGE-SOURCE: CL 3560180 in //Orion/Release-42/... via CL 3560183
	#!ROBOMERGE-BOT: ORION (Main -> Dev-Cinematics)

Change 3560192 on 2017/07/27 by robomerge

	#!ROBOMERGE-AUTHOR: daniel.lamb
	Added more information to the logging output for OR40458.
	#!rb Trivial
	#!test Compile and run orion server / ps4 client

	#!ROBOMERGE-SOURCE: CL 3560180 in //Orion/Release-42/... via CL 3560183
	#!ROBOMERGE-BOT: ORION (Main -> Dev-Balance)

Change 3560188 on 2017/07/27 by robomerge

	#!ROBOMERGE-AUTHOR: daniel.lamb
	Added more information to the logging output for OR40458.
	#!rb Trivial
	#!test Compile and run orion server / ps4 client

	#!ROBOMERGE-SOURCE: CL 3560180 in //Orion/Release-42/... via CL 3560183
	#!ROBOMERGE-BOT: ORION (Main -> Dev-REGS)

Change 3560186 on 2017/07/27 by robomerge

	#!ROBOMERGE-AUTHOR: daniel.lamb
	Added more information to the logging output for OR40458.
	#!rb Trivial
	#!test Compile and run orion server / ps4 client

	#!ROBOMERGE-SOURCE: CL 3560180 in //Orion/Release-42/... via CL 3560183
	#!ROBOMERGE-BOT: ORION (Main -> Dev-UI)

Change 3560185 on 2017/07/27 by robomerge

	#!ROBOMERGE-AUTHOR: daniel.lamb
	Added more information to the logging output for OR40458.
	#!rb Trivial
	#!test Compile and run orion server / ps4 client

	#!ROBOMERGE-SOURCE: CL 3560180 in //Orion/Release-42/... via CL 3560183
	#!ROBOMERGE-BOT: ORION (Main -> Dev-General)

Change 3560183 on 2017/07/27 by robomerge

	#!ROBOMERGE-AUTHOR: daniel.lamb
	Added more information to the logging output for OR40458.
	#!rb Trivial
	#!test Compile and run orion server / ps4 client

	#!ROBOMERGE-SOURCE: CL 3560180 in //Orion/Release-42/...
	#!ROBOMERGE-BOT: ORION (Release-42 -> Main)

Change 3560180 on 2017/07/27 by [REDACTED_PERSON]

	Added more information to the logging output for OR40458.
	#!rb Trivial
	#!test Compile and run orion server / ps4 client

Change 3560131 on 2017/07/27 by robomerge

	#!ROBOMERGE-AUTHOR: ori.cohen
	Fix rigid body node not working on ps4 due to fast path not allowing ragdolls to be created. This should not apply for animation.

	#!rb David.Hill
	#!jira OR-41774
	#!tests none

	#!ROBOMERGE-SOURCE: CL 3560123 in //Orion/Release-42/... via CL 3560126
	#!ROBOMERGE-BOT: ORION (Main -> Dev-Cinematics)

Change 3560130 on 2017/07/27 by robomerge

	#!ROBOMERGE-AUTHOR: ori.cohen
	Fix rigid body node not working on ps4 due to fast path not allowing ragdolls to be created. This should not apply for animation.

	#!rb David.Hill
	#!jira OR-41774
	#!tests none

	#!ROBOMERGE-SOURCE: CL 3560123 in //Orion/Release-42/... via CL 3560126
	#!ROBOMERGE-BOT: ORION (Main -> Dev-Balance)

Change 3560129 on 2017/07/27 by robomerge

	#!ROBOMERGE-AUTHOR: ori.cohen
	Fix rigid body node not working on ps4 due to fast path not allowing ragdolls to be created. This should not apply for animation.

	#!rb David.Hill
	#!jira OR-41774
	#!tests none

	#!ROBOMERGE-SOURCE: CL 3560123 in //Orion/Release-42/... via CL 3560126
	#!ROBOMERGE-BOT: ORION (Main -> Dev-REGS)

Change 3560128 on 2017/07/27 by robomerge

	#!ROBOMERGE-AUTHOR: ori.cohen
	Fix rigid body node not working on ps4 due to fast path not allowing ragdolls to be created. This should not apply for animation.

	#!rb David.Hill
	#!jira OR-41774
	#!tests none

	#!ROBOMERGE-SOURCE: CL 3560123 in //Orion/Release-42/... via CL 3560126
	#!ROBOMERGE-BOT: ORION (Main -> Dev-UI)

Change 3560127 on 2017/07/27 by robomerge

	#!ROBOMERGE-AUTHOR: ori.cohen
	Fix rigid body node not working on ps4 due to fast path not allowing ragdolls to be created. This should not apply for animation.

	#!rb David.Hill
	#!jira OR-41774
	#!tests none

	#!ROBOMERGE-SOURCE: CL 3560123 in //Orion/Release-42/... via CL 3560126
	#!ROBOMERGE-BOT: ORION (Main -> Dev-General)

Change 3560126 on 2017/07/27 by robomerge

	#!ROBOMERGE-AUTHOR: ori.cohen
	Fix rigid body node not working on ps4 due to fast path not allowing ragdolls to be created. This should not apply for animation.

	#!rb David.Hill
	#!jira OR-41774
	#!tests none

	#!ROBOMERGE-SOURCE: CL 3560123 in //Orion/Release-42/...
	#!ROBOMERGE-BOT: ORION (Release-42 -> Main)

Change 3560123 on 2017/07/27 by [REDACTED_PERSON]

	Fix rigid body node not working on ps4 due to fast path not allowing ragdolls to be created. This should not apply for animation.

	#!rb David.Hill
	#!jira OR-41774
	#!tests none

Change 3559908 on 2017/07/27 by [REDACTED_PERSON]

	Fixing compile error
	#!tests none
	#!rb none
	#!codereview Andrew.Grant

Change 3559674 on 2017/07/27 by [REDACTED_PERSON]

	Now batching up the shader constants into another data set for System/Emitter graphs.

	#!rb Simon.Tovey
	#!tests ran multiple copies of Hypnotizer and made sure that they obeyed the emitter lifetime module outputs.

Change 3559527 on 2017/07/27 by [REDACTED_PERSON]

	#!jira UE-45483 Integrating fix to //Orion/Dev-General
	#!rb none
	#!tests none

Change 3559284 on 2017/07/27 by robomerge

	#!ROBOMERGE-AUTHOR: laurent.delayen
	Exposed GetAzimuthAndElevation to blueprints.

	#!rb none
	#!tests Pyro turrets

	#!ROBOMERGE-SOURCE: CL 3559111 in //Orion/Release-42/... via CL 3559115
	#!ROBOMERGE-BOT: ORION (Main -> Dev-Cinematics)

Change 3559283 on 2017/07/27 by robomerge

	#!ROBOMERGE-AUTHOR: laurent.delayen
	Exposed GetAzimuthAndElevation to blueprints.

	#!rb none
	#!tests Pyro turrets

	#!ROBOMERGE-SOURCE: CL 3559111 in //Orion/Release-42/... via CL 3559115
	#!ROBOMERGE-BOT: ORION (Main -> Dev-Balance)

Change 3559282 on 2017/07/27 by robomerge

	#!ROBOMERGE-AUTHOR: laurent.delayen
	Exposed GetAzimuthAndElevation to blueprints.

	#!rb none
	#!tests Pyro turrets

	#!ROBOMERGE-SOURCE: CL 3559111 in //Orion/Release-42/... via CL 3559115
	#!ROBOMERGE-BOT: ORION (Main -> Dev-REGS)

Change 3559281 on 2017/07/27 by robomerge

	#!ROBOMERGE-AUTHOR: laurent.delayen
	Exposed GetAzimuthAndElevation to blueprints.

	#!rb none
	#!tests Pyro turrets

	#!ROBOMERGE-SOURCE: CL 3559111 in //Orion/Release-42/... via CL 3559115
	#!ROBOMERGE-BOT: ORION (Main -> Dev-UI)

Change 3559280 on 2017/07/27 by robomerge

	#!ROBOMERGE-AUTHOR: laurent.delayen
	Exposed GetAzimuthAndElevation to blueprints.

	#!rb none
	#!tests Pyro turrets

	#!ROBOMERGE-SOURCE: CL 3559111 in //Orion/Release-42/... via CL 3559115
	#!ROBOMERGE-BOT: ORION (Main -> Dev-General)

Change 3559254 on 2017/07/27 by robomerge

	#!ROBOMERGE-AUTHOR: jon.lietz
	compile fix

	#!rb none
	#!test compiles
	[at Daniel.Lamb]

	#!ROBOMERGE-SOURCE: CL 3559043 in //Orion/Release-42/... via CL 3559060
	#!ROBOMERGE-BOT: ORION (Main -> Dev-Cinematics)

Change 3559253 on 2017/07/27 by robomerge

	#!ROBOMERGE-AUTHOR: jon.lietz
	compile fix

	#!rb none
	#!test compiles
	[at Daniel.Lamb]

	#!ROBOMERGE-SOURCE: CL 3559043 in //Orion/Release-42/... via CL 3559060
	#!ROBOMERGE-BOT: ORION (Main -> Dev-Balance)

Change 3559252 on 2017/07/27 by robomerge

	#!ROBOMERGE-AUTHOR: jon.lietz
	compile fix

	#!rb none
	#!test compiles
	[at Daniel.Lamb]

	#!ROBOMERGE-SOURCE: CL 3559043 in //Orion/Release-42/... via CL 3559060
	#!ROBOMERGE-BOT: ORION (Main -> Dev-REGS)

Change 3559251 on 2017/07/27 by robomerge

	#!ROBOMERGE-AUTHOR: jon.lietz
	compile fix

	#!rb none
	#!test compiles
	[at Daniel.Lamb]

	#!ROBOMERGE-SOURCE: CL 3559043 in //Orion/Release-42/... via CL 3559060
	#!ROBOMERGE-BOT: ORION (Main -> Dev-UI)

Change 3559250 on 2017/07/27 by robomerge

	#!ROBOMERGE-AUTHOR: jon.lietz
	compile fix

	#!rb none
	#!test compiles
	[at Daniel.Lamb]

	#!ROBOMERGE-SOURCE: CL 3559043 in //Orion/Release-42/... via CL 3559060
	#!ROBOMERGE-BOT: ORION (Main -> Dev-General)

Change 3559192 on 2017/07/27 by [REDACTED_PERSON]

	Removing compile on load for standalone functions.

	#!rb none
	#!tests n/a

Change 3559115 on 2017/07/27 by robomerge

	#!ROBOMERGE-AUTHOR: laurent.delayen
	Exposed GetAzimuthAndElevation to blueprints.

	#!rb none
	#!tests Pyro turrets

	#!ROBOMERGE-SOURCE: CL 3559111 in //Orion/Release-42/...
	#!ROBOMERGE-BOT: ORION (Release-42 -> Main)

Change 3559111 on 2017/07/27 by [REDACTED_PERSON]

	Exposed GetAzimuthAndElevation to blueprints.

	#!rb none
	#!tests Pyro turrets

Change 3559060 on 2017/07/27 by robomerge

	#!ROBOMERGE-AUTHOR: jon.lietz
	compile fix

	#!rb none
	#!test compiles
	[REDACTED_HANDLE]

	#!ROBOMERGE-SOURCE: CL 3559043 in //Orion/Release-42/...
	#!ROBOMERGE-BOT: ORION (Release-42 -> Main)

Change 3559043 on 2017/07/27 by [REDACTED_PERSON]

	compile fix

	#!rb none
	#!test compiles
	#!review-3559054 [REDACTED_HANDLE]

Change 3558928 on 2017/07/27 by robomerge

	#!ROBOMERGE-AUTHOR: daniel.lamb
	More temporary code to track down OR-40458
	#!test Paragon boot test ps4
	#!rb None

	#!ROBOMERGE-SOURCE: CL 3558917 in //Orion/Release-42/... via CL 3558919
	#!ROBOMERGE-BOT: ORION (Main -> Dev-Cinematics)

Change 3558927 on 2017/07/27 by robomerge

	#!ROBOMERGE-AUTHOR: daniel.lamb
	More temporary code to track down OR-40458
	#!test Paragon boot test ps4
	#!rb None

	#!ROBOMERGE-SOURCE: CL 3558917 in //Orion/Release-42/... via CL 3558919
	#!ROBOMERGE-BOT: ORION (Main -> Dev-Balance)

Change 3558926 on 2017/07/27 by robomerge

	#!ROBOMERGE-AUTHOR: daniel.lamb
	More temporary code to track down OR-40458
	#!test Paragon boot test ps4
	#!rb None

	#!ROBOMERGE-SOURCE: CL 3558917 in //Orion/Release-42/... via CL 3558919
	#!ROBOMERGE-BOT: ORION (Main -> Dev-REGS)

Change 3558923 on 2017/07/27 by robomerge

	#!ROBOMERGE-AUTHOR: daniel.lamb
	More temporary code to track down OR-40458
	#!test Paragon boot test ps4
	#!rb None

	#!ROBOMERGE-SOURCE: CL 3558917 in //Orion/Release-42/... via CL 3558919
	#!ROBOMERGE-BOT: ORION (Main -> Dev-UI)

Change 3558921 on 2017/07/27 by robomerge

	#!ROBOMERGE-AUTHOR: daniel.lamb
	More temporary code to track down OR-40458
	#!test Paragon boot test ps4
	#!rb None

	#!ROBOMERGE-SOURCE: CL 3558917 in //Orion/Release-42/... via CL 3558919
	#!ROBOMERGE-BOT: ORION (Main -> Dev-General)

Change 3558919 on 2017/07/27 by robomerge

	#!ROBOMERGE-AUTHOR: daniel.lamb
	More temporary code to track down OR-40458
	#!test Paragon boot test ps4
	#!rb None

	#!ROBOMERGE-SOURCE: CL 3558917 in //Orion/Release-42/...
	#!ROBOMERGE-BOT: ORION (Release-42 -> Main)

Change 3558917 on 2017/07/27 by [REDACTED_PERSON]

	More temporary code to track down OR-40458
	#!test Paragon boot test ps4
	#!rb None

Change 3558264 on 2017/07/27 by [REDACTED_PERSON]

	Pondering update

Change 3558206 on 2017/07/27 by Jurre.deBaare

	HLOD: Need to be able to disable auto-LOD generation on meshes in a BP
	#!fix added flag to PrimitiveComponent to disable certain BP components to be excluded from HLOD generation, and also not have a LODParent primitive set
	#!jira UE-47711
	#!rb Benn.Gallagher
	#!Tests generate HLOD clusters with enabled/disabled components and actors

Change 3558200 on 2017/07/27 by Jurre.deBaare

	Crash rebuilding HLOD cluster
	#!fix Simplygon returns an empty mesh if the input is not overlapping the culling (landscape) mesh, so added bound check for input vs landscape to prevent this situation
	#!misc Added error when Simplygon returns an invalid raw mesh after processing
	#!jira UE-47709
	#!rb Benn.Gallagher

Change 3558116 on 2017/07/27 by [REDACTED_PERSON]

	Roughed in drag, while pondering physical correctness or lack therof

Change 3557918 on 2017/07/27 by [REDACTED_PERSON]

	~2x speed up of niagara compilation.
	Set of visited nodes in numeric fix up viistor was becoming massive and spending about half the total compile time just ensuring we'd not visited a node before.
	Moved over to a slightly clunkier but faster method of using a visitor ID on the node itself.

	#!codereview Olaf.Piesche, Frank.Fella, Shaun.Kime
	#!rb none
	#!tests tested several emitters. Seems to work

Change 3557439 on 2017/07/26 by [REDACTED_PERSON]

	Replicating CL3557068

	Adding a configurable spawn rate scaling reference value; sets the zero-scale reference value (default: 2),  so additional quality levels can be added and scaling customized further.
	IMPORTANT: This sets the reference to 3 in PS4Scalability.ini; effects on PS4 are again going to have reduced spawn rates versus PC and Neo, as intended by the FX artists starting with this change.

	#!rb marcus.wassmer
	#!tests QAGame

Change 3556915 on 2017/07/26 by robomerge

	#!ROBOMERGE-AUTHOR: daniel.lamb
	Temporary change to help track down garbage UTexture refrence related to OR-40458
	#!rb Trivial
	#!test Paragon cooked

	#!ROBOMERGE-SOURCE: CL 3556903 in //Orion/Release-42/... via CL 3556910
	#!ROBOMERGE-BOT: ORION (Main -> Dev-Cinematics)

Change 3556914 on 2017/07/26 by robomerge

	#!ROBOMERGE-AUTHOR: daniel.lamb
	Temporary change to help track down garbage UTexture refrence related to OR-40458
	#!rb Trivial
	#!test Paragon cooked

	#!ROBOMERGE-SOURCE: CL 3556903 in //Orion/Release-42/... via CL 3556910
	#!ROBOMERGE-BOT: ORION (Main -> Dev-Balance)

Change 3556913 on 2017/07/26 by robomerge

	#!ROBOMERGE-AUTHOR: daniel.lamb
	Temporary change to help track down garbage UTexture refrence related to OR-40458
	#!rb Trivial
	#!test Paragon cooked

	#!ROBOMERGE-SOURCE: CL 3556903 in //Orion/Release-42/... via CL 3556910
	#!ROBOMERGE-BOT: ORION (Main -> Dev-REGS)

Change 3556912 on 2017/07/26 by robomerge

	#!ROBOMERGE-AUTHOR: daniel.lamb
	Temporary change to help track down garbage UTexture refrence related to OR-40458
	#!rb Trivial
	#!test Paragon cooked

	#!ROBOMERGE-SOURCE: CL 3556903 in //Orion/Release-42/... via CL 3556910
	#!ROBOMERGE-BOT: ORION (Main -> Dev-UI)

Change 3556911 on 2017/07/26 by robomerge

	#!ROBOMERGE-AUTHOR: daniel.lamb
	Temporary change to help track down garbage UTexture refrence related to OR-40458
	#!rb Trivial
	#!test Paragon cooked

	#!ROBOMERGE-SOURCE: CL 3556903 in //Orion/Release-42/... via CL 3556910
	#!ROBOMERGE-BOT: ORION (Main -> Dev-General)

Change 3556910 on 2017/07/26 by robomerge

	#!ROBOMERGE-AUTHOR: daniel.lamb
	Temporary change to help track down garbage UTexture refrence related to OR-40458
	#!rb Trivial
	#!test Paragon cooked

	#!ROBOMERGE-SOURCE: CL 3556903 in //Orion/Release-42/...
	#!ROBOMERGE-BOT: ORION (Release-42 -> Main)

Change 3556903 on 2017/07/26 by [REDACTED_PERSON]

	Temporary change to help track down garbage UTexture refrence related to OR-40458
	#!rb Trivial
	#!test Paragon cooked

Change 3556592 on 2017/07/26 by robomerge

	#!ROBOMERGE-AUTHOR: andrew.grant
	Allow Notify nodes to be absolute - e.g. replace any notication settings the node has from being included in other targets.

	This is to allow us to restrict emails about certain warnings or failures to a smaller subset of people

	#!rb Ben.Marsh (review)
	#!tests Debugged through a Nightly Build target of OrionBuild and verified absolute notifies are correctly set up.

	#!ROBOMERGE-SOURCE: CL 3556570 in //Orion/Release-42/... via CL 3556587
	#!ROBOMERGE-BOT: ORION (Main -> Dev-Cinematics)

Change 3556591 on 2017/07/26 by robomerge

	#!ROBOMERGE-AUTHOR: andrew.grant
	Allow Notify nodes to be absolute - e.g. replace any notication settings the node has from being included in other targets.

	This is to allow us to restrict emails about certain warnings or failures to a smaller subset of people

	#!rb Ben.Marsh (review)
	#!tests Debugged through a Nightly Build target of OrionBuild and verified absolute notifies are correctly set up.

	#!ROBOMERGE-SOURCE: CL 3556570 in //Orion/Release-42/... via CL 3556587
	#!ROBOMERGE-BOT: ORION (Main -> Dev-Balance)

Change 3556590 on 2017/07/26 by robomerge

	#!ROBOMERGE-AUTHOR: andrew.grant
	Allow Notify nodes to be absolute - e.g. replace any notication settings the node has from being included in other targets.

	This is to allow us to restrict emails about certain warnings or failures to a smaller subset of people

	#!rb Ben.Marsh (review)
	#!tests Debugged through a Nightly Build target of OrionBuild and verified absolute notifies are correctly set up.

	#!ROBOMERGE-SOURCE: CL 3556570 in //Orion/Release-42/... via CL 3556587
	#!ROBOMERGE-BOT: ORION (Main -> Dev-REGS)

Change 3556589 on 2017/07/26 by robomerge

	#!ROBOMERGE-AUTHOR: andrew.grant
	Allow Notify nodes to be absolute - e.g. replace any notication settings the node has from being included in other targets.

	This is to allow us to restrict emails about certain warnings or failures to a smaller subset of people

	#!rb Ben.Marsh (review)
	#!tests Debugged through a Nightly Build target of OrionBuild and verified absolute notifies are correctly set up.

	#!ROBOMERGE-SOURCE: CL 3556570 in //Orion/Release-42/... via CL 3556587
	#!ROBOMERGE-BOT: ORION (Main -> Dev-UI)

Change 3556588 on 2017/07/26 by robomerge

	#!ROBOMERGE-AUTHOR: andrew.grant
	Allow Notify nodes to be absolute - e.g. replace any notication settings the node has from being included in other targets.

	This is to allow us to restrict emails about certain warnings or failures to a smaller subset of people

	#!rb Ben.Marsh (review)
	#!tests Debugged through a Nightly Build target of OrionBuild and verified absolute notifies are correctly set up.

	#!ROBOMERGE-SOURCE: CL 3556570 in //Orion/Release-42/... via CL 3556587
	#!ROBOMERGE-BOT: ORION (Main -> Dev-General)

Change 3556587 on 2017/07/26 by robomerge

	#!ROBOMERGE-AUTHOR: andrew.grant
	Allow Notify nodes to be absolute - e.g. replace any notication settings the node has from being included in other targets.

	This is to allow us to restrict emails about certain warnings or failures to a smaller subset of people

	#!rb Ben.Marsh (review)
	#!tests Debugged through a Nightly Build target of OrionBuild and verified absolute notifies are correctly set up.

	#!ROBOMERGE-SOURCE: CL 3556570 in //Orion/Release-42/...
	#!ROBOMERGE-BOT: ORION (Release-42 -> Main)

Change 3556570 on 2017/07/26 by [REDACTED_PERSON]

	Allow Notify nodes to be absolute - e.g. replace any notication settings the node has from being included in other targets.

	This is to allow us to restrict emails about certain warnings or failures to a smaller subset of people

	#!rb Ben.Marsh (review)
	#!tests Debugged through a Nightly Build target of OrionBuild and verified absolute notifies are correctly set up.

Change 3556239 on 2017/07/26 by robomerge

	#!ROBOMERGE-AUTHOR: david.ratti
	Fix issue where gameplay tags were not deterministic between cooked and uncooked

	#!rb none
	#!tests pie
	[at Daniel.Lamb]

	#!ROBOMERGE-SOURCE: CL 3556226 in //Orion/Release-42/... via CL 3556229
	#!ROBOMERGE-BOT: ORION (Main -> Dev-Cinematics)

Change 3556238 on 2017/07/26 by robomerge

	#!ROBOMERGE-AUTHOR: david.ratti
	Fix issue where gameplay tags were not deterministic between cooked and uncooked

	#!rb none
	#!tests pie
	[at Daniel.Lamb]

	#!ROBOMERGE-SOURCE: CL 3556226 in //Orion/Release-42/... via CL 3556229
	#!ROBOMERGE-BOT: ORION (Main -> Dev-Balance)

Change 3556237 on 2017/07/26 by robomerge

	#!ROBOMERGE-AUTHOR: david.ratti
	Fix issue where gameplay tags were not deterministic between cooked and uncooked

	#!rb none
	#!tests pie
	[at Daniel.Lamb]

	#!ROBOMERGE-SOURCE: CL 3556226 in //Orion/Release-42/... via CL 3556229
	#!ROBOMERGE-BOT: ORION (Main -> Dev-REGS)

Change 3556236 on 2017/07/26 by robomerge

	#!ROBOMERGE-AUTHOR: david.ratti
	Fix issue where gameplay tags were not deterministic between cooked and uncooked

	#!rb none
	#!tests pie
	[at Daniel.Lamb]

	#!ROBOMERGE-SOURCE: CL 3556226 in //Orion/Release-42/... via CL 3556229
	#!ROBOMERGE-BOT: ORION (Main -> Dev-UI)

Change 3556235 on 2017/07/26 by robomerge

	#!ROBOMERGE-AUTHOR: david.ratti
	Fix issue where gameplay tags were not deterministic between cooked and uncooked

	#!rb none
	#!tests pie
	[at Daniel.Lamb]

	#!ROBOMERGE-SOURCE: CL 3556226 in //Orion/Release-42/... via CL 3556229
	#!ROBOMERGE-BOT: ORION (Main -> Dev-General)

Change 3556229 on 2017/07/26 by robomerge

	#!ROBOMERGE-AUTHOR: david.ratti
	Fix issue where gameplay tags were not deterministic between cooked and uncooked

	#!rb none
	#!tests pie
	[REDACTED_HANDLE]

	#!ROBOMERGE-SOURCE: CL 3556226 in //Orion/Release-42/...
	#!ROBOMERGE-BOT: ORION (Release-42 -> Main)

Change 3556226 on 2017/07/26 by [REDACTED_PERSON]

	Fix issue where gameplay tags were not deterministic between cooked and uncooked

	#!rb none
	#!tests pie
	#!review-3556227 [REDACTED_HANDLE]

Change 3556163 on 2017/07/26 by [REDACTED_PERSON]

	Niagara - Rework the system toolkit so that it can edit stand alone emitters and systems.  This allows the use the attribute spreasheet and system views when editing emitters and enables inspecting and editing the emitter graphs (for debug purposes) when editing systems.

	#!Tests Verified general system and emitter editing functionality.
	#!rb none

	#!codereview Olaf.Piesche,Simon.Tovey,Shaun.Kime

Change 3556104 on 2017/07/26 by [REDACTED_PERSON]

	Changed OpacityConst and OpacityMaskConst default to 1.0 to prevent HLOD meshes from disappearing

Change 3555992 on 2017/07/26 by [REDACTED_PERSON]

	Niagara - Fix a bug when deleting dynanmic inputs which would leave the graph broken.

	#!Tests Removing a dynamic input now leaves the graph in a vaild state.
	#!rb none

	#!codereview Olaf.Piesche,Simon.Tovey,Shaun.Kime

Change 3555991 on 2017/07/26 by robomerge

	#!ROBOMERGE-AUTHOR: david.ratti
	Change GameplayCueManager to deal with UClasses instead of CDOs when managing preallocation lists.

	OR-41476

	#!rb none
	#!tests pie

	#!ROBOMERGE-SOURCE: CL 3555778 in //Orion/Release-42/... via CL 3555896
	#!ROBOMERGE-BOT: ORION (Main -> Dev-Cinematics)

Change 3555988 on 2017/07/26 by robomerge

	#!ROBOMERGE-AUTHOR: david.ratti
	Change GameplayCueManager to deal with UClasses instead of CDOs when managing preallocation lists.

	OR-41476

	#!rb none
	#!tests pie

	#!ROBOMERGE-SOURCE: CL 3555778 in //Orion/Release-42/... via CL 3555896
	#!ROBOMERGE-BOT: ORION (Main -> Dev-Balance)

Change 3555984 on 2017/07/26 by robomerge

	#!ROBOMERGE-AUTHOR: david.ratti
	Change GameplayCueManager to deal with UClasses instead of CDOs when managing preallocation lists.

	OR-41476

	#!rb none
	#!tests pie

	#!ROBOMERGE-SOURCE: CL 3555778 in //Orion/Release-42/... via CL 3555896
	#!ROBOMERGE-BOT: ORION (Main -> Dev-REGS)

Change 3555983 on 2017/07/26 by robomerge

	#!ROBOMERGE-AUTHOR: david.ratti
	Change GameplayCueManager to deal with UClasses instead of CDOs when managing preallocation lists.

	OR-41476

	#!rb none
	#!tests pie

	#!ROBOMERGE-SOURCE: CL 3555778 in //Orion/Release-42/... via CL 3555896
	#!ROBOMERGE-BOT: ORION (Main -> Dev-UI)

Change 3555982 on 2017/07/26 by robomerge

	#!ROBOMERGE-AUTHOR: david.ratti
	Change GameplayCueManager to deal with UClasses instead of CDOs when managing preallocation lists.

	OR-41476

	#!rb none
	#!tests pie

	#!ROBOMERGE-SOURCE: CL 3555778 in //Orion/Release-42/... via CL 3555896
	#!ROBOMERGE-BOT: ORION (Main -> Dev-General)

Change 3555896 on 2017/07/26 by robomerge

	#!ROBOMERGE-AUTHOR: david.ratti
	Change GameplayCueManager to deal with UClasses instead of CDOs when managing preallocation lists.

	OR-41476

	#!rb none
	#!tests pie

	#!ROBOMERGE-SOURCE: CL 3555778 in //Orion/Release-42/...
	#!ROBOMERGE-BOT: ORION (Release-42 -> Main)

Change 3555778 on 2017/07/26 by [REDACTED_PERSON]

	Change GameplayCueManager to deal with UClasses instead of CDOs when managing preallocation lists.

	OR-41476

	#!rb none
	#!tests pie

Change 3555726 on 2017/07/26 by [REDACTED_PERSON]

	Niagara - Don't clear keyboard focus on commit for float and int value editors.

	#!Tests keyboard focus is no longer cleared.
	#!rb none

Change 3555668 on 2017/07/26 by [REDACTED_PERSON]

	Niagara - Fix a bug in the hlsl translator where multiple dynamic input usages were not genering unique code like modules.

	#!Tests Multiple dynamic input usages generate correct code.
	#!rb Shaun K.

Change 3555188 on 2017/07/26 by robomerge

	#!ROBOMERGE-AUTHOR: andrew.grant
	Gauntlet - Fixed error in checking availability range of devices

	#!tests debugged
	#!rb none

	#!ROBOMERGE-SOURCE: CL 3555053 in //Orion/Release-42/... via CL 3555088
	#!ROBOMERGE-BOT: ORION (Main -> Dev-Cinematics)

Change 3555187 on 2017/07/26 by robomerge

	#!ROBOMERGE-AUTHOR: andrew.grant
	Gauntlet - Fixed error in checking availability range of devices

	#!tests debugged
	#!rb none

	#!ROBOMERGE-SOURCE: CL 3555053 in //Orion/Release-42/... via CL 3555088
	#!ROBOMERGE-BOT: ORION (Main -> Dev-Balance)

Change 3555186 on 2017/07/26 by robomerge

	#!ROBOMERGE-AUTHOR: andrew.grant
	Gauntlet - Fixed error in checking availability range of devices

	#!tests debugged
	#!rb none

	#!ROBOMERGE-SOURCE: CL 3555053 in //Orion/Release-42/... via CL 3555088
	#!ROBOMERGE-BOT: ORION (Main -> Dev-REGS)

Change 3555185 on 2017/07/26 by robomerge

	#!ROBOMERGE-AUTHOR: andrew.grant
	Gauntlet - Fixed error in checking availability range of devices

	#!tests debugged
	#!rb none

	#!ROBOMERGE-SOURCE: CL 3555053 in //Orion/Release-42/... via CL 3555088
	#!ROBOMERGE-BOT: ORION (Main -> Dev-UI)

Change 3555184 on 2017/07/26 by robomerge

	#!ROBOMERGE-AUTHOR: andrew.grant
	Gauntlet - Fixed error in checking availability range of devices

	#!tests debugged
	#!rb none

	#!ROBOMERGE-SOURCE: CL 3555053 in //Orion/Release-42/... via CL 3555088
	#!ROBOMERGE-BOT: ORION (Main -> Dev-General)

Change 3555088 on 2017/07/26 by robomerge

	#!ROBOMERGE-AUTHOR: andrew.grant
	Gauntlet - Fixed error in checking availability range of devices

	#!tests debugged
	#!rb none

	#!ROBOMERGE-SOURCE: CL 3555053 in //Orion/Release-42/...
	#!ROBOMERGE-BOT: ORION (Release-42 -> Main)

Change 3555053 on 2017/07/26 by [REDACTED_PERSON]

	Gauntlet - Fixed error in checking availability range of devices

	#!tests debugged
	#!rb none


diff --git a/Engine/Config/BaseEngine.ini b/Engine/Config/BaseEngine.ini
--- a/Engine/Config/BaseEngine.ini
+++ b/Engine/Config/BaseEngine.ini
@@ -547,6 +547,8 @@ CanSkipEditorReferencedPackagesWhenCooking=False
 DetailedCallstacksInNonMonolithicBuilds=True
 UseSeperateBulkDataFiles=False
 HangDuration=0.0
+AssetLogShowsDiskPath=True
+
 
 [/Script/Engine.StreamingSettings]
 s.MinBulkDataSizeForAsyncLoading=131072
diff --git a/Engine/Plugins/Editor/GameplayTagsEditor/Source/GameplayTagsEditor/Private/GameplayTagsEditorModule.cpp b/Engine/Plugins/Editor/GameplayTagsEditor/Source/GameplayTagsEditor/Private/GameplayTagsEditorModule.cpp
--- a/Engine/Plugins/Editor/GameplayTagsEditor/Source/GameplayTagsEditor/Private/GameplayTagsEditorModule.cpp
+++ b/Engine/Plugins/Editor/GameplayTagsEditor/Source/GameplayTagsEditor/Private/GameplayTagsEditorModule.cpp
@@ -231,7 +231,7 @@ public:
 		ShowNotification(LOCTEXT("MigrationText", "Migrated Tag Settings, check DefaultEngine.ini before checking in!"), 10.0f);
 	}
 
-	void GameplayTagsUpdateSourceControl(FString RelativeConfigFilePath)
+	void GameplayTagsUpdateSourceControl(const FString& RelativeConfigFilePath)
 	{
 		FString ConfigPath = FPaths::ConvertRelativePathToFull(RelativeConfigFilePath);
 
@@ -258,7 +258,7 @@ public:
 		}
 	}
 
-	bool DeleteTagRedirector(FString TagToDelete)
+	bool DeleteTagRedirector(const FString& TagToDelete)
 	{
 		FName TagName = FName(*TagToDelete);
 
@@ -286,7 +286,7 @@ public:
 		return false;
 	}
 
-	virtual bool AddNewGameplayTagToINI(FString NewTag, FString Comment, FName TagSourceName) override
+	virtual bool AddNewGameplayTagToINI(const FString& NewTag, const FString& Comment, FName TagSourceName) override
 	{
 		UGameplayTagsManager& Manager = UGameplayTagsManager::Get();
 
@@ -367,7 +367,7 @@ public:
 		return true;
 	}
 
-	virtual bool DeleteTagFromINI(FString TagToDelete) override
+	virtual bool DeleteTagFromINI(const FString& TagToDelete) override
 	{
 		FName TagName = FName(*TagToDelete);
 
@@ -484,7 +484,7 @@ public:
 		return false;
 	}
 
-	virtual bool RenameTagInINI(FString TagToRename, FString TagToRenameTo) override
+	virtual bool RenameTagInINI(const FString& TagToRename, const FString& TagToRenameTo) override
 	{
 		FName OldTagName = FName(*TagToRename);
 		FName NewTagName = FName(*TagToRenameTo);
@@ -556,6 +556,27 @@ public:
 		return true;
 	}
 
+	virtual bool AddTransientEditorGameplayTag(const FString& NewTransientTag) override
+	{
+		UGameplayTagsManager& Manager = UGameplayTagsManager::Get();
+
+		if (NewTransientTag.IsEmpty())
+		{
+			return false;
+		}
+
+		Manager.TransientEditorTags.Add(*NewTransientTag);
+
+		{
+			FString PerfMessage = FString::Printf(TEXT("ConstructGameplayTagTree GameplayTag tables after adding new transient tag"));
+			SCOPE_LOG_TIME_IN_SECONDS(*PerfMessage, nullptr)
+
+			Manager.EditorRefreshGameplayTagTree();
+		}
+
+		return true;
+	}
+
 	FDelegateHandle AssetImportHandle;
 	FDelegateHandle SettingsChangedHandle;
 
@@ -566,4 +587,4 @@ public:
 
 IMPLEMENT_MODULE(FGameplayTagsEditorModule, GameplayTagsEditor)
 
-#undef LOCTEXT_NAMESPACE
+#undef LOCTEXT_NAMESPACE
\ No newline at end of file
diff --git a/Engine/Plugins/Editor/GameplayTagsEditor/Source/GameplayTagsEditor/Public/GameplayTagsEditorModule.h b/Engine/Plugins/Editor/GameplayTagsEditor/Source/GameplayTagsEditor/Public/GameplayTagsEditorModule.h
--- a/Engine/Plugins/Editor/GameplayTagsEditor/Source/GameplayTagsEditor/Public/GameplayTagsEditorModule.h
+++ b/Engine/Plugins/Editor/GameplayTagsEditor/Source/GameplayTagsEditor/Public/GameplayTagsEditorModule.h
@@ -38,14 +38,16 @@ public:
 	}
 
 	/** Tries to add a new gameplay tag to the ini lists */
-	GAMEPLAYTAGSEDITOR_API virtual bool AddNewGameplayTagToINI(FString NewTag, FString Comment = TEXT(""), FName TagSourceName = NAME_None) = 0;
+	GAMEPLAYTAGSEDITOR_API virtual bool AddNewGameplayTagToINI(const FString& NewTag, const FString& Comment = TEXT(""), FName TagSourceName = NAME_None) = 0;
 
 	/** Tries to delete a tag from the library. This will pop up special UI or error messages as needed. It will also delete redirectors if that is specified. */
-	GAMEPLAYTAGSEDITOR_API virtual bool DeleteTagFromINI(FString TagToDelete) = 0;
+	GAMEPLAYTAGSEDITOR_API virtual bool DeleteTagFromINI(const FString& TagToDelete) = 0;
 
 	/** Tries to rename a tag, leaving a rediretor in the ini, and adding the new tag if it does not exist yet */
-	GAMEPLAYTAGSEDITOR_API virtual bool RenameTagInINI(FString TagToRename, FString TagToRenameTo) = 0;
+	GAMEPLAYTAGSEDITOR_API virtual bool RenameTagInINI(const FString& TagToRename, const FString& TagToRenameTo) = 0;
 
+	/** Adds a transient gameplay tag (only valid for the current editor session) */
+	GAMEPLAYTAGSEDITOR_API virtual bool AddTransientEditorGameplayTag(const FString& NewTransientTag) = 0;
 };
 
 /** This is public so that child structs of FGameplayTag can use the details customization */
diff --git a/Engine/Plugins/Runtime/GameplayAbilities/Source/GameplayAbilities/Private/Abilities/GameplayAbility.cpp b/Engine/Plugins/Runtime/GameplayAbilities/Source/GameplayAbilities/Private/Abilities/GameplayAbility.cpp
--- a/Engine/Plugins/Runtime/GameplayAbilities/Source/GameplayAbilities/Private/Abilities/GameplayAbility.cpp
+++ b/Engine/Plugins/Runtime/GameplayAbilities/Source/GameplayAbilities/Private/Abilities/GameplayAbility.cpp
@@ -583,6 +583,9 @@ void UGameplayAbility::EndAbility(const FGameplayAbilitySpecHandle Handle, const
 		OnGameplayAbilityEnded.Broadcast(this);
 		OnGameplayAbilityEnded.Clear();
 
+		OnGameplayAbilityEndedWithData.Broadcast(FAbilityEndedData(this, Handle, bReplicateEndAbility, bWasCancelled));
+		OnGameplayAbilityEndedWithData.Clear();
+
 		if (GetInstancingPolicy() != EGameplayAbilityInstancingPolicy::NonInstanced)
 		{
 			bIsActive = false;
diff --git a/Engine/Plugins/Runtime/GameplayAbilities/Source/GameplayAbilities/Private/Abilities/Tasks/AbilityTask_WaitAbilityCommit.cpp b/Engine/Plugins/Runtime/GameplayAbilities/Source/GameplayAbilities/Private/Abilities/Tasks/AbilityTask_WaitAbilityCommit.cpp
--- a/Engine/Plugins/Runtime/GameplayAbilities/Source/GameplayAbilities/Private/Abilities/Tasks/AbilityTask_WaitAbilityCommit.cpp
+++ b/Engine/Plugins/Runtime/GameplayAbilities/Source/GameplayAbilities/Private/Abilities/Tasks/AbilityTask_WaitAbilityCommit.cpp
@@ -33,7 +33,7 @@ void UAbilityTask_WaitAbilityCommit::Activate()
 {
 	if (AbilitySystemComponent)	
 	{		
-		OnAbilityCommitDelegateHandle = AbilitySystemComponent->AbilityCommitedCallbacks.AddUObject(this, &UAbilityTask_WaitAbilityCommit::OnAbilityCommit);
+		OnAbilityCommitDelegateHandle = AbilitySystemComponent->AbilityCommittedCallbacks.AddUObject(this, &UAbilityTask_WaitAbilityCommit::OnAbilityCommit);
 	}
 }
 
@@ -41,7 +41,7 @@ void UAbilityTask_WaitAbilityCommit::OnDestroy(bool AbilityEnded)
 {
 	if (AbilitySystemComponent)
 	{
-		AbilitySystemComponent->AbilityCommitedCallbacks.Remove(OnAbilityCommitDelegateHandle);
+		AbilitySystemComponent->AbilityCommittedCallbacks.Remove(OnAbilityCommitDelegateHandle);
 	}
 
 	Super::OnDestroy(AbilityEnded);
diff --git a/Engine/Plugins/Runtime/GameplayAbilities/Source/GameplayAbilities/Private/Abilities/Tasks/AbilityTask_WaitGameplayEffectRemoved.cpp b/Engine/Plugins/Runtime/GameplayAbilities/Source/GameplayAbilities/Private/Abilities/Tasks/AbilityTask_WaitGameplayEffectRemoved.cpp
--- a/Engine/Plugins/Runtime/GameplayAbilities/Source/GameplayAbilities/Private/Abilities/Tasks/AbilityTask_WaitGameplayEffectRemoved.cpp
+++ b/Engine/Plugins/Runtime/GameplayAbilities/Source/GameplayAbilities/Private/Abilities/Tasks/AbilityTask_WaitGameplayEffectRemoved.cpp
@@ -21,21 +21,23 @@ UAbilityTask_WaitGameplayEffectRemoved* UAbilityTask_WaitGameplayEffectRemoved::
 
 void UAbilityTask_WaitGameplayEffectRemoved::Activate()
 {
+	FGameplayEffectRemovalInfo EmptyGameplayEffectRemovalInfo;
+
 	if (Handle.IsValid() == false)
 	{
 		if (ShouldBroadcastAbilityTaskDelegates())
 		{
-			InvalidHandle.Broadcast();
+			InvalidHandle.Broadcast(EmptyGameplayEffectRemovalInfo);
 		}
 		EndTask();
-		return;;
+		return;
 	}
 
 	UAbilitySystemComponent* EffectOwningAbilitySystemComponent = Handle.GetOwningAbilitySystemComponent();
 
 	if (EffectOwningAbilitySystemComponent)
 	{
-		FOnActiveGameplayEffectRemoved* DelPtr = EffectOwningAbilitySystemComponent->OnGameplayEffectRemovedDelegate(Handle);
+		FOnActiveGameplayEffectRemoved_Info* DelPtr = EffectOwningAbilitySystemComponent->OnGameplayEffectRemoved_InfoDelegate(Handle);
 		if (DelPtr)
 		{
 			OnGameplayEffectRemovedDelegateHandle = DelPtr->AddUObject(this, &UAbilityTask_WaitGameplayEffectRemoved::OnGameplayEffectRemoved);
@@ -46,7 +48,7 @@ void UAbilityTask_WaitGameplayEffectRemoved::Activate()
 	if (!Registered)
 	{
 		// GameplayEffect was already removed, treat this as a warning? Could be cases of immunity or chained gameplay rules that would instant remove something
-		OnGameplayEffectRemoved();
+		OnGameplayEffectRemoved(EmptyGameplayEffectRemovalInfo);
 	}
 }
 
@@ -55,7 +57,7 @@ void UAbilityTask_WaitGameplayEffectRemoved::OnDestroy(bool AbilityIsEnding)
 	UAbilitySystemComponent* EffectOwningAbilitySystemComponent = Handle.GetOwningAbilitySystemComponent();
 	if (EffectOwningAbilitySystemComponent)
 	{
-		FOnActiveGameplayEffectRemoved* DelPtr = EffectOwningAbilitySystemComponent->OnGameplayEffectRemovedDelegate(Handle);
+		FOnActiveGameplayEffectRemoved_Info* DelPtr = EffectOwningAbilitySystemComponent->OnGameplayEffectRemoved_InfoDelegate(Handle);
 		if (DelPtr)
 		{
 			DelPtr->Remove(OnGameplayEffectRemovedDelegateHandle);
@@ -65,11 +67,13 @@ void UAbilityTask_WaitGameplayEffectRemoved::OnDestroy(bool AbilityIsEnding)
 	Super::OnDestroy(AbilityIsEnding);
 }
 
-void UAbilityTask_WaitGameplayEffectRemoved::OnGameplayEffectRemoved()
+void UAbilityTask_WaitGameplayEffectRemoved::OnGameplayEffectRemoved(const FGameplayEffectRemovalInfo& InGameplayEffectRemovalInfo)
 {
 	if (ShouldBroadcastAbilityTaskDelegates())
 	{
-		OnRemoved.Broadcast();
+		OnRemoved.Broadcast(InGameplayEffectRemovalInfo);
 	}
 	EndTask();
 }
+
+
diff --git a/Engine/Plugins/Runtime/GameplayAbilities/Source/GameplayAbilities/Private/Abilities/Tasks/AbilityTask_WaitGameplayEffectStackChange.cpp b/Engine/Plugins/Runtime/GameplayAbilities/Source/GameplayAbilities/Private/Abilities/Tasks/AbilityTask_WaitGameplayEffectStackChange.cpp
--- a/Engine/Plugins/Runtime/GameplayAbilities/Source/GameplayAbilities/Private/Abilities/Tasks/AbilityTask_WaitGameplayEffectStackChange.cpp
+++ b/Engine/Plugins/Runtime/GameplayAbilities/Source/GameplayAbilities/Private/Abilities/Tasks/AbilityTask_WaitGameplayEffectStackChange.cpp
@@ -48,7 +48,7 @@ void UAbilityTask_WaitGameplayEffectStackChange::OnDestroy(bool AbilityIsEnding)
 	UAbilitySystemComponent* EffectOwningAbilitySystemComponent = Handle.GetOwningAbilitySystemComponent();
 	if (EffectOwningAbilitySystemComponent && OnGameplayEffectStackChangeDelegateHandle.IsValid())
 	{
-		FOnActiveGameplayEffectRemoved* DelPtr = EffectOwningAbilitySystemComponent->OnGameplayEffectRemovedDelegate(Handle);
+		FOnActiveGameplayEffectRemoved_Info* DelPtr = EffectOwningAbilitySystemComponent->OnGameplayEffectRemoved_InfoDelegate(Handle);
 		if (DelPtr)
 		{
 			DelPtr->Remove(OnGameplayEffectStackChangeDelegateHandle);
diff --git a/Engine/Plugins/Runtime/GameplayAbilities/Source/GameplayAbilities/Private/AbilitySystemBlueprintLibrary.cpp b/Engine/Plugins/Runtime/GameplayAbilities/Source/GameplayAbilities/Private/AbilitySystemBlueprintLibrary.cpp
--- a/Engine/Plugins/Runtime/GameplayAbilities/Source/GameplayAbilities/Private/AbilitySystemBlueprintLibrary.cpp
+++ b/Engine/Plugins/Runtime/GameplayAbilities/Source/GameplayAbilities/Private/AbilitySystemBlueprintLibrary.cpp
@@ -881,6 +881,63 @@ int32 UAbilitySystemBlueprintLibrary::GetActiveGameplayEffectStackLimitCount(FAc
 	return 0;
 }
 
+float UAbilitySystemBlueprintLibrary::GetActiveGameplayEffectStartTime(FActiveGameplayEffectHandle ActiveHandle)
+{
+	UAbilitySystemComponent* ASC = ActiveHandle.GetOwningAbilitySystemComponent();
+	if (ASC)
+	{
+		if (const FActiveGameplayEffect* ActiveGE = ASC->GetActiveGameplayEffect(ActiveHandle))
+		{
+			return ActiveGE->StartWorldTime;
+		}
+	}
+	return 0;
+}
+	
+float UAbilitySystemBlueprintLibrary::GetActiveGameplayEffectExpectedEndTime(FActiveGameplayEffectHandle ActiveHandle)
+{
+	UAbilitySystemComponent* ASC = ActiveHandle.GetOwningAbilitySystemComponent();
+	if (ASC)
+	{
+		if (const FActiveGameplayEffect* ActiveGE = ASC->GetActiveGameplayEffect(ActiveHandle))
+		{
+			return ActiveGE->GetEndTime();
+		}
+	}
+	return 0;
+}
+
+float UAbilitySystemBlueprintLibrary::GetActiveGameplayEffectTotalDuration(FActiveGameplayEffectHandle ActiveHandle)
+{
+	UAbilitySystemComponent* ASC = ActiveHandle.GetOwningAbilitySystemComponent();
+	if (ASC)
+	{
+		if (const FActiveGameplayEffect* ActiveGE = ASC->GetActiveGameplayEffect(ActiveHandle))
+		{
+			return ActiveGE->GetDuration();
+		}
+	}
+	return 0;
+}
+float UAbilitySystemBlueprintLibrary::GetActiveGameplayEffectRemainingDuration(UObject* WorldContextObject, FActiveGameplayEffectHandle ActiveHandle)
+{
+	UAbilitySystemComponent* ASC = ActiveHandle.GetOwningAbilitySystemComponent();
+	if (ASC)
+	{
+		if (const FActiveGameplayEffect* ActiveGE = ASC->GetActiveGameplayEffect(ActiveHandle))
+		{
+			if (WorldContextObject)
+			{
+				if (UWorld* World = WorldContextObject->GetWorld())
+				{
+					return ActiveGE->GetTimeRemaining(World->GetTimeSeconds());
+				}
+			}
+		}
+	}
+	return 0;
+}
+
 float UAbilitySystemBlueprintLibrary::GetModifiedAttributeMagnitude(const FGameplayEffectSpec& Spec, FGameplayAttribute Attribute)
 {
 	float Delta = 0.f;
diff --git a/Engine/Plugins/Runtime/GameplayAbilities/Source/GameplayAbilities/Private/AbilitySystemComponent.cpp b/Engine/Plugins/Runtime/GameplayAbilities/Source/GameplayAbilities/Private/AbilitySystemComponent.cpp
--- a/Engine/Plugins/Runtime/GameplayAbilities/Source/GameplayAbilities/Private/AbilitySystemComponent.cpp
+++ b/Engine/Plugins/Runtime/GameplayAbilities/Source/GameplayAbilities/Private/AbilitySystemComponent.cpp
@@ -75,14 +75,14 @@ void UAbilitySystemComponent::K2_InitStats(TSubclassOf<class UAttributeSet> Attr
 
 const UAttributeSet* UAbilitySystemComponent::GetOrCreateAttributeSubobject(TSubclassOf<UAttributeSet> AttributeClass)
 {
-	AActor *OwningActor = GetOwner();
-	const UAttributeSet *MyAttributes  = NULL;
+	AActor* OwningActor = GetOwner();
+	const UAttributeSet* MyAttributes = nullptr;
 	if (OwningActor && AttributeClass)
 	{
 		MyAttributes = GetAttributeSubobject(AttributeClass);
 		if (!MyAttributes)
 		{
-			UAttributeSet *Attributes = NewObject<UAttributeSet>(OwningActor, AttributeClass);
+			UAttributeSet* Attributes = NewObject<UAttributeSet>(OwningActor, AttributeClass);
 			SpawnedAttributes.AddUnique(Attributes);
 			MyAttributes = Attributes;
 		}
@@ -93,7 +93,7 @@ const UAttributeSet* UAbilitySystemComponent::GetOrCreateAttributeSubobject(TSub
 
 const UAttributeSet* UAbilitySystemComponent::GetAttributeSubobjectChecked(const TSubclassOf<UAttributeSet> AttributeClass) const
 {
-	const UAttributeSet *Set = GetAttributeSubobject(AttributeClass);
+	const UAttributeSet* Set = GetAttributeSubobject(AttributeClass);
 	check(Set);
 	return Set;
 }
@@ -405,6 +405,17 @@ FOnActiveGameplayEffectRemoved* UAbilitySystemComponent::OnGameplayEffectRemoved
 	return nullptr;
 }
 
+FOnActiveGameplayEffectRemoved_Info* UAbilitySystemComponent::OnGameplayEffectRemoved_InfoDelegate(FActiveGameplayEffectHandle Handle)
+{
+	FActiveGameplayEffect* ActiveEffect = ActiveGameplayEffects.GetActiveGameplayEffect(Handle);
+	if (ActiveEffect)
+	{
+		return &ActiveEffect->OnRemoved_InfoDelegate;
+	}
+
+	return nullptr;
+}
+
 FOnGivenActiveGameplayEffectRemoved& UAbilitySystemComponent::OnAnyGameplayEffectRemovedDelegate()
 {
 	return ActiveGameplayEffects.OnActiveGameplayEffectRemovedDelegate;
@@ -1420,6 +1431,17 @@ void UAbilitySystemComponent::OnPredictiveGameplayCueCatchup(FGameplayTag Tag)
 	}
 }
 
+void UAbilitySystemComponent::ReinvokeActiveGameplayCues()
+{
+	for (const FActiveGameplayEffect& Effect : &ActiveGameplayEffects)
+	{
+		if (Effect.bIsInhibited == false)
+		{
+			InvokeGameplayCueEvent(Effect.Spec, EGameplayCueEvent::WhileActive);
+		}
+	}
+}
+
 // ---------------------------------------------------------------------------------------
 
 void UAbilitySystemComponent::PrintAllGameplayEffects() const
@@ -1970,7 +1992,7 @@ void UAbilitySystemComponent::Debug_Internal(FAbilitySystemComponentDebugInfo& I
 				FAggregator& Aggregator = *AggregatorRef.Get();
 
 				TMap<EGameplayModEvaluationChannel, const TArray<FAggregatorMod>*> ModMap;
-				Aggregator.DebugGetAllAggregatorMods(ModMap);
+				Aggregator.GetAllAggregatorMods(ModMap);
 
 				if (ModMap.Num() == 0)
 				{
diff --git a/Engine/Plugins/Runtime/GameplayAbilities/Source/GameplayAbilities/Private/AbilitySystemComponent_Abilities.cpp b/Engine/Plugins/Runtime/GameplayAbilities/Source/GameplayAbilities/Private/AbilitySystemComponent_Abilities.cpp
--- a/Engine/Plugins/Runtime/GameplayAbilities/Source/GameplayAbilities/Private/AbilitySystemComponent_Abilities.cpp
+++ b/Engine/Plugins/Runtime/GameplayAbilities/Source/GameplayAbilities/Private/AbilitySystemComponent_Abilities.cpp
@@ -505,6 +505,7 @@ void UAbilitySystemComponent::CheckForClearedAbilities()
 		{
 			if (AbilitySpec.AssignedHandle.IsValid() && FindAbilitySpecFromHandle(AbilitySpec.AssignedHandle) == nullptr)
 			{
+				ABILITY_LOG(Display, TEXT("::CheckForClearedAbilities is clearing AssignedHandle %s from GE %s / %s"), *AbilitySpec.AssignedHandle.ToString(), *ActiveGE.GetDebugString(), *ActiveGE.Handle.ToString() );
 				AbilitySpec.AssignedHandle = FGameplayAbilitySpecHandle();
 			}
 		}
@@ -707,6 +708,8 @@ void UAbilitySystemComponent::NotifyAbilityEnded(FGameplayAbilitySpecHandle Hand
 			MarkAbilitySpecDirty(*Spec);
 		}
 	}
+
+	OnAbilityEnded.Broadcast(FAbilityEndedData(Ability, Handle, false, bWasCancelled));
 }
 
 void UAbilitySystemComponent::CancelAbility(UGameplayAbility* Ability)
@@ -1848,7 +1851,7 @@ void UAbilitySystemComponent::SetUserAbilityActivationInhibited(bool NewInhibit)
 
 void UAbilitySystemComponent::NotifyAbilityCommit(UGameplayAbility* Ability)
 {
-	AbilityCommitedCallbacks.Broadcast(Ability);
+	AbilityCommittedCallbacks.Broadcast(Ability);
 }
 
 void UAbilitySystemComponent::NotifyAbilityActivated(FGameplayAbilitySpecHandle Handle, UGameplayAbility* Ability)
diff --git a/Engine/Plugins/Runtime/GameplayAbilities/Source/GameplayAbilities/Private/AttributeSet.cpp b/Engine/Plugins/Runtime/GameplayAbilities/Source/GameplayAbilities/Private/AttributeSet.cpp
--- a/Engine/Plugins/Runtime/GameplayAbilities/Source/GameplayAbilities/Private/AttributeSet.cpp
+++ b/Engine/Plugins/Runtime/GameplayAbilities/Source/GameplayAbilities/Private/AttributeSet.cpp
@@ -14,11 +14,13 @@
 #include "UObject/UObjectIterator.h"
 #include "AbilitySystemGlobals.h"
 #include "AbilitySystemComponent.h"
-
+#include "AbilitySystemTestAttributeSet.h"
 
 #if WITH_EDITOR
 #include "EditorReimportHandler.h"
 #endif
+#include "UObjectThreadContext.h"
+
 
 #if ENABLE_VISUAL_LOG
 namespace
@@ -219,14 +221,17 @@ void FGameplayAttribute::PostSerialize(const FArchive& Ar)
 
 			if (!Attribute)
 			{
+				FUObjectThreadContext& ThreadContext = FUObjectThreadContext::Get();
+				FString AssetName = ThreadContext.SerializedObject ? ThreadContext.SerializedObject->GetPathName() : TEXT("Unknown Object");
+
 				FString OwnerName = AttributeOwner ? AttributeOwner->GetName() : TEXT("NONE");
-				ABILITY_LOG(Warning, TEXT("FGameplayAttribute::PostSerialize called on an invalid attribute with owner %s and name %s."), *OwnerName, *AttributeName);
+				ABILITY_LOG(Warning, TEXT("FGameplayAttribute::PostSerialize called on an invalid attribute with owner %s and name %s. (Asset: %s)"), *OwnerName, *AttributeName, *AssetName);
 			}
 		}
 	}
 }
 
-void FGameplayAttribute::GetAllAttributeProperties(TArray<UProperty*>& OutProperties, FString FilterMetaStr)
+void FGameplayAttribute::GetAllAttributeProperties(TArray<UProperty*>& OutProperties, FString FilterMetaStr, bool UseEditorOnlyData)
 {
 	// Gather all UAttribute classes
 	for (TObjectIterator<UClass> ClassIt; ClassIt; ++ClassIt)
@@ -234,53 +239,68 @@ void FGameplayAttribute::GetAllAttributeProperties(TArray<UProperty*>& OutProper
 		UClass *Class = *ClassIt;
 		if (Class->IsChildOf(UAttributeSet::StaticClass()) && !Class->ClassGeneratedBy)
 		{
-#if WITH_EDITOR
-			// Allow entire classes to be filtered globally
-			if (Class->HasMetaData(TEXT("HideInDetailsView")))
+			if (UseEditorOnlyData)
+			{
+				#if WITH_EDITOR
+				// Allow entire classes to be filtered globally
+				if (Class->HasMetaData(TEXT("HideInDetailsView")))
+				{
+					continue;
+				}
+				#endif
+			}
+
+			if (Class == UAbilitySystemTestAttributeSet::StaticClass())
 			{
 				continue;
 			}
-#endif
+
 
 			for (TFieldIterator<UProperty> PropertyIt(Class, EFieldIteratorFlags::ExcludeSuper); PropertyIt; ++PropertyIt)
 			{
 				UProperty* Property = *PropertyIt;
 
-#if WITH_EDITOR
-				if (!FilterMetaStr.IsEmpty() && Property->HasMetaData(*FilterMetaStr))
+				if (UseEditorOnlyData)
 				{
-					continue;
-				}
+					#if WITH_EDITOR
+					if (!FilterMetaStr.IsEmpty() && Property->HasMetaData(*FilterMetaStr))
+					{
+						continue;
+					}
 
-				// Allow properties to be filtered globally (never show up)
-				if (Property->HasMetaData(TEXT("HideInDetailsView")))
-				{
-					continue;
+					// Allow properties to be filtered globally (never show up)
+					if (Property->HasMetaData(TEXT("HideInDetailsView")))
+					{
+						continue;
+					}
+					#endif
 				}
-#endif
 				
 				OutProperties.Add(Property);
 			}
 		}
 
-#if WITH_EDITOR
-		// UAbilitySystemComponent can add 'system' attributes
-		if (Class->IsChildOf(UAbilitySystemComponent::StaticClass()) && !Class->ClassGeneratedBy)
+		if (UseEditorOnlyData)
 		{
-			for (TFieldIterator<UProperty> PropertyIt(Class, EFieldIteratorFlags::ExcludeSuper); PropertyIt; ++PropertyIt)
+			#if WITH_EDITOR
+			// UAbilitySystemComponent can add 'system' attributes
+			if (Class->IsChildOf(UAbilitySystemComponent::StaticClass()) && !Class->ClassGeneratedBy)
 			{
-				UProperty* Property = *PropertyIt;
+				for (TFieldIterator<UProperty> PropertyIt(Class, EFieldIteratorFlags::ExcludeSuper); PropertyIt; ++PropertyIt)
+				{
+					UProperty* Property = *PropertyIt;
 
 
-				// SystemAttributes have to be explicitly tagged
-				if (Property->HasMetaData(TEXT("SystemGameplayAttribute")) == false)
-				{
-					continue;
+					// SystemAttributes have to be explicitly tagged
+					if (Property->HasMetaData(TEXT("SystemGameplayAttribute")) == false)
+					{
+						continue;
+					}
+					OutProperties.Add(Property);
 				}
-				OutProperties.Add(Property);
 			}
+			#endif
 		}
-#endif
 	}
 }
 
diff --git a/Engine/Plugins/Runtime/GameplayAbilities/Source/GameplayAbilities/Private/GameplayCueManager.cpp b/Engine/Plugins/Runtime/GameplayAbilities/Source/GameplayAbilities/Private/GameplayCueManager.cpp
--- a/Engine/Plugins/Runtime/GameplayAbilities/Source/GameplayAbilities/Private/GameplayCueManager.cpp
+++ b/Engine/Plugins/Runtime/GameplayAbilities/Source/GameplayAbilities/Private/GameplayCueManager.cpp
@@ -1472,16 +1472,16 @@ void UGameplayCueManager::CheckForPreallocation(UClass* GCClass)
 {
 	if (AGameplayCueNotify_Actor* InstancedCue = Cast<AGameplayCueNotify_Actor>(GCClass->ClassDefaultObject))
 	{
-		if (InstancedCue->NumPreallocatedInstances > 0 && GameplayCueClassesForPreallocation.Contains(InstancedCue) == false)
+		if (InstancedCue->NumPreallocatedInstances > 0 && GameplayCueClassesForPreallocation.Contains(GCClass) == false)
 		{
 			// Add this to the global list
-			GameplayCueClassesForPreallocation.Add(InstancedCue);
+			GameplayCueClassesForPreallocation.Add(GCClass);
 
 			// Add it to any world specific lists
 			for (FPreallocationInfo& Info : PreallocationInfoList_Internal)
 			{
-				ensure(Info.ClassesNeedingPreallocation.Contains(InstancedCue)==false);
-				Info.ClassesNeedingPreallocation.Push(InstancedCue);
+				ensure(Info.ClassesNeedingPreallocation.Contains(GCClass)==false);
+				Info.ClassesNeedingPreallocation.Push(GCClass);
 			}
 		}
 	}
@@ -1511,7 +1511,8 @@ void UGameplayCueManager::UpdatePreallocation(UWorld* World)
 
 	if (Info.ClassesNeedingPreallocation.Num() > 0)
 	{
-		AGameplayCueNotify_Actor* CDO = Info.ClassesNeedingPreallocation.Last();
+		TSubclassOf<AGameplayCueNotify_Actor> GCClass = Info.ClassesNeedingPreallocation.Last();
+		AGameplayCueNotify_Actor* CDO = GCClass->GetDefaultObject<AGameplayCueNotify_Actor>();
 		TArray<AGameplayCueNotify_Actor*>& PreallocatedList = Info.PreallocatedInstances.FindOrAdd(CDO->GetClass());
 
 		AGameplayCueNotify_Actor* PrespawnedInstance = Cast<AGameplayCueNotify_Actor>(World->SpawnActor(CDO->GetClass()));
diff --git a/Engine/Plugins/Runtime/GameplayAbilities/Source/GameplayAbilities/Private/GameplayCueNotify_Actor.cpp b/Engine/Plugins/Runtime/GameplayAbilities/Source/GameplayAbilities/Private/GameplayCueNotify_Actor.cpp
--- a/Engine/Plugins/Runtime/GameplayAbilities/Source/GameplayAbilities/Private/GameplayCueNotify_Actor.cpp
+++ b/Engine/Plugins/Runtime/GameplayAbilities/Source/GameplayAbilities/Private/GameplayCueNotify_Actor.cpp
@@ -160,7 +160,7 @@ void AGameplayCueNotify_Actor::HandleGameplayCue(AActor* MyTarget, EGameplayCueE
 
 		if (EventType == EGameplayCueEvent::WhileActive && !bAllowMultipleWhileActiveEvents && bHasHandledWhileActiveEvent)
 		{
-			ABILITY_LOG(Display, TEXT("GameplayCue Notify %s WhileActive already handled, skipping this one."), *GetName());
+			ABILITY_LOG(Log, TEXT("GameplayCue Notify %s WhileActive already handled, skipping this one."), *GetName());
 			return;
 		}
 
diff --git a/Engine/Plugins/Runtime/GameplayAbilities/Source/GameplayAbilities/Private/GameplayEffect.cpp b/Engine/Plugins/Runtime/GameplayAbilities/Source/GameplayAbilities/Private/GameplayEffect.cpp
--- a/Engine/Plugins/Runtime/GameplayAbilities/Source/GameplayAbilities/Private/GameplayEffect.cpp
+++ b/Engine/Plugins/Runtime/GameplayAbilities/Source/GameplayAbilities/Private/GameplayEffect.cpp
@@ -288,7 +288,14 @@ float FCustomCalculationBasedFloat::CalculateMagnitude(const FGameplayEffectSpec
 
 	const float SpecLvl = InRelevantSpec.GetLevel();
 	FString ContextString = FString::Printf(TEXT("FCustomCalculationBasedFloat::CalculateMagnitude from effect %s"), *CalcCDO->GetName());
-	return ((Coefficient.GetValueAtLevel(SpecLvl, &ContextString) * (CustomBaseValue + PreMultiplyAdditiveValue.GetValueAtLevel(SpecLvl, &ContextString))) + PostMultiplyAdditiveValue.GetValueAtLevel(SpecLvl, &ContextString));
+
+	float FinalValue = ((Coefficient.GetValueAtLevel(SpecLvl, &ContextString) * (CustomBaseValue + PreMultiplyAdditiveValue.GetValueAtLevel(SpecLvl, &ContextString))) + PostMultiplyAdditiveValue.GetValueAtLevel(SpecLvl, &ContextString));
+	if (FinalLookupCurve.IsValid(ContextString))
+	{
+		FinalValue = FinalLookupCurve.Eval(FinalValue, ContextString);
+	}
+
+	return FinalValue;
 }
 
 /** Equality/Inequality operators */
@@ -537,15 +544,22 @@ FText FGameplayEffectModifierMagnitude::GetValueForEditorDisplay() const
 
 void FGameplayEffectModifierMagnitude::ReportErrors(const FString& PathName) const
 {
-	GETCURVE_REPORTERROR_WITHPATHNAME(ScalableFloatMagnitude.Curve, PathName);
-
-	GETCURVE_REPORTERROR_WITHPATHNAME(AttributeBasedMagnitude.Coefficient.Curve, PathName);
-	GETCURVE_REPORTERROR_WITHPATHNAME(AttributeBasedMagnitude.PreMultiplyAdditiveValue.Curve, PathName);
-	GETCURVE_REPORTERROR_WITHPATHNAME(AttributeBasedMagnitude.PostMultiplyAdditiveValue.Curve, PathName);
-
-	GETCURVE_REPORTERROR_WITHPATHNAME(CustomMagnitude.Coefficient.Curve, PathName);
-	GETCURVE_REPORTERROR_WITHPATHNAME(CustomMagnitude.PreMultiplyAdditiveValue.Curve, PathName);
-	GETCURVE_REPORTERROR_WITHPATHNAME(CustomMagnitude.PostMultiplyAdditiveValue.Curve, PathName);
+	if (MagnitudeCalculationType == EGameplayEffectMagnitudeCalculation::ScalableFloat)
+	{
+		GETCURVE_REPORTERROR_WITHPATHNAME(ScalableFloatMagnitude.Curve, PathName);
+	}
+	else if (MagnitudeCalculationType == EGameplayEffectMagnitudeCalculation::AttributeBased)
+	{
+		GETCURVE_REPORTERROR_WITHPATHNAME(AttributeBasedMagnitude.Coefficient.Curve, PathName);
+		GETCURVE_REPORTERROR_WITHPATHNAME(AttributeBasedMagnitude.PreMultiplyAdditiveValue.Curve, PathName);
+		GETCURVE_REPORTERROR_WITHPATHNAME(AttributeBasedMagnitude.PostMultiplyAdditiveValue.Curve, PathName);
+	}
+	else if (MagnitudeCalculationType == EGameplayEffectMagnitudeCalculation::CustomCalculationClass)
+	{
+		GETCURVE_REPORTERROR_WITHPATHNAME(CustomMagnitude.Coefficient.Curve, PathName);
+		GETCURVE_REPORTERROR_WITHPATHNAME(CustomMagnitude.PreMultiplyAdditiveValue.Curve, PathName);
+		GETCURVE_REPORTERROR_WITHPATHNAME(CustomMagnitude.PostMultiplyAdditiveValue.Curve, PathName);
+	}
 }
 #endif // WITH_EDITOR
 
@@ -737,11 +751,14 @@ FGameplayEffectSpecForRPC::FGameplayEffectSpecForRPC(const FGameplayEffectSpec&
 	// Only copy attributes that are in the gameplay cue info
 	for (int32 i = InSpec.ModifiedAttributes.Num() - 1; i >= 0; i--)
 	{
-		for (const FGameplayEffectCue& CueInfo : Def->GameplayCues)
+		if (Def)
 		{
-			if (CueInfo.MagnitudeAttribute == InSpec.ModifiedAttributes[i].Attribute)
+			for (const FGameplayEffectCue& CueInfo : Def->GameplayCues)
 			{
-				ModifiedAttributes.Add(InSpec.ModifiedAttributes[i]);
+				if (CueInfo.MagnitudeAttribute == InSpec.ModifiedAttributes[i].Attribute)
+				{
+					ModifiedAttributes.Add(InSpec.ModifiedAttributes[i]);
+				}
 			}
 		}
 	}
@@ -1306,6 +1323,18 @@ bool FGameplayEffectAttributeCaptureSpec::AttemptAddAggregatorModsToAggregator(O
 	return false;
 }
 
+bool FGameplayEffectAttributeCaptureSpec::AttemptGatherAttributeMods(OUT TMap<EGameplayModEvaluationChannel, const TArray<FAggregatorMod>*>& OutModMap) const
+{
+	FAggregator* Agg = AttributeAggregator.Get();
+	if (Agg)
+	{
+		Agg->GetAllAggregatorMods(OutModMap);
+		return true;
+	}
+
+	return false;
+}
+
 void FGameplayEffectAttributeCaptureSpec::RegisterLinkedAggregatorCallback(FActiveGameplayEffectHandle Handle) const
 {
 	if (BackingDefinition.bSnapshot == false)
@@ -1551,6 +1580,7 @@ FActiveGameplayEffect::FActiveGameplayEffect(FActiveGameplayEffect&& Other)
 	,IsPendingRemove(Other.IsPendingRemove)
 	,ClientCachedStackCount(0)
 	,OnRemovedDelegate(Other.OnRemovedDelegate)
+	,OnRemoved_InfoDelegate(Other.OnRemoved_InfoDelegate)
 	,PeriodHandle(Other.PeriodHandle)
 	,DurationHandle(Other.DurationHandle)
 {
@@ -1575,6 +1605,7 @@ FActiveGameplayEffect& FActiveGameplayEffect::operator=(FActiveGameplayEffect&&
 	IsPendingRemove = Other.IsPendingRemove;
 	ClientCachedStackCount = Other.ClientCachedStackCount;
 	OnRemovedDelegate = Other.OnRemovedDelegate;
+	OnRemoved_InfoDelegate = Other.OnRemoved_InfoDelegate;
 	PeriodHandle = Other.PeriodHandle;
 	DurationHandle = Other.DurationHandle;
 	// Note: purposefully not copying PendingNext pointer.
@@ -1598,6 +1629,7 @@ FActiveGameplayEffect& FActiveGameplayEffect::operator=(const FActiveGameplayEff
 	IsPendingRemove = Other.IsPendingRemove;
 	ClientCachedStackCount = Other.ClientCachedStackCount;
 	OnRemovedDelegate = Other.OnRemovedDelegate;
+	OnRemoved_InfoDelegate = Other.OnRemoved_InfoDelegate;
 	PeriodHandle = Other.PeriodHandle;
 	DurationHandle = Other.DurationHandle;
 	PendingNext = Other.PendingNext;
@@ -1642,7 +1674,21 @@ void FActiveGameplayEffect::PreReplicatedRemove(const struct FActiveGameplayEffe
 
 	ABILITY_LOG(Verbose, TEXT("PreReplicatedRemove: %s %s Marked as Pending Remove: %s"), *Handle.ToString(), *Spec.Def->GetName(), IsPendingRemove ? TEXT("TRUE") : TEXT("FALSE"));
 
-	const_cast<FActiveGameplayEffectsContainer&>(InArray).InternalOnActiveGameplayEffectRemoved(*this, !bIsInhibited);	// Const cast is ok. It is there to prevent mutation of the GameplayEffects array, which this wont do.
+	FGameplayEffectRemovalInfo GameplayEffectRemovalInfo;
+	GameplayEffectRemovalInfo.StackCount = ClientCachedStackCount;
+	//Check duration to set bPrematureRemoval as req.
+	if (DurationHandle.IsValid())
+	{
+		float SecondsRemaining = GetTimeRemaining(const_cast<FActiveGameplayEffectsContainer&>(InArray).GetWorldTime());
+
+		if (SecondsRemaining > 0.f)
+		{
+			GameplayEffectRemovalInfo.bPrematureRemoval = true;
+		}
+	}
+	GameplayEffectRemovalInfo.EffectContext = Spec.GetEffectContext();
+
+	const_cast<FActiveGameplayEffectsContainer&>(InArray).InternalOnActiveGameplayEffectRemoved(*this, !bIsInhibited, GameplayEffectRemovalInfo);	// Const cast is ok. It is there to prevent mutation of the GameplayEffects array, which this wont do.
 }
 
 void FActiveGameplayEffect::PostReplicatedAdd(const struct FActiveGameplayEffectsContainer &InArray)
@@ -1933,6 +1979,8 @@ void FActiveGameplayEffectsContainer::ExecutePeriodicGameplayEffect(FActiveGamep
 	FActiveGameplayEffect* ActiveEffect = GetActiveGameplayEffect(Handle);
 	if (ActiveEffect && !ActiveEffect->bIsInhibited)
 	{
+		FScopeCurrentGameplayEffectBeingApplied ScopedGEApplication(&ActiveEffect->Spec, Owner);
+
 		if (UE_LOG_ACTIVE(VLogAbilitySystem, Log))
 		{
 			ABILITY_VLOG(Owner->OwnerActor, Log, TEXT("Executed Periodic Effect %s"), *ActiveEffect->Spec.Def->GetFName().ToString());
@@ -2992,6 +3040,8 @@ void FActiveGameplayEffectsContainer::AddActiveGameplayEffectGrantedTagsAndModif
 			if (AbilitySpecDef.AssignedHandle.IsValid() == false)
 			{
 				Owner->GiveAbility( FGameplayAbilitySpec(AbilitySpecDef, Effect.Spec.GetLevel(), Effect.Handle) );
+
+				ABILITY_LOG(Display, TEXT("::AddActiveGameplayEffectGrantedTagsAndModifiers granted ability %s (Handle %s) from GE %s (Handle: %s)"), *GetNameSafe(AbilitySpecDef.Ability), *AbilitySpecDef.AssignedHandle.ToString(), *Effect.GetDebugString(), *Effect.Handle.ToString());
 			}
 		}	
 	}
@@ -3074,6 +3124,11 @@ bool FActiveGameplayEffectsContainer::InternalRemoveActiveGameplayEffect(int32 I
 
 		ABILITY_LOG(Verbose, TEXT("InternalRemoveActiveGameplayEffect: Auth: %s Handle: %s Def: %s"), IsNetAuthority() ? TEXT("TRUE") : TEXT("FALSE"), *Effect.Handle.ToString(), Effect.Spec.Def ? *Effect.Spec.Def->GetName() : TEXT("NONE"));
 
+		FGameplayEffectRemovalInfo GameplayEffectRemovalInfo;
+		GameplayEffectRemovalInfo.StackCount = Effect.Spec.StackCount;
+		GameplayEffectRemovalInfo.bPrematureRemoval = bPrematureRemoval;
+		GameplayEffectRemovalInfo.EffectContext = Effect.Spec.GetEffectContext();
+
 		if (StacksToRemove > 0 && Effect.Spec.StackCount > StacksToRemove)
 		{
 			// This won't be a full remove, only a change in StackCount.
@@ -3099,7 +3154,7 @@ bool FActiveGameplayEffectsContainer::InternalRemoveActiveGameplayEffect(int32 I
 		ShouldInvokeGameplayCueEvent &= !Effect.bIsInhibited;
 
 		// Mark the effect pending remove, and remove all side effects from the effect
-		InternalOnActiveGameplayEffectRemoved(Effect, ShouldInvokeGameplayCueEvent);
+		InternalOnActiveGameplayEffectRemoved(Effect, ShouldInvokeGameplayCueEvent, GameplayEffectRemovalInfo);
 
 		if (Effect.DurationHandle.IsValid())
 		{
@@ -3159,7 +3214,7 @@ bool FActiveGameplayEffectsContainer::InternalRemoveActiveGameplayEffect(int32 I
 }
 
 /** Called by client and server: This does cleanup that has to happen whether the effect is being removed locally or due to replication */
-void FActiveGameplayEffectsContainer::InternalOnActiveGameplayEffectRemoved(FActiveGameplayEffect& Effect, bool bInvokeGameplayCueEvents)
+void FActiveGameplayEffectsContainer::InternalOnActiveGameplayEffectRemoved(FActiveGameplayEffect& Effect, bool bInvokeGameplayCueEvents, const FGameplayEffectRemovalInfo& GameplayEffectRemovalInfo)
 {
 	SCOPE_CYCLE_COUNTER(STAT_OnActiveGameplayEffectRemoved);
 
@@ -3186,6 +3241,7 @@ void FActiveGameplayEffectsContainer::InternalOnActiveGameplayEffectRemoved(FAct
 	}
 
 	Effect.OnRemovedDelegate.Broadcast();
+	Effect.OnRemoved_InfoDelegate.Broadcast(GameplayEffectRemovalInfo);
 	OnActiveGameplayEffectRemovedDelegate.Broadcast(Effect);
 }
 
@@ -3671,6 +3727,8 @@ void FActiveGameplayEffectsContainer::CheckDuration(FActiveGameplayEffectHandle
 					float PeriodTimeRemaining = TimerManager.GetTimerRemaining(Effect.PeriodHandle);
 					if (PeriodTimeRemaining <= KINDA_SMALL_NUMBER && !Effect.bIsInhibited)
 					{
+						FScopeCurrentGameplayEffectBeingApplied ScopedGEApplication(&Effect.Spec, Owner);
+
 						ExecuteActiveEffectsFrom(Effect.Spec);
 
 						// The above call to ExecuteActiveEffectsFrom could cause this effect to be explicitly removed
diff --git a/Engine/Plugins/Runtime/GameplayAbilities/Source/GameplayAbilities/Private/GameplayEffectAggregator.cpp b/Engine/Plugins/Runtime/GameplayAbilities/Source/GameplayAbilities/Private/GameplayEffectAggregator.cpp
--- a/Engine/Plugins/Runtime/GameplayAbilities/Source/GameplayAbilities/Private/GameplayEffectAggregator.cpp
+++ b/Engine/Plugins/Runtime/GameplayAbilities/Source/GameplayAbilities/Private/GameplayEffectAggregator.cpp
@@ -143,7 +143,7 @@ void FAggregatorModChannel::AddModsFrom(const FAggregatorModChannel& Other)
 	}
 }
 
-void FAggregatorModChannel::DebugGetAllAggregatorMods(EGameplayModEvaluationChannel Channel, OUT TMap<EGameplayModEvaluationChannel, const TArray<FAggregatorMod>*>& OutMods) const
+void FAggregatorModChannel::GetAllAggregatorMods(EGameplayModEvaluationChannel Channel, OUT TMap<EGameplayModEvaluationChannel, const TArray<FAggregatorMod>* >& OutMods) const
 {
 	OutMods.Add(Channel, Mods);
 }
@@ -278,14 +278,14 @@ void FAggregatorModChannelContainer::AddModsFrom(const FAggregatorModChannelCont
 	}
 }
 
-void FAggregatorModChannelContainer::DebugGetAllAggregatorMods(OUT TMap<EGameplayModEvaluationChannel, const TArray<FAggregatorMod>*>& OutMods) const
+void FAggregatorModChannelContainer::GetAllAggregatorMods(OUT TMap<EGameplayModEvaluationChannel, const TArray<FAggregatorMod>*>& OutMods) const
 {
 	for (const auto& ChannelEntry : ModChannelsMap)
 	{
 		EGameplayModEvaluationChannel CurChannelEnum = ChannelEntry.Key;
 		const FAggregatorModChannel& CurChannel = ChannelEntry.Value;
 		
-		CurChannel.DebugGetAllAggregatorMods(CurChannelEnum, OutMods);
+		CurChannel.GetAllAggregatorMods(CurChannelEnum, OutMods);
 	}
 }
 
@@ -446,9 +446,9 @@ void FAggregator::RemoveDependent(FActiveGameplayEffectHandle Handle)
 	Dependents.Remove(Handle);
 }
 
-void FAggregator::DebugGetAllAggregatorMods(OUT TMap<EGameplayModEvaluationChannel, const TArray<FAggregatorMod>*>& OutMods) const
+void FAggregator::GetAllAggregatorMods(OUT TMap<EGameplayModEvaluationChannel, const TArray<FAggregatorMod>*>& OutMods) const
 {
-	ModChannels.DebugGetAllAggregatorMods(OutMods);
+	ModChannels.GetAllAggregatorMods(OutMods);
 }
 
 void FAggregator::OnActiveEffectDependenciesSwapped(const TMap<FActiveGameplayEffectHandle, FActiveGameplayEffectHandle>& SwappedDependencies)
diff --git a/Engine/Plugins/Runtime/GameplayAbilities/Source/GameplayAbilities/Private/GameplayEffectExecutionCalculation.cpp b/Engine/Plugins/Runtime/GameplayAbilities/Source/GameplayAbilities/Private/GameplayEffectExecutionCalculation.cpp
--- a/Engine/Plugins/Runtime/GameplayAbilities/Source/GameplayAbilities/Private/GameplayEffectExecutionCalculation.cpp
+++ b/Engine/Plugins/Runtime/GameplayAbilities/Source/GameplayAbilities/Private/GameplayEffectExecutionCalculation.cpp
@@ -200,6 +200,53 @@ bool FGameplayEffectCustomExecutionParameters::AttemptGetCapturedAttributeAggreg
 	return false;
 }
 
+bool FGameplayEffectCustomExecutionParameters::AttemptGatherAttributeMods(const FGameplayEffectAttributeCaptureDefinition& InCaptureDef, OUT TMap<EGameplayModEvaluationChannel, const TArray<FAggregatorMod>*>& OutModMap) const
+{
+	check(OwningSpec);
+
+	const FAggregator* CalcAgg = ScopedModifierAggregators.Find(InCaptureDef);
+	if (CalcAgg)
+	{
+		CalcAgg->GetAllAggregatorMods(OutModMap);
+	}
+	else
+	{
+		const FGameplayEffectAttributeCaptureSpec* CaptureSpec = OwningSpec->CapturedRelevantAttributes.FindCaptureSpecByDefinition(InCaptureDef, true);
+		if (CaptureSpec)
+		{
+			return CaptureSpec->AttemptGatherAttributeMods(OutModMap);
+		}
+	}
+
+	return false;
+}
+
+bool FGameplayEffectCustomExecutionParameters::ForEachQualifiedAttributeMod(const FGameplayEffectAttributeCaptureDefinition& InCaptureDef, const FAggregatorEvaluateParameters& InEvalParams, TFunction< void(EGameplayModEvaluationChannel, EGameplayModOp::Type, const FAggregatorMod&) > Func) const
+{
+	TMap<EGameplayModEvaluationChannel, const TArray<FAggregatorMod>*> ModMap;
+	if (AttemptGatherAttributeMods(InCaptureDef, ModMap))
+	{
+		for (auto It : ModMap)
+		{
+			const TArray<FAggregatorMod>* ModList = It.Value;
+			for (int32 ModOpIdx = 0; ModOpIdx < EGameplayModOp::Max; ++ModOpIdx)
+			{
+				const TArray<FAggregatorMod>& CurModArray = ModList[ModOpIdx];
+				for (const FAggregatorMod& AggMod : CurModArray)
+				{
+					if (AggMod.Qualifies(InEvalParams))
+					{
+						Func(It.Key, (EGameplayModOp::Type)ModOpIdx, AggMod);
+					}
+				}
+			}
+		}
+		return true;
+	}
+
+	return false;
+}
+
 FGameplayEffectCustomExecutionOutput::FGameplayEffectCustomExecutionOutput()
 	: bTriggerConditionalGameplayEffects(false)
 	, bHandledStackCountManually(false)
diff --git a/Engine/Plugins/Runtime/GameplayAbilities/Source/GameplayAbilities/Private/GameplayEffectTypes.cpp b/Engine/Plugins/Runtime/GameplayAbilities/Source/GameplayAbilities/Private/GameplayEffectTypes.cpp
--- a/Engine/Plugins/Runtime/GameplayAbilities/Source/GameplayAbilities/Private/GameplayEffectTypes.cpp
+++ b/Engine/Plugins/Runtime/GameplayAbilities/Source/GameplayAbilities/Private/GameplayEffectTypes.cpp
@@ -616,7 +616,7 @@ FGameplayCueParameters::FGameplayCueParameters(const struct FGameplayEffectConte
 
 bool FGameplayCueParameters::NetSerialize(FArchive& Ar, class UPackageMap* Map, bool& bOutSuccess)
 {
-	static const uint8 NUM_LEVEL_BITS = 4;
+	static const uint8 NUM_LEVEL_BITS = 5; // need to bump this up to support 20 levels for AbilityLevel
 	static const uint8 MAX_LEVEL = (1 << NUM_LEVEL_BITS) - 1;
 
 	enum RepFlag
diff --git a/Engine/Plugins/Runtime/GameplayAbilities/Source/GameplayAbilities/Public/Abilities/GameplayAbility.h b/Engine/Plugins/Runtime/GameplayAbilities/Source/GameplayAbilities/Public/Abilities/GameplayAbility.h
--- a/Engine/Plugins/Runtime/GameplayAbilities/Source/GameplayAbilities/Public/Abilities/GameplayAbility.h
+++ b/Engine/Plugins/Runtime/GameplayAbilities/Source/GameplayAbilities/Public/Abilities/GameplayAbility.h
@@ -21,6 +21,7 @@ class UGameplayAbility;
 class UGameplayTask;
 class UGameplayTasksComponent;
 struct EventData;
+struct FAbilityEndedData;
 
 /**
  * UGameplayAbility
@@ -244,6 +245,9 @@ public:
 	/** Notification that the ability has ended.  Set using TryActivateAbility. */
 	FOnGameplayAbilityEnded OnGameplayAbilityEnded;
 
+	/** Notification that the ability has ended with data on how it was ended */
+	FGameplayAbilityEndedDelegate OnGameplayAbilityEndedWithData;
+
 	/** Notification that the ability is being cancelled.  Called before OnGameplayAbilityEnded. */
 	FOnGameplayAbilityCancelled OnGameplayAbilityCancelled;
 
diff --git a/Engine/Plugins/Runtime/GameplayAbilities/Source/GameplayAbilities/Public/Abilities/GameplayAbilityTypes.h b/Engine/Plugins/Runtime/GameplayAbilities/Source/GameplayAbilities/Public/Abilities/GameplayAbilityTypes.h
--- a/Engine/Plugins/Runtime/GameplayAbilities/Source/GameplayAbilities/Public/Abilities/GameplayAbilityTypes.h
+++ b/Engine/Plugins/Runtime/GameplayAbilities/Source/GameplayAbilities/Public/Abilities/GameplayAbilityTypes.h
@@ -353,6 +353,42 @@ struct GAMEPLAYABILITIES_API FGameplayEventData
 /** Used for cleaning up predicted data on network clients */
 DECLARE_MULTICAST_DELEGATE_OneParam(FGameplayEventMulticastDelegate, const FGameplayEventData*);
 
+/** Ability Ended Data */
+USTRUCT()
+struct FAbilityEndedData
+{
+	GENERATED_USTRUCT_BODY()
+	FAbilityEndedData()
+		: AbilityThatEnded(nullptr)
+		, bReplicateEndAbility(false)
+		, bWasCancelled(false)
+	{
+	}
+
+	FAbilityEndedData(UGameplayAbility* InAbility, FGameplayAbilitySpecHandle InHandle, bool bInReplicateEndAbility, bool bInWasCancelled)
+		: AbilityThatEnded(InAbility)
+		, AbilitySpecHandle(InHandle)
+		, bReplicateEndAbility(bInReplicateEndAbility)
+		, bWasCancelled(bInWasCancelled)
+	{
+	}
+
+	UPROPERTY()
+	UGameplayAbility* AbilityThatEnded;
+
+	UPROPERTY()
+	FGameplayAbilitySpecHandle AbilitySpecHandle;
+
+	UPROPERTY()
+	bool bReplicateEndAbility;
+
+	UPROPERTY()
+	bool bWasCancelled;
+};
+
+/** Notification delegate definition for when the gameplay ability ends */
+DECLARE_MULTICAST_DELEGATE_OneParam(FGameplayAbilityEndedDelegate, const FAbilityEndedData&);
+
 /** 
  *	Structure that tells AbilitySystemComponent what to bind to an InputComponent (see BindAbilityActivationToInputComponent) 
  *	
diff --git a/Engine/Plugins/Runtime/GameplayAbilities/Source/GameplayAbilities/Public/Abilities/Tasks/AbilityTask_WaitGameplayEffectRemoved.h b/Engine/Plugins/Runtime/GameplayAbilities/Source/GameplayAbilities/Public/Abilities/Tasks/AbilityTask_WaitGameplayEffectRemoved.h
--- a/Engine/Plugins/Runtime/GameplayAbilities/Source/GameplayAbilities/Public/Abilities/Tasks/AbilityTask_WaitGameplayEffectRemoved.h
+++ b/Engine/Plugins/Runtime/GameplayAbilities/Source/GameplayAbilities/Public/Abilities/Tasks/AbilityTask_WaitGameplayEffectRemoved.h
@@ -7,7 +7,7 @@
 #include "Abilities/Tasks/AbilityTask.h"
 #include "AbilityTask_WaitGameplayEffectRemoved.generated.h"
 
-DECLARE_DYNAMIC_MULTICAST_DELEGATE(FWaitGameplayEffectRemovedDelegate);
+DECLARE_DYNAMIC_MULTICAST_DELEGATE_OneParam(FWaitGameplayEffectRemovedDelegate, const FGameplayEffectRemovalInfo &, GameplayEffectRemovalInfo);
 
 class AActor;
 class UPrimitiveComponent;
@@ -29,7 +29,7 @@ class GAMEPLAYABILITIES_API UAbilityTask_WaitGameplayEffectRemoved : public UAbi
 	virtual void Activate() override;
 
 	UFUNCTION()
-	void OnGameplayEffectRemoved();
+	void OnGameplayEffectRemoved(const FGameplayEffectRemovalInfo& InGameplayEffectRemovalInfo);
 
 	/** Wait until the specified gameplay effect is removed. */
 	UFUNCTION(BlueprintCallable, Category="Ability|Tasks", meta = (HidePin = "OwningAbility", DefaultToSelf = "OwningAbility", BlueprintInternalUseOnly = "TRUE"))
@@ -43,4 +43,5 @@ protected:
 	bool Registered;
 
 	FDelegateHandle OnGameplayEffectRemovedDelegateHandle;
+	FDelegateHandle OnGameplayEffectRemovedDelegateHandle_Deprecated;
 };
diff --git a/Engine/Plugins/Runtime/GameplayAbilities/Source/GameplayAbilities/Public/AbilitySystemBlueprintLibrary.h b/Engine/Plugins/Runtime/GameplayAbilities/Source/GameplayAbilities/Public/AbilitySystemBlueprintLibrary.h
--- a/Engine/Plugins/Runtime/GameplayAbilities/Source/GameplayAbilities/Public/AbilitySystemBlueprintLibrary.h
+++ b/Engine/Plugins/Runtime/GameplayAbilities/Source/GameplayAbilities/Public/AbilitySystemBlueprintLibrary.h
@@ -312,6 +312,21 @@ class GAMEPLAYABILITIES_API UAbilitySystemBlueprintLibrary : public UBlueprintFu
 	UFUNCTION(BlueprintCallable, Category = "Ability|GameplayEffect")
 	static int32 GetActiveGameplayEffectStackLimitCount(FActiveGameplayEffectHandle ActiveHandle);
 
+	/** Returns the start time (time which the GE was added) for a given GameplayEffect */
+	UFUNCTION(BlueprintCallable, Category = "Ability|GameplayEffect")
+	static float GetActiveGameplayEffectStartTime(FActiveGameplayEffectHandle ActiveHandle);
+
+	/** Returns the expected end time (when we think the GE will expire) for a given GameplayEffect (note someone could remove or change it before that happens!) */
+	UFUNCTION(BlueprintCallable, Category = "Ability|GameplayEffect")
+	static float GetActiveGameplayEffectExpectedEndTime(FActiveGameplayEffectHandle ActiveHandle);
+
+	UFUNCTION(BlueprintCallable, Category = "Ability|GameplayEffect")
+	static float GetActiveGameplayEffectTotalDuration(FActiveGameplayEffectHandle ActiveHandle);
+
+	UFUNCTION(BlueprintCallable, Category = "Ability|GameplayEffect", meta = (WorldContext = "WorldContextObject"))
+	static float GetActiveGameplayEffectRemainingDuration(UObject* WorldContextObject, FActiveGameplayEffectHandle ActiveHandle);
+
 	UFUNCTION(BlueprintPure, Category = "Ability|GameplayEffect", Meta = (DisplayName = "Get Active GameplayEffect Debug String "))
 	static FString GetActiveGameplayEffectDebugString(FActiveGameplayEffectHandle ActiveHandle);
+
 };
diff --git a/Engine/Plugins/Runtime/GameplayAbilities/Source/GameplayAbilities/Public/AbilitySystemComponent.h b/Engine/Plugins/Runtime/GameplayAbilities/Source/GameplayAbilities/Public/AbilitySystemComponent.h
--- a/Engine/Plugins/Runtime/GameplayAbilities/Source/GameplayAbilities/Public/AbilitySystemComponent.h
+++ b/Engine/Plugins/Runtime/GameplayAbilities/Source/GameplayAbilities/Public/AbilitySystemComponent.h
@@ -579,7 +579,11 @@ class GAMEPLAYABILITIES_API UAbilitySystemComponent : public UGameplayTasksCompo
 	// Additional Helper Functions
 	// --------------------------------------------
 	
+	DEPRECATED(4.17, "Use OnGameplayEffectRemoved_InfoDelegate (the delegate signature has changed)")
 	FOnActiveGameplayEffectRemoved* OnGameplayEffectRemovedDelegate(FActiveGameplayEffectHandle Handle);
+
+	FOnActiveGameplayEffectRemoved_Info* OnGameplayEffectRemoved_InfoDelegate(FActiveGameplayEffectHandle Handle);
+
 	FOnGivenActiveGameplayEffectRemoved& OnAnyGameplayEffectRemovedDelegate();
 
 	FOnActiveGameplayEffectStackChange* OnGameplayEffectStackChangeDelegate(FActiveGameplayEffectHandle Handle);
@@ -588,7 +592,6 @@ class GAMEPLAYABILITIES_API UAbilitySystemComponent : public UGameplayTasksCompo
 
 	UFUNCTION(BlueprintCallable, Category = GameplayEffects, meta=(DisplayName = "ApplyGameplayEffectToTarget"))
 	FActiveGameplayEffectHandle BP_ApplyGameplayEffectToTarget(TSubclassOf<UGameplayEffect> GameplayEffectClass, UAbilitySystemComponent *Target, float Level, FGameplayEffectContextHandle Context);
-
 	FActiveGameplayEffectHandle ApplyGameplayEffectToTarget(UGameplayEffect *GameplayEffect, UAbilitySystemComponent *Target, float Level = UGameplayEffect::INVALID_LEVEL, FGameplayEffectContextHandle Context = FGameplayEffectContextHandle(), FPredictionKey PredictionKey = FPredictionKey());
 
 	UFUNCTION(BlueprintCallable, Category = GameplayEffects, meta=(DisplayName = "ApplyGameplayEffectToSelf"))
@@ -636,6 +639,7 @@ class GAMEPLAYABILITIES_API UAbilitySystemComponent : public UGameplayTasksCompo
 
 	TArray<TPair<float,float>> GetActiveEffectsTimeRemainingAndDuration(const FGameplayEffectQuery& Query) const;
 
+	UFUNCTION(BlueprintCallable, BlueprintPure=false, Category = "GameplayEffects", meta=(DisplayName = "Get Activate Gameplay Effects for Query"))
 	TArray<FActiveGameplayEffectHandle> GetActiveEffects(const FGameplayEffectQuery& Query) const;
 
 	/** This will give the world time that all effects matching this query will be finished. If multiple effects match, it returns the one that returns last.*/
@@ -673,6 +677,11 @@ class GAMEPLAYABILITIES_API UAbilitySystemComponent : public UGameplayTasksCompo
 	/** Removes all active effects that match given query. StacksToRemove=-1 will remove all stacks. */
 	int32 RemoveActiveEffects(const FGameplayEffectQuery& Query, int32 StacksToRemove = -1);
 
+	/** Invokes the WhileActive event for all GCs on active, non inhibited, GEs. This would typically be used on "respawn" or something where the mesh/avatar has changed */
+	void ReinvokeActiveGameplayCues();
+
+	// -----------------------------------------------------------
+
 	void OnRestackGameplayEffects();	
 	
 	void PrintAllGameplayEffects() const;
@@ -1160,8 +1169,11 @@ public:
 	/** A generic callback anytime an ability is activated (started) */	
 	FGenericAbilityDelegate AbilityActivatedCallbacks;
 
-	/** A generic callback anytime an ability is commited (cost/cooldown applied) */
-	FGenericAbilityDelegate AbilityCommitedCallbacks;
+	/** callback anytime an ability is ended */
+	FGameplayAbilityEndedDelegate OnAbilityEnded;
+
+	/** A generic callback anytime an ability is committed (cost/cooldown applied) */
+	FGenericAbilityDelegate AbilityCommittedCallbacks;
 	FAbilityFailedDelegate AbilityFailedCallbacks;
 
 	/** Executes a gameplay event. Returns the number of successful ability activations triggered by the event */
diff --git a/Engine/Plugins/Runtime/GameplayAbilities/Source/GameplayAbilities/Public/AttributeSet.h b/Engine/Plugins/Runtime/GameplayAbilities/Source/GameplayAbilities/Public/AttributeSet.h
--- a/Engine/Plugins/Runtime/GameplayAbilities/Source/GameplayAbilities/Public/AttributeSet.h
+++ b/Engine/Plugins/Runtime/GameplayAbilities/Source/GameplayAbilities/Public/AttributeSet.h
@@ -126,7 +126,7 @@ struct GAMEPLAYABILITIES_API FGameplayAttribute
 	FString AttributeName;
 
 	// In editor, this will filter out properties with meta tag "HideInDetailsView" or equal to FilterMetaStr. In non editor, it returns all properties.
-	static void GetAllAttributeProperties(TArray<UProperty*>& OutProperties, FString FilterMetaStr=FString());
+	static void GetAllAttributeProperties(TArray<UProperty*>& OutProperties, FString FilterMetaStr=FString(), bool UseEditorOnlyData=true);
 
 private:
 	friend class FAttributePropertyDetails;
@@ -255,10 +255,10 @@ struct GAMEPLAYABILITIES_API FScalableFloat
 
 public:
 
-	UPROPERTY(EditDefaultsOnly, Category=ScalableFloat)
+	UPROPERTY(EditDefaultsOnly, BlueprintReadWrite, Category=ScalableFloat)
 	float	Value;
 
-	UPROPERTY(EditDefaultsOnly, Category=ScalableFloat)
+	UPROPERTY(EditDefaultsOnly, BlueprintReadWrite, Category=ScalableFloat)
 	FCurveTableRowHandle	Curve;
 
 	float GetValueAtLevel(float Level, const FString* ContextString = nullptr) const;
diff --git a/Engine/Plugins/Runtime/GameplayAbilities/Source/GameplayAbilities/Public/GameplayCueManager.h b/Engine/Plugins/Runtime/GameplayAbilities/Source/GameplayAbilities/Public/GameplayCueManager.h
--- a/Engine/Plugins/Runtime/GameplayAbilities/Source/GameplayAbilities/Public/GameplayCueManager.h
+++ b/Engine/Plugins/Runtime/GameplayAbilities/Source/GameplayAbilities/Public/GameplayCueManager.h
@@ -316,7 +316,7 @@ protected:
 
 	/** Classes that we need to preallocate instances for */
 	UPROPERTY(transient)
-	TArray<AGameplayCueNotify_Actor*> GameplayCueClassesForPreallocation;
+	TArray<TSubclassOf<AGameplayCueNotify_Actor>> GameplayCueClassesForPreallocation;
 
 	/** List of gameplay cue executes that haven't been processed yet */
 	UPROPERTY(transient)
diff --git a/Engine/Plugins/Runtime/GameplayAbilities/Source/GameplayAbilities/Public/GameplayCue_Types.h b/Engine/Plugins/Runtime/GameplayAbilities/Source/GameplayAbilities/Public/GameplayCue_Types.h
--- a/Engine/Plugins/Runtime/GameplayAbilities/Source/GameplayAbilities/Public/GameplayCue_Types.h
+++ b/Engine/Plugins/Runtime/GameplayAbilities/Source/GameplayAbilities/Public/GameplayCue_Types.h
@@ -67,7 +67,7 @@ struct FPreallocationInfo
 	TMap<UClass*, TArray<AGameplayCueNotify_Actor*> >	PreallocatedInstances;
 
 	UPROPERTY(transient)
-	TArray<AGameplayCueNotify_Actor*>	ClassesNeedingPreallocation;
+	TArray<TSubclassOf<AGameplayCueNotify_Actor>>	ClassesNeedingPreallocation;
 
 	FObjectKey OwningWorldKey;
 };
diff --git a/Engine/Plugins/Runtime/GameplayAbilities/Source/GameplayAbilities/Public/GameplayEffect.h b/Engine/Plugins/Runtime/GameplayAbilities/Source/GameplayAbilities/Public/GameplayEffect.h
--- a/Engine/Plugins/Runtime/GameplayAbilities/Source/GameplayAbilities/Public/GameplayEffect.h
+++ b/Engine/Plugins/Runtime/GameplayAbilities/Source/GameplayAbilities/Public/GameplayEffect.h
@@ -184,13 +184,17 @@ public:
 	FScalableFloat Coefficient;
 
 	/** Additive value to the attribute calculation, added in before the coefficient applies */
-	UPROPERTY(EditDefaultsOnly, Category=AttributeFloat)
+	UPROPERTY(EditDefaultsOnly, Category=CustomCalculation)
 	FScalableFloat PreMultiplyAdditiveValue;
 
 	/** Additive value to the attribute calculation, added in after the coefficient applies */
-	UPROPERTY(EditDefaultsOnly, Category=AttributeFloat)
+	UPROPERTY(EditDefaultsOnly, Category=CustomCalculation)
 	FScalableFloat PostMultiplyAdditiveValue;
 
+	/** If a curve table entry is specified, the OUTPUT of this custom class magnitude (including the pre and post additive values) lookup into the curve instead of using the attribute directly. */
+	UPROPERTY(EditDefaultsOnly, Category=CustomCalculation)
+	FCurveTableRowHandle FinalLookupCurve;
+
 	/** Equality/Inequality operators */
 	bool operator==(const FCustomCalculationBasedFloat& Other) const;
 	bool operator!=(const FCustomCalculationBasedFloat& Other) const;
@@ -757,6 +761,9 @@ struct GAMEPLAYABILITIES_API FGameplayEffectAttributeCaptureSpec
 	 */
 	bool AttemptAddAggregatorModsToAggregator(OUT FAggregator& OutAggregatorToAddTo) const;
 	
+	/** Gathers made for a given capture. Note that these mods are unqualified and direct references (not copies). Only use this if you know what you are doing. */
+	bool AttemptGatherAttributeMods(OUT TMap<EGameplayModEvaluationChannel, const TArray<FAggregatorMod>*>& OutModMap) const;
+	
 	/** Simple accessor to backing capture definition */
 	const FGameplayEffectAttributeCaptureDefinition& GetBackingDefinition() const;
 
@@ -1252,6 +1259,7 @@ struct GAMEPLAYABILITIES_API FActiveGameplayEffect : public FFastArraySerializer
 	int32 ClientCachedStackCount;
 
 	FOnActiveGameplayEffectRemoved OnRemovedDelegate;
+	FOnActiveGameplayEffectRemoved_Info OnRemoved_InfoDelegate;
 
 	FOnActiveGameplayEffectStackChange OnStackChangeDelegate;
 
@@ -1700,7 +1708,7 @@ private:
 	
 	/** Called both in server side creation and replication creation/deletion */
 	void InternalOnActiveGameplayEffectAdded(FActiveGameplayEffect& Effect);
-	void InternalOnActiveGameplayEffectRemoved(FActiveGameplayEffect& Effect, bool bInvokeGameplayCueEvents);
+	void InternalOnActiveGameplayEffectRemoved(FActiveGameplayEffect& Effect, bool bInvokeGameplayCueEvents, const FGameplayEffectRemovalInfo& GameplayEffectRemovalInfo);
 
 	void RemoveActiveGameplayEffectGrantedTagsAndModifiers(const FActiveGameplayEffect& Effect, bool bInvokeGameplayCueEvents);
 	void AddActiveGameplayEffectGrantedTagsAndModifiers(FActiveGameplayEffect& Effect, bool bInvokeGameplayCueEvents);
diff --git a/Engine/Plugins/Runtime/GameplayAbilities/Source/GameplayAbilities/Public/GameplayEffectAggregator.h b/Engine/Plugins/Runtime/GameplayAbilities/Source/GameplayAbilities/Public/GameplayEffectAggregator.h
--- a/Engine/Plugins/Runtime/GameplayAbilities/Source/GameplayAbilities/Public/GameplayEffectAggregator.h
+++ b/Engine/Plugins/Runtime/GameplayAbilities/Source/GameplayAbilities/Public/GameplayEffectAggregator.h
@@ -100,12 +100,15 @@ struct GAMEPLAYABILITIES_API FAggregatorModChannel
 	void AddModsFrom(const FAggregatorModChannel& Other);
 
 	/**
-	 * Populate a mapping of channel to corresponding mods for debugging purposes
+	 * Populate a mapping of channel to corresponding mods
 	 * 
 	 * @param Channel	Enum channel associated with this channel
 	 * @param OutMods	Mapping of channel enum to mods
 	 */
-	void DebugGetAllAggregatorMods(EGameplayModEvaluationChannel Channel, OUT TMap<EGameplayModEvaluationChannel, const TArray<FAggregatorMod>*>& OutMods) const;
+	void GetAllAggregatorMods(EGameplayModEvaluationChannel Channel, OUT TMap<EGameplayModEvaluationChannel, const TArray<FAggregatorMod>*>& OutMods) const;
+
+	DEPRECATED(4.17, "Use GetAllAggregatorMods")
+	void DebugGetAllAggregatorMods(EGameplayModEvaluationChannel Channel, OUT TMap<EGameplayModEvaluationChannel, const TArray<FAggregatorMod>*>& OutMods) const { return GetAllAggregatorMods(Channel, OutMods); }
 	
 	/**
 	 * Called when the mod channel's gameplay effect dependencies have potentially been swapped out for new ones, like when GE arrays are cloned.
@@ -208,7 +211,10 @@ struct GAMEPLAYABILITIES_API FAggregatorModChannelContainer
 	 * 
 	 * @param OutMods	Mapping of channel enum to mods
 	 */
-	void DebugGetAllAggregatorMods(OUT TMap<EGameplayModEvaluationChannel, const TArray<FAggregatorMod>*>& OutMods) const;
+	void GetAllAggregatorMods(OUT TMap<EGameplayModEvaluationChannel, const TArray<FAggregatorMod>*>& OutMods) const;
+
+	DEPRECATED(4.17, "Use GetAllAggregatorMods")
+	void DebugGetAllAggregatorMods(OUT TMap<EGameplayModEvaluationChannel, const TArray<FAggregatorMod>*>& OutMods) const { return GetAllAggregatorMods(OutMods); }
 
 	/**
 	 * Called when the container's gameplay effect dependencies have potentially been swapped out for new ones, like when GE arrays are cloned.
@@ -280,11 +286,14 @@ struct GAMEPLAYABILITIES_API FAggregator : public TSharedFromThis<FAggregator>
 	void RemoveDependent(FActiveGameplayEffectHandle Handle);
 
 	/**
-	 * Populate a mapping of channel to corresponding mods for debugging purposes
+	 * Populate a mapping of channel to corresponding mods
 	 * 
 	 * @param OutMods	Mapping of channel enum to mods
 	 */
-	void DebugGetAllAggregatorMods(OUT TMap<EGameplayModEvaluationChannel, const TArray<FAggregatorMod>*>& OutMods) const;
+	void GetAllAggregatorMods(OUT TMap<EGameplayModEvaluationChannel, const TArray<FAggregatorMod>*>& OutMods) const;
+
+	DEPRECATED(4.17, "Use GetAllAggregatorMods")
+	void DebugGetAllAggregatorMods(OUT TMap<EGameplayModEvaluationChannel, const TArray<FAggregatorMod>*>& OutMods) const { return GetAllAggregatorMods(OutMods); }
 	
 	/**
 	 * Called when the aggregator's gameplay effect dependencies have potentially been swapped out for new ones, like when GE arrays are cloned.
diff --git a/Engine/Plugins/Runtime/GameplayAbilities/Source/GameplayAbilities/Public/GameplayEffectExecutionCalculation.h b/Engine/Plugins/Runtime/GameplayAbilities/Source/GameplayAbilities/Public/GameplayEffectExecutionCalculation.h
--- a/Engine/Plugins/Runtime/GameplayAbilities/Source/GameplayAbilities/Public/GameplayEffectExecutionCalculation.h
+++ b/Engine/Plugins/Runtime/GameplayAbilities/Source/GameplayAbilities/Public/GameplayEffectExecutionCalculation.h
@@ -105,6 +105,18 @@ public:
 	 */
 	bool AttemptGetCapturedAttributeAggregatorSnapshot(const FGameplayEffectAttributeCaptureDefinition& InCaptureDef, OUT FAggregator& OutSnapshottedAggregator) const;
 
+	/** 
+	 *	Returns all modifiers for a given captured def. Note the returned list is a direct reference to the internal list stored in the attribute aggregators (possibly a snapshot copy or possibly the "real" ones).
+	 *	Also note the modifiers returned or NOT qualified! You will still want to qualifiy them on an FAggregatorEvaluateParameters yourself. (Having this function internally do it would onyl be possible if this function
+	 *	returned a copy of the mod list which is something we would like to avoid).
+	 *	
+	 *	Consider using ForEachQualifiedAttributeMod when you want to "do something for every qualifier mod"
+	 */
+	bool AttemptGatherAttributeMods(const FGameplayEffectAttributeCaptureDefinition& InCaptureDef, OUT TMap<EGameplayModEvaluationChannel, const TArray<FAggregatorMod>*>& OutModMap) const;
+
+	/** Runs given TFunction on every qualifier mod for a given AttributeCaptureDefinition */
+	bool ForEachQualifiedAttributeMod(const FGameplayEffectAttributeCaptureDefinition& InCaptureDef, const FAggregatorEvaluateParameters& InEvalParams, TFunction< void(EGameplayModEvaluationChannel, EGameplayModOp::Type, const FAggregatorMod&) >) const;
+
 private:
 
 	/** Mapping of capture definition to aggregator with scoped modifiers added in; Used to process scoped modifiers w/o modifying underlying aggregators in the capture */
diff --git a/Engine/Plugins/Runtime/GameplayAbilities/Source/GameplayAbilities/Public/GameplayEffectTypes.h b/Engine/Plugins/Runtime/GameplayAbilities/Source/GameplayAbilities/Public/GameplayEffectTypes.h
--- a/Engine/Plugins/Runtime/GameplayAbilities/Source/GameplayAbilities/Public/GameplayEffectTypes.h
+++ b/Engine/Plugins/Runtime/GameplayAbilities/Source/GameplayAbilities/Public/GameplayEffectTypes.h
@@ -838,6 +838,28 @@ struct TStructOpsTypeTraits<FGameplayEffectContextHandle> : public TStructOpsTyp
 	};
 };
 
+
+/**
+ * FGameplayEffectRemovalInfo
+ *	Data struct for containing information pertinent to GameplayEffects as they are removed.
+ */
+USTRUCT(BlueprintType)
+struct FGameplayEffectRemovalInfo
+{
+	GENERATED_USTRUCT_BODY()
+
+	/** True when the gameplay effect's duration has not expired, meaning the gameplay effect is being forcefully removed.  */
+	UPROPERTY(EditAnywhere, BlueprintReadWrite, Category="Removal")
+	bool bPrematureRemoval;
+
+	/** Number of Stacks this gameplay effect had before it was removed. */
+	UPROPERTY(EditAnywhere, BlueprintReadWrite, Category="Removal")
+	int32 StackCount;
+
+	/** Actor this gameplay effect was targeting. */
+	UPROPERTY(EditAnywhere, BlueprintReadWrite, Category="Removal")
+	FGameplayEffectContextHandle EffectContext;
+};
 // -----------------------------------------------------------
 
 
@@ -962,6 +984,7 @@ DECLARE_DELEGATE_OneParam(FOnGameplayAttributeEffectExecuted, struct FGameplayMo
 DECLARE_MULTICAST_DELEGATE_TwoParams(FOnGameplayEffectTagCountChanged, const FGameplayTag, int32);
 
 DECLARE_MULTICAST_DELEGATE(FOnActiveGameplayEffectRemoved);
+DECLARE_MULTICAST_DELEGATE_OneParam(FOnActiveGameplayEffectRemoved_Info, const FGameplayEffectRemovalInfo&);
 
 DECLARE_MULTICAST_DELEGATE_OneParam(FOnGivenActiveGameplayEffectRemoved, const FActiveGameplayEffect&);
 
diff --git a/Engine/Plugins/Runtime/ImmediatePhysics/Source/ImmediatePhysics/Private/ImmediatePhysicsSimulation.cpp b/Engine/Plugins/Runtime/ImmediatePhysics/Source/ImmediatePhysics/Private/ImmediatePhysicsSimulation.cpp
--- a/Engine/Plugins/Runtime/ImmediatePhysics/Source/ImmediatePhysics/Private/ImmediatePhysicsSimulation.cpp
+++ b/Engine/Plugins/Runtime/ImmediatePhysics/Source/ImmediatePhysics/Private/ImmediatePhysicsSimulation.cpp
@@ -281,17 +281,17 @@ void FSimulation::EvictCache()
 	
 }*/
 
-DECLARE_DWORD_COUNTER_STAT(TEXT("Simulated Bodies"), STAT_NumSimulatedBodies, STATGROUP_ImmediatePhysics);
-DECLARE_DWORD_COUNTER_STAT(TEXT("Active Simulated Bodies"), STAT_NumActiveSimulatedBodies, STATGROUP_ImmediatePhysics);
-DECLARE_DWORD_COUNTER_STAT(TEXT("Kinematic Bodies"), STAT_NumKinematicBodies, STATGROUP_ImmediatePhysics);
+DECLARE_DWORD_COUNTER_STAT(TEXT("Simulated Bodies"), STAT_IPNumSimulatedBodies, STATGROUP_ImmediatePhysics);
+DECLARE_DWORD_COUNTER_STAT(TEXT("Active Simulated Bodies"), STAT_IPNumActiveSimulatedBodies, STATGROUP_ImmediatePhysics);
+DECLARE_DWORD_COUNTER_STAT(TEXT("Kinematic Bodies"), STAT_IPNumKinematicBodies, STATGROUP_ImmediatePhysics);
 
 DECLARE_CYCLE_STAT(TEXT("FSimulation::Simulate"), STAT_ImmediateSimulate, STATGROUP_ImmediatePhysics);
 
 void FSimulation::Simulate(float DeltaTime, const FVector& Gravity)
 {
-	SET_DWORD_STAT(STAT_NumSimulatedBodies, NumSimulatedBodies);
-	SET_DWORD_STAT(STAT_NumActiveSimulatedBodies, NumActiveSimulatedBodies);
-	SET_DWORD_STAT(STAT_NumKinematicBodies, NumKinematicBodies);
+	SET_DWORD_STAT(STAT_IPNumSimulatedBodies, NumSimulatedBodies);
+	SET_DWORD_STAT(STAT_IPNumActiveSimulatedBodies, NumActiveSimulatedBodies);
+	SET_DWORD_STAT(STAT_IPNumKinematicBodies, NumKinematicBodies);
 
 	DeltaTime = FMath::Min(DeltaTime, 0.033f);
 	//Create dynamic bodies and integrate their unconstrained velocities
diff --git a/Engine/Plugins/Runtime/SignificanceManager/Source/SignificanceManager/Private/SignificanceManager.cpp b/Engine/Plugins/Runtime/SignificanceManager/Source/SignificanceManager/Private/SignificanceManager.cpp
--- a/Engine/Plugins/Runtime/SignificanceManager/Source/SignificanceManager/Private/SignificanceManager.cpp
+++ b/Engine/Plugins/Runtime/SignificanceManager/Source/SignificanceManager/Private/SignificanceManager.cpp
@@ -43,7 +43,7 @@ TSubclassOf<USignificanceManager> FSignificanceManagerModule::SignificanceManage
 void FSignificanceManagerModule::StartupModule()
 {
 	FWorldDelegates::OnPreWorldInitialization.AddStatic(&FSignificanceManagerModule::OnWorldInit);
-	FWorldDelegates::OnWorldCleanup.AddStatic(&FSignificanceManagerModule::OnWorldCleanup);
+	FWorldDelegates::OnPostWorldCleanup.AddStatic(&FSignificanceManagerModule::OnWorldCleanup);
 	if (!IsRunningDedicatedServer())
 	{
 		AHUD::OnShowDebugInfo.AddStatic(&FSignificanceManagerModule::OnShowDebugInfo);
diff --git a/Engine/Source/Developer/HierarchicalLODUtilities/Private/HierarchicalLODUtilities.cpp b/Engine/Source/Developer/HierarchicalLODUtilities/Private/HierarchicalLODUtilities.cpp
--- a/Engine/Source/Developer/HierarchicalLODUtilities/Private/HierarchicalLODUtilities.cpp
+++ b/Engine/Source/Developer/HierarchicalLODUtilities/Private/HierarchicalLODUtilities.cpp
@@ -146,7 +146,10 @@ bool FHierarchicalLODUtilities::BuildStaticMeshForLODActor(ALODActor* LODActor,
 				}
 
 				// TODO: support instanced static meshes
-				Components.RemoveAll([](UStaticMeshComponent* Val){ return Val->IsA(UInstancedStaticMeshComponent::StaticClass()); });
+				Components.RemoveAll([](UStaticMeshComponent* Val)
+				{ 
+					return Val->IsA(UInstancedStaticMeshComponent::StaticClass()) || !Val->ShouldGenerateAutoLOD();
+				});
 
 				AllComponents.Append(Components);
 			}
diff --git a/Engine/Source/Developer/MeshMergeUtilities/Private/ProxyMaterialUtilities.h b/Engine/Source/Developer/MeshMergeUtilities/Private/ProxyMaterialUtilities.h
--- a/Engine/Source/Developer/MeshMergeUtilities/Private/ProxyMaterialUtilities.h
+++ b/Engine/Source/Developer/MeshMergeUtilities/Private/ProxyMaterialUtilities.h
@@ -180,6 +180,10 @@ namespace ProxyMaterialUtilities
 
 			// Merge properties into one texture using the separate colour channels
 			const EFlattenMaterialProperties Properties[3] = { EFlattenMaterialProperties::Metallic , EFlattenMaterialProperties::Roughness, EFlattenMaterialProperties::Specular};
+
+			//Property that is not part of the pack (because of a different size), will see is reserve pack space fill with Black Color.
+			const bool PropertyShouldBePack[3] = { bPackMetallic , bPackRoughness , bPackSpecular };
+
 			// Red mask (all properties are rendered into the red channel)
 			FColor NonAlphaRed = FColor::Red;
 			NonAlphaRed.A = 0;
@@ -188,7 +192,7 @@ namespace ProxyMaterialUtilities
 			for (int32 PropertyIndex = 0; PropertyIndex < 3; ++PropertyIndex)
 			{
 				const EFlattenMaterialProperties Property = Properties[PropertyIndex];
-				const bool HasProperty = FlattenMaterial.DoesPropertyContainData(Property) && !FlattenMaterial.IsPropertyConstant(Property);
+				const bool HasProperty = PropertyShouldBePack[PropertyIndex] && FlattenMaterial.DoesPropertyContainData(Property) && !FlattenMaterial.IsPropertyConstant(Property);
 
 				if (HasProperty)
 				{
diff --git a/Engine/Source/Developer/SimplygonMeshReduction/Private/SimplygonMeshReduction.cpp b/Engine/Source/Developer/SimplygonMeshReduction/Private/SimplygonMeshReduction.cpp
--- a/Engine/Source/Developer/SimplygonMeshReduction/Private/SimplygonMeshReduction.cpp
+++ b/Engine/Source/Developer/SimplygonMeshReduction/Private/SimplygonMeshReduction.cpp
@@ -1094,6 +1094,11 @@ public:
 		{
 			SmoothingMask = 1;
 		}
+
+		if (!OutProxyMesh.IsValid())
+		{
+			FailedDelegate.ExecuteIfBound(InJobGUID, TEXT("Simplygon failed to generate a valid proxy mesh"));
+		}
 				
 		CompleteDelegate.ExecuteIfBound(OutProxyMesh, OutMaterial, InJobGUID);
 	}
diff --git a/Engine/Source/Developer/SlateReflector/Private/Models/WidgetReflectorNode.cpp b/Engine/Source/Developer/SlateReflector/Private/Models/WidgetReflectorNode.cpp
--- a/Engine/Source/Developer/SlateReflector/Private/Models/WidgetReflectorNode.cpp
+++ b/Engine/Source/Developer/SlateReflector/Private/Models/WidgetReflectorNode.cpp
@@ -121,6 +121,11 @@ FText FLiveWidgetReflectorNode::GetWidgetClippingText() const
 	return FWidgetReflectorNodeUtils::GetWidgetClippingText(Widget.Pin());
 }
 
+bool FLiveWidgetReflectorNode::GetWidgetFocusable() const
+{
+	return FWidgetReflectorNodeUtils::GetWidgetFocusable(Widget.Pin());
+}
+
 FText FLiveWidgetReflectorNode::GetWidgetReadableLocation() const
 {
 	return FWidgetReflectorNodeUtils::GetWidgetReadableLocation(Widget.Pin());
@@ -187,6 +192,7 @@ FSnapshotWidgetReflectorNode::FSnapshotWidgetReflectorNode(const FArrangedWidget
 	: FWidgetReflectorNodeBase(InArrangedWidget)
 	, CachedWidgetType(FWidgetReflectorNodeUtils::GetWidgetType(InArrangedWidget.Widget))
 	, CachedWidgetVisibilityText(FWidgetReflectorNodeUtils::GetWidgetVisibilityText(InArrangedWidget.Widget))
+	, bCachedWidgetFocusable(FWidgetReflectorNodeUtils::GetWidgetFocusable(InArrangedWidget.Widget))
 	, CachedWidgetClippingText(FWidgetReflectorNodeUtils::GetWidgetClippingText(InArrangedWidget.Widget))
 	, CachedWidgetReadableLocation(FWidgetReflectorNodeUtils::GetWidgetReadableLocation(InArrangedWidget.Widget))
 	, CachedWidgetFile(FWidgetReflectorNodeUtils::GetWidgetFile(InArrangedWidget.Widget))
@@ -219,6 +225,11 @@ FText FSnapshotWidgetReflectorNode::GetWidgetVisibilityText() const
 	return CachedWidgetVisibilityText;
 }
 
+bool FSnapshotWidgetReflectorNode::GetWidgetFocusable() const
+{
+	return bCachedWidgetFocusable;
+}
+
 FText FSnapshotWidgetReflectorNode::GetWidgetClippingText() const
 {
 	return CachedWidgetClippingText;
@@ -343,6 +354,7 @@ TSharedRef<FJsonValue> FSnapshotWidgetReflectorNode::ToJson(const TSharedRef<FSn
 	RootJsonObject->SetField(TEXT("HitTestInfo"), Internal::CreateWidgetHitTestInfoJsonValue(RootSnapshotNode->HitTestInfo));
 	RootJsonObject->SetStringField(TEXT("WidgetType"), RootSnapshotNode->CachedWidgetType.ToString());
 	RootJsonObject->SetStringField(TEXT("WidgetVisibilityText"), RootSnapshotNode->CachedWidgetVisibilityText.ToString());
+	RootJsonObject->SetBoolField(TEXT("WidgetFocusable"), RootSnapshotNode->bCachedWidgetFocusable);
 	RootJsonObject->SetStringField(TEXT("WidgetReadableLocation"), RootSnapshotNode->CachedWidgetReadableLocation.ToString());
 	RootJsonObject->SetStringField(TEXT("WidgetFile"), RootSnapshotNode->CachedWidgetFile);
 	RootJsonObject->SetNumberField(TEXT("WidgetLineNumber"), RootSnapshotNode->CachedWidgetLineNumber);
@@ -474,6 +486,7 @@ TSharedRef<FSnapshotWidgetReflectorNode> FSnapshotWidgetReflectorNode::FromJson(
 	RootSnapshotNode->HitTestInfo = Internal::ParseWidgetHitTestInfoJsonValue(RootJsonObject->GetField<EJson::None>(TEXT("HitTestInfo")));
 	RootSnapshotNode->CachedWidgetType = FText::FromString(RootJsonObject->GetStringField(TEXT("WidgetType")));
 	RootSnapshotNode->CachedWidgetVisibilityText = FText::FromString(RootJsonObject->GetStringField(TEXT("WidgetVisibilityText")));
+	RootSnapshotNode->bCachedWidgetFocusable = RootJsonObject->GetBoolField(TEXT("WidgetFocusable"));
 	RootSnapshotNode->CachedWidgetReadableLocation = FText::FromString(RootJsonObject->GetStringField(TEXT("WidgetReadableLocation")));
 	RootSnapshotNode->CachedWidgetFile = RootJsonObject->GetStringField(TEXT("WidgetFile"));
 	RootSnapshotNode->CachedWidgetLineNumber = RootJsonObject->GetIntegerField(TEXT("WidgetLineNumber"));
@@ -578,6 +591,12 @@ FText FWidgetReflectorNodeUtils::GetWidgetVisibilityText(const TSharedPtr<SWidge
 	return (InWidget.IsValid()) ? FText::FromString(InWidget->GetVisibility().ToString()) : FText::GetEmpty();
 }
 
+bool FWidgetReflectorNodeUtils::GetWidgetFocusable(const TSharedPtr<SWidget>& InWidget)
+{
+	return InWidget.IsValid() ? InWidget->SupportsKeyboardFocus() : false;
+}
+
+
 FText FWidgetReflectorNodeUtils::GetWidgetClippingText(const TSharedPtr<SWidget>& InWidget)
 {
 	if ( InWidget.IsValid() )
diff --git a/Engine/Source/Developer/SlateReflector/Private/Models/WidgetReflectorNode.h b/Engine/Source/Developer/SlateReflector/Private/Models/WidgetReflectorNode.h
--- a/Engine/Source/Developer/SlateReflector/Private/Models/WidgetReflectorNode.h
+++ b/Engine/Source/Developer/SlateReflector/Private/Models/WidgetReflectorNode.h
@@ -69,6 +69,11 @@ public:
 	 */
 	virtual FText GetWidgetClippingText() const = 0;
 
+	/**
+	* @return The bool indicating whether or not the widget we were initialized from reports as Focusable
+	*/
+	virtual bool GetWidgetFocusable() const = 0;
+
 	/**
 	 * The human readable location for widgets that are defined in C++ is the file and line number
 	 * The human readable location for widgets that are defined in UMG is the asset name
@@ -208,6 +213,7 @@ public:
 	virtual FText GetWidgetType() const override;
 	virtual FText GetWidgetVisibilityText() const override;
 	virtual FText GetWidgetClippingText() const override;
+	virtual bool GetWidgetFocusable() const override;
 	virtual FText GetWidgetReadableLocation() const override;
 	virtual FString GetWidgetFile() const override;
 	virtual int32 GetWidgetLineNumber() const override;
@@ -256,6 +262,7 @@ public:
 	virtual FText GetWidgetType() const override;
 	virtual FText GetWidgetVisibilityText() const override;
 	virtual FText GetWidgetClippingText() const override;
+	virtual bool GetWidgetFocusable() const override;
 	virtual FText GetWidgetReadableLocation() const override;
 	virtual FString GetWidgetFile() const override;
 	virtual int32 GetWidgetLineNumber() const override;
@@ -290,7 +297,10 @@ private:
 	/** The visibility string of the widget at the point it was passed to Initialize */
 	FText CachedWidgetVisibilityText;
 
-	/** The clipping string of the widget at the point it was passed to Initialize */
+	/** The focusability of the widget at the point it was passed to Initialize */
+	bool bCachedWidgetFocusable;
+	
+		/** The clipping string of the widget at the point it was passed to Initialize */
 	FText CachedWidgetClippingText;
 
 	/** The human readable location (source file for C++ widgets, asset name for UMG widgets) of the widget at the point it was passed to Initialize */
@@ -397,10 +407,16 @@ public:
 	 */
 	static FText GetWidgetVisibilityText(const TSharedPtr<SWidget>& InWidget);
 
-	/**
+/**
 	 * @return The current clipping string for the given widget
 	 */
 	static FText GetWidgetClippingText(const TSharedPtr<SWidget>& InWidget);
+
+	/**
+	* @return The current focusability for the given widget
+	*/
+	static bool GetWidgetFocusable(const TSharedPtr<SWidget>& InWidget);
+
 	
 	/**
 	 * The human readable location for widgets that are defined in C++ is the file and line number
diff --git a/Engine/Source/Developer/SlateReflector/Private/Widgets/SWidgetReflector.cpp b/Engine/Source/Developer/SlateReflector/Private/Widgets/SWidgetReflector.cpp
--- a/Engine/Source/Developer/SlateReflector/Private/Widgets/SWidgetReflector.cpp
+++ b/Engine/Source/Developer/SlateReflector/Private/Widgets/SWidgetReflector.cpp
@@ -842,8 +842,27 @@ TSharedRef<SDockTab> SWidgetReflector::SpawnWidgetHierarchyTab(const FSpawnTabAr
 						]
 
 						+SHeaderRow::Column(SReflectorTreeWidgetItem::NAME_Visibility)
-						.DefaultLabel(LOCTEXT("Visibility", "Visibility" ))
 						.FixedWidth(125.0f)
+						.HAlignHeader(HAlign_Center)
+						.VAlignHeader(VAlign_Center)
+						.HeaderContent()
+						[
+							SNew(STextBlock)
+							.Text(LOCTEXT("Visibility", "Visibility" ))
+							.ToolTipText(LOCTEXT("VisibilityTooltip", "Visibility"))
+						]
+
+						+ SHeaderRow::Column("Focusable")
+						.DefaultLabel(LOCTEXT("Focusable", "Focusable?"))
+						.FixedWidth(125.0f)
+						.HAlignHeader(HAlign_Center)
+						.VAlignHeader(VAlign_Center)
+						.HeaderContent()
+						[
+							SNew(STextBlock)
+							.Text(LOCTEXT("Focusable", "Focusable?"))
+							.ToolTipText(LOCTEXT("FocusableTooltip", "Focusability (Note that for hit-test directional navigation to work it must be Focusable and \"Visible\"!)"))
+						]
 
 						+SHeaderRow::Column(SReflectorTreeWidgetItem::NAME_Clipping)
 						.DefaultLabel(LOCTEXT("Clipping", "Clipping" ))
diff --git a/Engine/Source/Developer/SlateReflector/Private/Widgets/SWidgetReflectorTreeWidgetItem.cpp b/Engine/Source/Developer/SlateReflector/Private/Widgets/SWidgetReflectorTreeWidgetItem.cpp
--- a/Engine/Source/Developer/SlateReflector/Private/Widgets/SWidgetReflectorTreeWidgetItem.cpp
+++ b/Engine/Source/Developer/SlateReflector/Private/Widgets/SWidgetReflectorTreeWidgetItem.cpp
@@ -5,6 +5,8 @@
 #include "Widgets/Colors/SColorBlock.h"
 #include "Widgets/Input/SHyperlink.h"
 #include "HAL/PlatformApplicationMisc.h"
+#include "EditorStyleSet.h"
+#include "SCheckBox.h"
 
 
 /* SMultiColumnTableRow overrides
@@ -62,6 +64,19 @@ TSharedRef<SWidget> SReflectorTreeWidgetItem::GenerateWidgetForColumn(const FNam
 			[
 				SNew(STextBlock)
 				.Text(this, &SReflectorTreeWidgetItem::GetVisibilityAsString)
+					.Justification(ETextJustify::Center)
+			];
+	}
+	else if (ColumnName == "Focusable")
+	{
+		return SNew(SBox)
+			.HAlign(HAlign_Center)
+			.VAlign(VAlign_Center)
+			.Padding(FMargin(2.0f, 0.0f))
+			[
+				SNew(SCheckBox)
+				.Style(&FEditorStyle::Get().GetWidgetStyle<FCheckBoxStyle>("Toolbar.Check"))
+				.IsChecked(this, &SReflectorTreeWidgetItem::GetFocusableAsCheckBoxState)
 			];
 	}
 	else if ( ColumnName == NAME_Clipping )
diff --git a/Engine/Source/Developer/SlateReflector/Private/Widgets/SWidgetReflectorTreeWidgetItem.h b/Engine/Source/Developer/SlateReflector/Private/Widgets/SWidgetReflectorTreeWidgetItem.h
--- a/Engine/Source/Developer/SlateReflector/Private/Widgets/SWidgetReflectorTreeWidgetItem.h
+++ b/Engine/Source/Developer/SlateReflector/Private/Widgets/SWidgetReflectorTreeWidgetItem.h
@@ -55,6 +55,7 @@ public:
 		CachedWidgetType = WidgetInfo->GetWidgetType();
 		CachedWidgetVisibility = WidgetInfo->GetWidgetVisibilityText();
 		CachedWidgetClipping = WidgetInfo->GetWidgetClippingText();
+		bCachedWidgetFocusable = WidgetInfo->GetWidgetFocusable();
 		CachedReadableLocation = WidgetInfo->GetWidgetReadableLocation();
 		CachedWidgetFile = WidgetInfo->GetWidgetFile();
 		CachedWidgetLineNumber = WidgetInfo->GetWidgetLineNumber();
@@ -100,12 +101,18 @@ protected:
 	{
 		return CachedWidgetVisibility;
 	}
-
+	
+	
 	FText GetClippingAsString() const
 	{
 		return CachedWidgetClipping;
 	}
 
+	ECheckBoxState GetFocusableAsCheckBoxState() const
+	{
+		return bCachedWidgetFocusable ? ECheckBoxState::Checked : ECheckBoxState::Unchecked;
+	}
+
 	/** @return The tint of the reflector node */
 	FSlateColor GetTint() const
 	{
@@ -122,6 +129,7 @@ private:
 	FText CachedWidgetType;
 	FText CachedWidgetVisibility;
 	FText CachedWidgetClipping;
+	bool bCachedWidgetFocusable;
 	FText CachedReadableLocation;
 	FString CachedWidgetFile;
 	int32 CachedWidgetLineNumber;
diff --git a/Engine/Source/Developer/SlateReflector/SlateReflector.Build.cs b/Engine/Source/Developer/SlateReflector/SlateReflector.Build.cs
--- a/Engine/Source/Developer/SlateReflector/SlateReflector.Build.cs
+++ b/Engine/Source/Developer/SlateReflector/SlateReflector.Build.cs
@@ -18,6 +18,7 @@ public class SlateReflector : ModuleRules
 				"InputCore",
 				"Slate",
 				"SlateCore",
+				"EditorStyle",
 				"Json",
 				"AssetRegistry",
 			}
diff --git a/Engine/Source/Editor/BlueprintGraph/Private/K2Node_Variable.cpp b/Engine/Source/Editor/BlueprintGraph/Private/K2Node_Variable.cpp
--- a/Engine/Source/Editor/BlueprintGraph/Private/K2Node_Variable.cpp
+++ b/Engine/Source/Editor/BlueprintGraph/Private/K2Node_Variable.cpp
@@ -97,9 +97,12 @@ bool UK2Node_Variable::CreatePinForVariable(EEdGraphPinDirection Direction, FStr
 	{
 		if (!VariableReference.IsLocalScope())
 		{
-			FString WarningMsg = FString::Printf(TEXT("'%s' variable not found. Base class was probably changed."), *GetVarNameString());
-
 			UBlueprint* OwnerBP = GetBlueprint();
+
+			FString AssetName = OwnerBP ? OwnerBP->GetPathName() : TEXT("Unknown BP");
+
+			FString WarningMsg = FString::Printf(TEXT("'%s' variable not found. Base class was probably changed. %s"), *GetVarNameString(), *AssetName);
+
 			if (OwnerBP && OwnerBP->CurrentMessageLog)
 			{
 				OwnerBP->CurrentMessageLog->Warning(*FString::Printf(TEXT("@@: %s"), *WarningMsg), this);
diff --git a/Engine/Source/Editor/DetailCustomizations/Private/LevelSequenceActorDetails.cpp b/Engine/Source/Editor/DetailCustomizations/Private/LevelSequenceActorDetails.cpp
--- a/Engine/Source/Editor/DetailCustomizations/Private/LevelSequenceActorDetails.cpp
+++ b/Engine/Source/Editor/DetailCustomizations/Private/LevelSequenceActorDetails.cpp
@@ -72,7 +72,11 @@ FReply FLevelSequenceActorDetails::OnOpenLevelSequenceForActor()
 {
 	if( LevelSequenceActor.IsValid() )
 	{
-		FAssetEditorManager::Get().OpenEditorForAsset(LevelSequenceActor.Get()->LevelSequence.TryLoad());
+		UObject* LoadedObject = LevelSequenceActor.Get()->LevelSequence.TryLoad();
+		if (LoadedObject != nullptr)
+		{
+			FAssetEditorManager::Get().OpenEditorForAsset(LoadedObject);
+		}
 	}
 
 	return FReply::Handled();
diff --git a/Engine/Source/Editor/MaterialEditor/Private/MaterialEditor.h b/Engine/Source/Editor/MaterialEditor/Private/MaterialEditor.h
--- a/Engine/Source/Editor/MaterialEditor/Private/MaterialEditor.h
+++ b/Engine/Source/Editor/MaterialEditor/Private/MaterialEditor.h
@@ -332,6 +332,12 @@ public:
 	{
 		return true;
 	}
+
+	virtual bool IsTickableInEditor() const override
+	{
+		return true;
+	}
+
 	virtual TStatId GetStatId() const override;
 
 	/** Pushes the PreviewMesh assigned the the material instance to the thumbnail info */
diff --git a/Engine/Source/Editor/PropertyEditor/Private/Presentation/PropertyEditor/PropertyEditor.cpp b/Engine/Source/Editor/PropertyEditor/Private/Presentation/PropertyEditor/PropertyEditor.cpp
--- a/Engine/Source/Editor/PropertyEditor/Private/Presentation/PropertyEditor/PropertyEditor.cpp
+++ b/Engine/Source/Editor/PropertyEditor/Private/Presentation/PropertyEditor/PropertyEditor.cpp
@@ -659,26 +659,23 @@ bool FPropertyEditor::GetEditConditionPropertyAddress( UBoolProperty*& Condition
 		{
 			for (int32 Index = 0; Index < ComplexParentNode->GetInstancesNum(); ++Index)
 			{
-				TWeakObjectPtr<UObject> Object = ComplexParentNode->GetInstanceAsUObject(Index);
-
-			if( Object.IsValid() )
-			{
-				UObject* Obj = Object.Get();
-
-				// Get the address corresponding to the base of this property (i.e. if a struct property, set BaseOffset to the address of value for the whole struct)
-				uint8* BaseOffset = ParentNode->GetValueAddress((uint8*)Obj);
-				check(BaseOffset != NULL);
-
-				FPropertyConditionInfo NewCondition;
-				// now calculate the address of the property value being used as the condition and add it to the array.
-				NewCondition.Address = EditConditionProperty->ContainerPtrToValuePtr<uint8>(BaseOffset);
-				NewCondition.bNegateValue = bNegate;
-				ConditionPropertyAddresses.Add(NewCondition);
-				bResult = true;
+				uint8* BaseAddress = ComplexParentNode->GetMemoryOfInstance(Index);
+				if (BaseAddress)
+				{
+					// Get the address corresponding to the base of this property (i.e. if a struct property, set BaseOffset to the address of value for the whole struct)
+					uint8* BaseOffset = ParentNode->GetValueAddress(BaseAddress);
+					check(BaseOffset != NULL);
+
+					FPropertyConditionInfo NewCondition;
+					// now calculate the address of the property value being used as the condition and add it to the array.
+					NewCondition.Address = EditConditionProperty->ContainerPtrToValuePtr<uint8>(BaseOffset);
+					NewCondition.bNegateValue = bNegate;
+					ConditionPropertyAddresses.Add(NewCondition);
+					bResult = true;
+				}
 			}
 		}
 	}
-	}
 
 	if ( bResult )
 	{
diff --git a/Engine/Source/Editor/UnrealEd/Private/AnimationEditorUtils.cpp b/Engine/Source/Editor/UnrealEd/Private/AnimationEditorUtils.cpp
--- a/Engine/Source/Editor/UnrealEd/Private/AnimationEditorUtils.cpp
+++ b/Engine/Source/Editor/UnrealEd/Private/AnimationEditorUtils.cpp
@@ -39,6 +39,7 @@
 #include "AnimationStateMachineGraph.h"
 #include "K2Node_Composite.h"
 #include "Kismet2/BlueprintEditorUtils.h"
+#include "Animation/AnimCompress_Automatic.h"
 
 #define LOCTEXT_NAMESPACE "AnimationEditorUtils"
 
@@ -501,6 +502,13 @@ namespace AnimationEditorUtils
 
 					for (UAnimSequence* AnimSeq : AnimSequencePtrs)
 					{
+						// If we are not compressing with 'Auto', then clear CompressCommandletVersion
+						// So we can recompress these animations later.
+						const bool bIsAutoCompressor = Algorithm->IsA(UAnimCompress_Automatic::StaticClass());
+						if (!bIsAutoCompressor)
+						{
+							AnimSeq->CompressCommandletVersion = 0;
+						}
 						AnimSeq->CompressionScheme = static_cast<UAnimCompress*>(StaticDuplicateObject(Algorithm, AnimSeq));
 						AnimSeq->RequestAnimCompression(false, CompressContext);
 						++CompressContext->AnimIndex;
