Copying Orion-Staging @ 3748653 (Orion/Dev-General @ 3722124) to //UE4/Main

diff --git a/Engine/Plugins/Runtime/GameplayAbilities/Source/GameplayAbilities/Private/Abilities/GameplayAbility.cpp b/Engine/Plugins/Runtime/GameplayAbilities/Source/GameplayAbilities/Private/Abilities/GameplayAbility.cpp
--- a/Engine/Plugins/Runtime/GameplayAbilities/Source/GameplayAbilities/Private/Abilities/GameplayAbility.cpp
+++ b/Engine/Plugins/Runtime/GameplayAbilities/Source/GameplayAbilities/Private/Abilities/GameplayAbility.cpp
@@ -945,6 +945,10 @@ FGameplayEffectSpecHandle UGameplayAbility::MakeOutgoingGameplayEffectSpec(const
 	{
 		FGameplayAbilitySpec* AbilitySpec =  ActorInfo->AbilitySystemComponent->FindAbilitySpecFromHandle(Handle);
 		ApplyAbilityTagsToGameplayEffectSpec(*NewHandle.Data.Get(), AbilitySpec);
+
+		// Copy over set by caller magnitudes
+		NewHandle.Data->SetByCallerTagMagnitudes = AbilitySpec->SetByCallerTagMagnitudes;
+
 	}
 	return NewHandle;
 }
@@ -964,6 +968,9 @@ void UGameplayAbility::ApplyAbilityTagsToGameplayEffectSpec(FGameplayEffectSpec&
 
 			Spec.CapturedSourceTags.GetSpecTags().AppendTags(SourceObjTags);
 		}
+
+		// Copy SetByCallerMagnitudes 
+		Spec.MergeSetByCallerMagnitudes(AbilitySpec->SetByCallerTagMagnitudes);
 	}
 }
 
@@ -1420,6 +1427,12 @@ FGameplayEffectContextHandle UGameplayAbility::MakeEffectContext(const FGameplay
 	// add in the ability tracking here.
 	Context.SetAbility(this);
 
+	// Pass along the source object to the effect
+	if (FGameplayAbilitySpec* AbilitySpec = ActorInfo->AbilitySystemComponent->FindAbilitySpecFromHandle(Handle))
+	{
+		Context.AddSourceObject(AbilitySpec->SourceObject);
+	}
+
 	return Context;
 }
 
diff --git a/Engine/Plugins/Runtime/GameplayAbilities/Source/GameplayAbilities/Private/Abilities/Tasks/AbilityTask_WaitAttributeChange.cpp b/Engine/Plugins/Runtime/GameplayAbilities/Source/GameplayAbilities/Private/Abilities/Tasks/AbilityTask_WaitAttributeChange.cpp
--- a/Engine/Plugins/Runtime/GameplayAbilities/Source/GameplayAbilities/Private/Abilities/Tasks/AbilityTask_WaitAttributeChange.cpp
+++ b/Engine/Plugins/Runtime/GameplayAbilities/Source/GameplayAbilities/Private/Abilities/Tasks/AbilityTask_WaitAttributeChange.cpp
@@ -5,6 +5,7 @@
 #include "AbilitySystemComponent.h"
 
 #include "GameplayEffectExtension.h"
+#include "AbilitySystemGlobals.h"
 
 UAbilityTask_WaitAttributeChange::UAbilityTask_WaitAttributeChange(const FObjectInitializer& ObjectInitializer)
 	: Super(ObjectInitializer)
@@ -12,7 +13,7 @@ UAbilityTask_WaitAttributeChange::UAbilityTask_WaitAttributeChange(const FObject
 	bTriggerOnce = false;
 }
 
-UAbilityTask_WaitAttributeChange* UAbilityTask_WaitAttributeChange::WaitForAttributeChange(UGameplayAbility* OwningAbility, FGameplayAttribute InAttribute, FGameplayTag InWithTag, FGameplayTag InWithoutTag, bool TriggerOnce)
+UAbilityTask_WaitAttributeChange* UAbilityTask_WaitAttributeChange::WaitForAttributeChange(UGameplayAbility* OwningAbility, FGameplayAttribute InAttribute, FGameplayTag InWithTag, FGameplayTag InWithoutTag, bool TriggerOnce, AActor* OptionalExternalOwner)
 {
 	UAbilityTask_WaitAttributeChange* MyObj = NewAbilityTask<UAbilityTask_WaitAttributeChange>(OwningAbility);
 	MyObj->WithTag = InWithTag;
@@ -20,11 +21,12 @@ UAbilityTask_WaitAttributeChange* UAbilityTask_WaitAttributeChange::WaitForAttri
 	MyObj->Attribute = InAttribute;
 	MyObj->ComparisonType = EWaitAttributeChangeComparison::None;
 	MyObj->bTriggerOnce = TriggerOnce;
+	MyObj->ExternalOwner = OptionalExternalOwner ? UAbilitySystemGlobals::GetAbilitySystemComponentFromActor(OptionalExternalOwner) : nullptr;
 
 	return MyObj;
 }
 
-UAbilityTask_WaitAttributeChange* UAbilityTask_WaitAttributeChange::WaitForAttributeChangeWithComparison(UGameplayAbility* OwningAbility, FGameplayAttribute InAttribute, FGameplayTag InWithTag, FGameplayTag InWithoutTag, TEnumAsByte<EWaitAttributeChangeComparison::Type> InComparisonType, float InComparisonValue, bool TriggerOnce)
+UAbilityTask_WaitAttributeChange* UAbilityTask_WaitAttributeChange::WaitForAttributeChangeWithComparison(UGameplayAbility* OwningAbility, FGameplayAttribute InAttribute, FGameplayTag InWithTag, FGameplayTag InWithoutTag, TEnumAsByte<EWaitAttributeChangeComparison::Type> InComparisonType, float InComparisonValue, bool TriggerOnce, AActor* OptionalExternalOwner)
 {
 	UAbilityTask_WaitAttributeChange* MyObj = NewAbilityTask<UAbilityTask_WaitAttributeChange>(OwningAbility);
 	MyObj->WithTag = InWithTag;
@@ -33,15 +35,16 @@ UAbilityTask_WaitAttributeChange* UAbilityTask_WaitAttributeChange::WaitForAttri
 	MyObj->ComparisonType = InComparisonType;
 	MyObj->ComparisonValue = InComparisonValue;
 	MyObj->bTriggerOnce = TriggerOnce;
+	MyObj->ExternalOwner = OptionalExternalOwner ? UAbilitySystemGlobals::GetAbilitySystemComponentFromActor(OptionalExternalOwner) : nullptr;
 
 	return MyObj;
 }
 
 void UAbilityTask_WaitAttributeChange::Activate()
 {
-	if (AbilitySystemComponent)
+	if (UAbilitySystemComponent* ASC = GetFocusedASC())
 	{
-		OnAttributeChangeDelegateHandle = AbilitySystemComponent->GetGameplayAttributeValueChangeDelegate(Attribute).AddUObject(this, &UAbilityTask_WaitAttributeChange::OnAttributeChange);
+		OnAttributeChangeDelegateHandle = ASC->GetGameplayAttributeValueChangeDelegate(Attribute).AddUObject(this, &UAbilityTask_WaitAttributeChange::OnAttributeChange);
 	}
 }
 
@@ -106,11 +109,16 @@ void UAbilityTask_WaitAttributeChange::OnAttributeChange(const FOnAttributeChang
 	}
 }
 
+UAbilitySystemComponent* UAbilityTask_WaitAttributeChange::GetFocusedASC()
+{
+	return ExternalOwner ? ExternalOwner : AbilitySystemComponent;
+}
+
 void UAbilityTask_WaitAttributeChange::OnDestroy(bool AbilityEnded)
 {
-	if (AbilitySystemComponent)
+	if (UAbilitySystemComponent* ASC = GetFocusedASC())
 	{
-		AbilitySystemComponent->GetGameplayAttributeValueChangeDelegate(Attribute).Remove(OnAttributeChangeDelegateHandle);
+		ASC->GetGameplayAttributeValueChangeDelegate(Attribute).Remove(OnAttributeChangeDelegateHandle);
 	}
 
 	Super::OnDestroy(AbilityEnded);
diff --git a/Engine/Plugins/Runtime/GameplayAbilities/Source/GameplayAbilities/Private/Abilities/Tasks/AbilityTask_WaitAttributeChangeRatioThreshold.cpp b/Engine/Plugins/Runtime/GameplayAbilities/Source/GameplayAbilities/Private/Abilities/Tasks/AbilityTask_WaitAttributeChangeRatioThreshold.cpp
--- a/Engine/Plugins/Runtime/GameplayAbilities/Source/GameplayAbilities/Private/Abilities/Tasks/AbilityTask_WaitAttributeChangeRatioThreshold.cpp
+++ b/Engine/Plugins/Runtime/GameplayAbilities/Source/GameplayAbilities/Private/Abilities/Tasks/AbilityTask_WaitAttributeChangeRatioThreshold.cpp
@@ -4,6 +4,7 @@
 #include "TimerManager.h"
 
 #include "AbilitySystemComponent.h"
+#include "AbilitySystemGlobals.h"
 
 
 UAbilityTask_WaitAttributeChangeRatioThreshold::UAbilityTask_WaitAttributeChangeRatioThreshold(const FObjectInitializer& ObjectInitializer)
@@ -15,7 +16,7 @@ UAbilityTask_WaitAttributeChangeRatioThreshold::UAbilityTask_WaitAttributeChange
 	LastAttributeDenominatorValue = 1.f;
 }
 
-UAbilityTask_WaitAttributeChangeRatioThreshold* UAbilityTask_WaitAttributeChangeRatioThreshold::WaitForAttributeChangeRatioThreshold(UGameplayAbility* OwningAbility, FGameplayAttribute AttributeNumerator, FGameplayAttribute AttributeDenominator, TEnumAsByte<EWaitAttributeChangeComparison::Type> ComparisonType, float ComparisonValue, bool bTriggerOnce)
+UAbilityTask_WaitAttributeChangeRatioThreshold* UAbilityTask_WaitAttributeChangeRatioThreshold::WaitForAttributeChangeRatioThreshold(UGameplayAbility* OwningAbility, FGameplayAttribute AttributeNumerator, FGameplayAttribute AttributeDenominator, TEnumAsByte<EWaitAttributeChangeComparison::Type> ComparisonType, float ComparisonValue, bool bTriggerOnce, AActor* OptionalExternalOwner)
 {
 	UAbilityTask_WaitAttributeChangeRatioThreshold* MyTask = NewAbilityTask<UAbilityTask_WaitAttributeChangeRatioThreshold>(OwningAbility);
 	MyTask->AttributeNumerator = AttributeNumerator;
@@ -23,16 +24,17 @@ UAbilityTask_WaitAttributeChangeRatioThreshold* UAbilityTask_WaitAttributeChange
 	MyTask->ComparisonType = ComparisonType;
 	MyTask->ComparisonValue = ComparisonValue;
 	MyTask->bTriggerOnce = bTriggerOnce;
+	MyTask->ExternalOwner = OptionalExternalOwner ? UAbilitySystemGlobals::GetAbilitySystemComponentFromActor(OptionalExternalOwner) : nullptr;
 
 	return MyTask;
 }
 
 void UAbilityTask_WaitAttributeChangeRatioThreshold::Activate()
 {
-	if (AbilitySystemComponent)
+	if (UAbilitySystemComponent* ASC = GetFocusedASC())
 	{
-		LastAttributeNumeratorValue = AbilitySystemComponent->GetNumericAttribute(AttributeNumerator);
-		LastAttributeDenominatorValue = AbilitySystemComponent->GetNumericAttribute(AttributeDenominator);
+		LastAttributeNumeratorValue = ASC->GetNumericAttribute(AttributeNumerator);
+		LastAttributeDenominatorValue = ASC->GetNumericAttribute(AttributeDenominator);
 		bMatchedComparisonLastAttributeChange = DoesValuePassComparison(LastAttributeNumeratorValue, LastAttributeDenominatorValue);
 
 		// Broadcast OnChange immediately with current value
@@ -41,11 +43,16 @@ void UAbilityTask_WaitAttributeChangeRatioThreshold::Activate()
 			OnChange.Broadcast(bMatchedComparisonLastAttributeChange, LastAttributeDenominatorValue != 0.f ? LastAttributeNumeratorValue/LastAttributeDenominatorValue : 0.f);
 		}
 
-		OnNumeratorAttributeChangeDelegateHandle = AbilitySystemComponent->GetGameplayAttributeValueChangeDelegate(AttributeNumerator).AddUObject(this, &UAbilityTask_WaitAttributeChangeRatioThreshold::OnNumeratorAttributeChange);
-		OnDenominatorAttributeChangeDelegateHandle = AbilitySystemComponent->GetGameplayAttributeValueChangeDelegate(AttributeDenominator).AddUObject(this, &UAbilityTask_WaitAttributeChangeRatioThreshold::OnDenominatorAttributeChange);
+		OnNumeratorAttributeChangeDelegateHandle = ASC->GetGameplayAttributeValueChangeDelegate(AttributeNumerator).AddUObject(this, &UAbilityTask_WaitAttributeChangeRatioThreshold::OnNumeratorAttributeChange);
+		OnDenominatorAttributeChangeDelegateHandle = ASC->GetGameplayAttributeValueChangeDelegate(AttributeDenominator).AddUObject(this, &UAbilityTask_WaitAttributeChangeRatioThreshold::OnDenominatorAttributeChange);
 	}
 }
 
+UAbilitySystemComponent* UAbilityTask_WaitAttributeChangeRatioThreshold::GetFocusedASC()
+{
+	return ExternalOwner ? ExternalOwner : AbilitySystemComponent;
+}
+
 void UAbilityTask_WaitAttributeChangeRatioThreshold::OnAttributeChange()
 {
 	UWorld* World = GetWorld();
@@ -125,10 +132,10 @@ bool UAbilityTask_WaitAttributeChangeRatioThreshold::DoesValuePassComparison(flo
 
 void UAbilityTask_WaitAttributeChangeRatioThreshold::OnDestroy(bool AbilityEnded)
 {
-	if (AbilitySystemComponent)
+	if (UAbilitySystemComponent* ASC = GetFocusedASC())
 	{
-		AbilitySystemComponent->GetGameplayAttributeValueChangeDelegate(AttributeNumerator).Remove(OnNumeratorAttributeChangeDelegateHandle);
-		AbilitySystemComponent->GetGameplayAttributeValueChangeDelegate(AttributeDenominator).Remove(OnDenominatorAttributeChangeDelegateHandle);
+		ASC->GetGameplayAttributeValueChangeDelegate(AttributeNumerator).Remove(OnNumeratorAttributeChangeDelegateHandle);
+		ASC->GetGameplayAttributeValueChangeDelegate(AttributeDenominator).Remove(OnDenominatorAttributeChangeDelegateHandle);
 	}
 
 	Super::OnDestroy(AbilityEnded);
diff --git a/Engine/Plugins/Runtime/GameplayAbilities/Source/GameplayAbilities/Private/Abilities/Tasks/AbilityTask_WaitAttributeChangeThreshold.cpp b/Engine/Plugins/Runtime/GameplayAbilities/Source/GameplayAbilities/Private/Abilities/Tasks/AbilityTask_WaitAttributeChangeThreshold.cpp
--- a/Engine/Plugins/Runtime/GameplayAbilities/Source/GameplayAbilities/Private/Abilities/Tasks/AbilityTask_WaitAttributeChangeThreshold.cpp
+++ b/Engine/Plugins/Runtime/GameplayAbilities/Source/GameplayAbilities/Private/Abilities/Tasks/AbilityTask_WaitAttributeChangeThreshold.cpp
@@ -3,6 +3,7 @@
 #include "Abilities/Tasks/AbilityTask_WaitAttributeChangeThreshold.h"
 
 #include "AbilitySystemComponent.h"
+#include "AbilitySystemGlobals.h"
 
 
 UAbilityTask_WaitAttributeChangeThreshold::UAbilityTask_WaitAttributeChangeThreshold(const FObjectInitializer& ObjectInitializer)
@@ -12,22 +13,23 @@ UAbilityTask_WaitAttributeChangeThreshold::UAbilityTask_WaitAttributeChangeThres
 	bMatchedComparisonLastAttributeChange = false;
 }
 
-UAbilityTask_WaitAttributeChangeThreshold* UAbilityTask_WaitAttributeChangeThreshold::WaitForAttributeChangeThreshold(UGameplayAbility* OwningAbility, FGameplayAttribute Attribute, TEnumAsByte<EWaitAttributeChangeComparison::Type> ComparisonType, float ComparisonValue, bool bTriggerOnce)
+UAbilityTask_WaitAttributeChangeThreshold* UAbilityTask_WaitAttributeChangeThreshold::WaitForAttributeChangeThreshold(UGameplayAbility* OwningAbility, FGameplayAttribute Attribute, TEnumAsByte<EWaitAttributeChangeComparison::Type> ComparisonType, float ComparisonValue, bool bTriggerOnce, AActor* OptionalExternalOwner)
 {
 	UAbilityTask_WaitAttributeChangeThreshold* MyTask = NewAbilityTask<UAbilityTask_WaitAttributeChangeThreshold>(OwningAbility);
 	MyTask->Attribute = Attribute;
 	MyTask->ComparisonType = ComparisonType;
 	MyTask->ComparisonValue = ComparisonValue;
 	MyTask->bTriggerOnce = bTriggerOnce;
+	MyTask->ExternalOwner = OptionalExternalOwner ? UAbilitySystemGlobals::GetAbilitySystemComponentFromActor(OptionalExternalOwner) : nullptr;
 
 	return MyTask;
 }
 
 void UAbilityTask_WaitAttributeChangeThreshold::Activate()
 {
-	if (AbilitySystemComponent)
+	if (UAbilitySystemComponent* ASC = GetFocusedASC())
 	{
-		const float CurrentValue = AbilitySystemComponent->GetNumericAttribute(Attribute);
+		const float CurrentValue = ASC->GetNumericAttribute(Attribute);
 		bMatchedComparisonLastAttributeChange = DoesValuePassComparison(CurrentValue);
 
 		// Broadcast OnChange immediately with current value
@@ -36,7 +38,7 @@ void UAbilityTask_WaitAttributeChangeThreshold::Activate()
 			OnChange.Broadcast(bMatchedComparisonLastAttributeChange, CurrentValue);
 		}
 
-		OnAttributeChangeDelegateHandle = AbilitySystemComponent->GetGameplayAttributeValueChangeDelegate(Attribute).AddUObject(this, &UAbilityTask_WaitAttributeChangeThreshold::OnAttributeChange);
+		OnAttributeChangeDelegateHandle = ASC->GetGameplayAttributeValueChangeDelegate(Attribute).AddUObject(this, &UAbilityTask_WaitAttributeChangeThreshold::OnAttributeChange);
 	}
 }
 
@@ -88,11 +90,16 @@ bool UAbilityTask_WaitAttributeChangeThreshold::DoesValuePassComparison(float Va
 	return bPassedComparison;
 }
 
+UAbilitySystemComponent* UAbilityTask_WaitAttributeChangeThreshold::GetFocusedASC()
+{
+	return ExternalOwner ? ExternalOwner : AbilitySystemComponent;
+}
+
 void UAbilityTask_WaitAttributeChangeThreshold::OnDestroy(bool AbilityEnded)
 {
-	if (AbilitySystemComponent)
+	if (UAbilitySystemComponent* ASC = GetFocusedASC())
 	{
-		AbilitySystemComponent->GetGameplayAttributeValueChangeDelegate(Attribute).Remove(OnAttributeChangeDelegateHandle);
+		ASC->GetGameplayAttributeValueChangeDelegate(Attribute).Remove(OnAttributeChangeDelegateHandle);
 	}
 
 	Super::OnDestroy(AbilityEnded);
diff --git a/Engine/Plugins/Runtime/GameplayAbilities/Source/GameplayAbilities/Private/AbilitySystemComponent.cpp b/Engine/Plugins/Runtime/GameplayAbilities/Source/GameplayAbilities/Private/AbilitySystemComponent.cpp
--- a/Engine/Plugins/Runtime/GameplayAbilities/Source/GameplayAbilities/Private/AbilitySystemComponent.cpp
+++ b/Engine/Plugins/Runtime/GameplayAbilities/Source/GameplayAbilities/Private/AbilitySystemComponent.cpp
@@ -369,6 +369,11 @@ FActiveGameplayEffectHandle UAbilitySystemComponent::BP_ApplyGameplayEffectToSel
 {
 	if ( GameplayEffectClass )
 	{
+		if (!EffectContext.IsValid())
+		{
+			EffectContext = MakeEffectContext();
+		}
+
 		UGameplayEffect* GameplayEffect = GameplayEffectClass->GetDefaultObject<UGameplayEffect>();
 		return ApplyGameplayEffectToSelf(GameplayEffect, Level, EffectContext);
 	}
@@ -397,50 +402,36 @@ FActiveGameplayEffectHandle UAbilitySystemComponent::ApplyGameplayEffectToSelf(c
 FOnActiveGameplayEffectRemoved* UAbilitySystemComponent::OnGameplayEffectRemovedDelegate(FActiveGameplayEffectHandle Handle)
 {
 	FActiveGameplayEffect* ActiveEffect = ActiveGameplayEffects.GetActiveGameplayEffect(Handle);
-	if (ActiveEffect)
-	{
-		return &ActiveEffect->OnRemovedDelegate;
-	}
-
-	return nullptr;
+	return ActiveEffect ? &ActiveEffect->EventSet.DEPRECATED_OnEffectRemoved : nullptr;
 }
 
-FOnActiveGameplayEffectRemoved_Info* UAbilitySystemComponent::OnGameplayEffectRemoved_InfoDelegate(FActiveGameplayEffectHandle Handle)
+FActiveGameplayEffectEvents* UAbilitySystemComponent::GetActiveEffectEventSet(FActiveGameplayEffectHandle Handle)
 {
 	FActiveGameplayEffect* ActiveEffect = ActiveGameplayEffects.GetActiveGameplayEffect(Handle);
-	if (ActiveEffect)
-	{
-		return &ActiveEffect->OnRemoved_InfoDelegate;
-	}
-
-	return nullptr;
+	return ActiveEffect ? &ActiveEffect->EventSet : nullptr;
 }
 
-FOnGivenActiveGameplayEffectRemoved& UAbilitySystemComponent::OnAnyGameplayEffectRemovedDelegate()
+FOnActiveGameplayEffectRemoved_Info* UAbilitySystemComponent::OnGameplayEffectRemoved_InfoDelegate(FActiveGameplayEffectHandle Handle)
 {
-	return ActiveGameplayEffects.OnActiveGameplayEffectRemovedDelegate;
+	FActiveGameplayEffect* ActiveEffect = ActiveGameplayEffects.GetActiveGameplayEffect(Handle);
+	return ActiveEffect ? &ActiveEffect->EventSet.OnEffectRemoved : nullptr;
 }
 
 FOnActiveGameplayEffectStackChange* UAbilitySystemComponent::OnGameplayEffectStackChangeDelegate(FActiveGameplayEffectHandle Handle)
 {
 	FActiveGameplayEffect* ActiveEffect = ActiveGameplayEffects.GetActiveGameplayEffect(Handle);
-	if (ActiveEffect)
-	{
-		return &ActiveEffect->OnStackChangeDelegate;
-	}
-
-	return nullptr;
+	return ActiveEffect ? &ActiveEffect->EventSet.OnStackChanged : nullptr;
 }
 
 FOnActiveGameplayEffectTimeChange* UAbilitySystemComponent::OnGameplayEffectTimeChangeDelegate(FActiveGameplayEffectHandle Handle)
 {
 	FActiveGameplayEffect* ActiveEffect = ActiveGameplayEffects.GetActiveGameplayEffect(Handle);
-	if (ActiveEffect)
-	{
-		return &ActiveEffect->OnTimeChangeDelegate;
-	}
+	return ActiveEffect ? &ActiveEffect->EventSet.OnTimeChanged : nullptr;
+}
 
-	return nullptr;
+FOnGivenActiveGameplayEffectRemoved& UAbilitySystemComponent::OnAnyGameplayEffectRemovedDelegate()
+{
+	return ActiveGameplayEffects.OnActiveGameplayEffectRemovedDelegate;
 }
 
 FOnGameplayEffectTagCountChanged& UAbilitySystemComponent::RegisterGameplayTagEvent(FGameplayTag Tag, EGameplayTagEventType::Type EventType)
@@ -1571,11 +1562,11 @@ void UAbilitySystemComponent::OnRep_ServerDebugString()
 	}
 }
 
-float UAbilitySystemComponent::GetFilteredAttributeValue(const FGameplayAttribute& Attribute, const FGameplayTagRequirements& SourceTags, const FGameplayTagContainer& TargetTags)
+float UAbilitySystemComponent::GetFilteredAttributeValue(const FGameplayAttribute& Attribute, const FGameplayTagRequirements& SourceTags, const FGameplayTagContainer& TargetTags, const TArray<FActiveGameplayEffectHandle>& HandlesToIgnore)
 {
 	float AttributeValue = 0.f;
 
-	if (SourceTags.RequireTags.Num() == 0 && SourceTags.IgnoreTags.Num() == 0)
+	if (SourceTags.RequireTags.Num() == 0 && SourceTags.IgnoreTags.Num() == 0 && HandlesToIgnore.Num() == 0)
 	{
 		// No qualifiers so we can just read this attribute normally
 		AttributeValue = GetNumericAttribute(Attribute);
@@ -1605,6 +1596,7 @@ float UAbilitySystemComponent::GetFilteredAttributeValue(const FGameplayAttribut
 		Params.SourceTags = &QuerySourceTags;
 		Params.TargetTags = &QueryTargetTags;
 		Params.IncludePredictiveMods = true;
+		Params.IgnoreHandles = HandlesToIgnore;
 
 		if (CaptureSpec.AttemptCalculateAttributeMagnitude(Params, AttributeValue) == false)
 		{
@@ -2054,7 +2046,10 @@ void UAbilitySystemComponent::Debug_Internal(FAbilitySystemComponentDebugInfo& I
 			{
 				FAggregator& Aggregator = *AggregatorRef.Get();
 
+				FAggregatorEvaluateParameters EmptyParams;
+
 				TMap<EGameplayModEvaluationChannel, const TArray<FAggregatorMod>*> ModMap;
+				Aggregator.EvaluateQualificationForAllMods(EmptyParams);
 				Aggregator.GetAllAggregatorMods(ModMap);
 
 				if (ModMap.Num() == 0)
@@ -2091,8 +2086,7 @@ void UAbilitySystemComponent::Debug_Internal(FAbilitySystemComponentDebugInfo& I
 						const TArray<FAggregatorMod>& CurModArray = ModArrays[ModOpIdx];
 						for (const FAggregatorMod& Mod : CurModArray)
 						{
-							FAggregatorEvaluateParameters EmptyParams;
-							bool IsActivelyModifyingAttribute = Mod.Qualifies(EmptyParams);
+							bool IsActivelyModifyingAttribute = Mod.Qualifies();
 							if (Info.Canvas)
 							{
 								Info.Canvas->SetDrawColor(IsActivelyModifyingAttribute ? FColor::Yellow : FColor(128, 128, 128));
@@ -2202,20 +2196,6 @@ void UAbilitySystemComponent::Debug_Internal(FAbilitySystemComponentDebugInfo& I
 				const FModifierSpec& ModSpec = ActiveGE.Spec.Modifiers[ModIdx];
 				const FGameplayModifierInfo& ModInfo = ActiveGE.Spec.Def->Modifiers[ModIdx];
 
-				// Do a quick Qualifies() check to see if this mod is active.
-				FAggregatorMod TempMod;
-				TempMod.SourceTagReqs = &ModInfo.SourceTags;
-				TempMod.TargetTagReqs = &ModInfo.TargetTags;
-				TempMod.IsPredicted = false;
-
-				FAggregatorEvaluateParameters EmptyParams;
-				bool IsActivelyModifyingAttribute = TempMod.Qualifies(EmptyParams);
-
-				if (IsActivelyModifyingAttribute == false)
-				{
-					if (Info.Canvas) Info.Canvas->SetDrawColor(FColor(128, 128, 128));
-				}
-
 				DebugLine(Info, FString::Printf(TEXT("Mod: %s. %s. %.2f"), *ModInfo.Attribute.GetName(), *EGameplayModOpToString(ModInfo.ModifierOp), ModSpec.GetEvaluatedMagnitude()), 7.f, 0.f);
 
 				if (Info.Canvas) Info.Canvas->SetDrawColor(ActiveGE.bIsInhibited ? FColor(128, 128, 128) : FColor::White);
diff --git a/Engine/Plugins/Runtime/GameplayAbilities/Source/GameplayAbilities/Private/GameplayAbilityTypes.cpp b/Engine/Plugins/Runtime/GameplayAbilities/Source/GameplayAbilities/Private/GameplayAbilityTypes.cpp
--- a/Engine/Plugins/Runtime/GameplayAbilities/Source/GameplayAbilities/Private/GameplayAbilityTypes.cpp
+++ b/Engine/Plugins/Runtime/GameplayAbilities/Source/GameplayAbilities/Private/GameplayAbilityTypes.cpp
@@ -209,6 +209,8 @@ FGameplayAbilitySpec::FGameplayAbilitySpec(FGameplayAbilitySpecDef& InDef, int32
 	InDef.AssignedHandle = Handle;
 	GameplayEffectHandle = InGameplayEffectHandle;
 
+	SetByCallerTagMagnitudes = InDef.SetByCallerTagMagnitudes;
+
 	FString ContextString = FString::Printf(TEXT("FGameplayAbilitySpec::FGameplayAbilitySpec for %s from %s"), 
 		(InDef.Ability ? *InDef.Ability->GetName() : TEXT("INVALID ABILITY")), 
 		(InDef.SourceObject ? *InDef.SourceObject->GetName() : TEXT("INVALID ABILITY")));
diff --git a/Engine/Plugins/Runtime/GameplayAbilities/Source/GameplayAbilities/Private/GameplayCueManager.cpp b/Engine/Plugins/Runtime/GameplayAbilities/Source/GameplayAbilities/Private/GameplayCueManager.cpp
--- a/Engine/Plugins/Runtime/GameplayAbilities/Source/GameplayAbilities/Private/GameplayCueManager.cpp
+++ b/Engine/Plugins/Runtime/GameplayAbilities/Source/GameplayAbilities/Private/GameplayCueManager.cpp
@@ -693,8 +693,10 @@ static void SearchDynamicClassCues(const FName PropertyName, const TArray<FStrin
 	}
 }
 
-void UGameplayCueManager::InitObjectLibrary(FGameplayCueObjectLibrary& Lib)
+TSharedPtr<FStreamableHandle> UGameplayCueManager::InitObjectLibrary(FGameplayCueObjectLibrary& Lib)
 {
+	TSharedPtr<FStreamableHandle> RetVal;
+
 	DECLARE_SCOPE_CYCLE_COUNTER(TEXT("Loading Library"), STAT_ObjectLibrary, STATGROUP_LoadTime);
 
 	// Instantiate the UObjectLibraries if they aren't there already
@@ -796,6 +798,7 @@ void UGameplayCueManager::InitObjectLibrary(FGameplayCueObjectLibrary& Lib)
 		if (AssetsToLoad.Num() > 0)
 		{
 			GameplayCueAssetHandle = StreamableManager.RequestAsyncLoad(AssetsToLoad, FStreamableDelegate::CreateStatic( ForwardLambda, AssetsToLoad, Lib.OnLoaded), Lib.AsyncPriority);
+			RetVal = GameplayCueAssetHandle;
 		}
 		else
 		{
@@ -806,6 +809,7 @@ void UGameplayCueManager::InitObjectLibrary(FGameplayCueObjectLibrary& Lib)
 
 	// Build Tag Translation table
 	TranslationManager.BuildTagTranslationTable();
+	return RetVal;
 }
 
 static FAutoConsoleVariable CVarGameplyCueAddToGlobalSetDebug(TEXT("GameplayCue.AddToGlobalSet.DebugTag"), TEXT(""), TEXT("Debug Tag adding to global set"), ECVF_Default	);
@@ -1342,6 +1346,8 @@ void UGameplayCueManager::EndGameplayCueSendContext()
 
 void UGameplayCueManager::FlushPendingCues()
 {
+	OnFlushPendingCues.Broadcast();
+
 	TArray<FGameplayCuePendingExecute> LocalPendingExecuteCues = PendingExecuteCues;
 	PendingExecuteCues.Empty();
 	for (int32 i = 0; i < LocalPendingExecuteCues.Num(); i++)
diff --git a/Engine/Plugins/Runtime/GameplayAbilities/Source/GameplayAbilities/Private/GameplayCueTranslator.cpp b/Engine/Plugins/Runtime/GameplayAbilities/Source/GameplayAbilities/Private/GameplayCueTranslator.cpp
--- a/Engine/Plugins/Runtime/GameplayAbilities/Source/GameplayAbilities/Private/GameplayCueTranslator.cpp
+++ b/Engine/Plugins/Runtime/GameplayAbilities/Source/GameplayAbilities/Private/GameplayCueTranslator.cpp
@@ -632,4 +632,54 @@ bool FGameplayCueTranslationManager::GetTranslatedTags(const FName& ParentTag, T
 	}
 	return Children.Num() > 0;
 }
+
+
+FGameplayTag FGameplayCueTranslationManager::SearchSlowForTranslationParent(FGameplayTag Tag)
+{
+	TFunction<FGameplayTag(int32)> FindParent_r = [&](int32 FindIndex) -> FGameplayTag
+	{
+		for (int32 idx=0; idx < TranslationLUT.Num(); ++idx)
+		{
+			for (FGameplayCueTranslationLink& Link : TranslationLUT[idx].Links)
+			{
+				for (FGameplayCueTranslatorNodeIndex& LinkIdx : Link.NodeLookup)
+				{
+					if (LinkIdx.Index == FindIndex)
+					{
+						FGameplayCueTranslatorNode& Node = TranslationLUT[idx];
+						//UE_LOG(LogGameplayCueTranslator, Display, TEXT("%s links to %s"), *Node.CachedGameplayTag.ToString(), *TranslationLUT[FindIndex].CachedGameplayTag.ToString());
+
+						FGameplayTag ParentTag = FindParent_r(idx);
+						if (ParentTag.IsValid())
+						{
+							// Found a higher level parent
+							return ParentTag;
+						}
+						else
+						{
+							// Topmost tag
+							return Node.CachedGameplayTag;
+						}
+					}
+				}
+			}
+		}
+
+		// Found nothing
+		return FGameplayTag();
+
+	};
+
+	for (int32 RootIdx=0; RootIdx < TranslationLUT.Num(); ++RootIdx)
+	{
+		FGameplayCueTranslatorNode& Node = TranslationLUT[RootIdx];
+		if (Node.CachedGameplayTag == Tag)
+		{
+			return FindParent_r(RootIdx);
+		}
+	}
+
+	// No match, return invalid tag
+	return FGameplayTag();
+}
 #endif
diff --git a/Engine/Plugins/Runtime/GameplayAbilities/Source/GameplayAbilities/Private/GameplayEffect.cpp b/Engine/Plugins/Runtime/GameplayAbilities/Source/GameplayAbilities/Private/GameplayEffect.cpp
--- a/Engine/Plugins/Runtime/GameplayAbilities/Source/GameplayAbilities/Private/GameplayEffect.cpp
+++ b/Engine/Plugins/Runtime/GameplayAbilities/Source/GameplayAbilities/Private/GameplayEffect.cpp
@@ -834,6 +834,26 @@ void FGameplayEffectSpec::InitializeFromLinkedSpec(const UGameplayEffect* InDef,
 
 	// Now initialize like the normal cstor would have. Note that this will add the new GE's asset tags (in case they were removed in the line above / e.g., shared asset tags with the originating GE)					
 	Initialize(InDef, NewContextHandle, OriginalSpec.GetLevel());
+
+	// Finally, copy over set by called magnitudes
+	CopySetByCallerMagnitudes(OriginalSpec);
+}
+
+void FGameplayEffectSpec::CopySetByCallerMagnitudes(const FGameplayEffectSpec& OriginalSpec)
+{
+	SetByCallerNameMagnitudes = OriginalSpec.SetByCallerNameMagnitudes;
+	SetByCallerTagMagnitudes = OriginalSpec.SetByCallerTagMagnitudes;
+}
+
+void FGameplayEffectSpec::MergeSetByCallerMagnitudes(const TMap<FGameplayTag, float>& Magnitudes)
+{
+	for (auto It : Magnitudes)
+	{
+		if (SetByCallerTagMagnitudes.Contains(It.Key) == false)
+		{
+			SetByCallerTagMagnitudes.Add(It.Key) = It.Value;
+		}
+	}
 }
 
 void FGameplayEffectSpec::SetupAttributeCaptureDefinitions()
@@ -1323,11 +1343,12 @@ bool FGameplayEffectAttributeCaptureSpec::AttemptAddAggregatorModsToAggregator(O
 	return false;
 }
 
-bool FGameplayEffectAttributeCaptureSpec::AttemptGatherAttributeMods(OUT TMap<EGameplayModEvaluationChannel, const TArray<FAggregatorMod>*>& OutModMap) const
+bool FGameplayEffectAttributeCaptureSpec::AttemptGatherAttributeMods(const FAggregatorEvaluateParameters& InEvalParams, OUT TMap<EGameplayModEvaluationChannel, const TArray<FAggregatorMod>*>& OutModMap) const
 {
 	FAggregator* Agg = AttributeAggregator.Get();
 	if (Agg)
 	{
+		Agg->EvaluateQualificationForAllMods(InEvalParams);
 		Agg->GetAllAggregatorMods(OutModMap);
 		return true;
 	}
@@ -1568,21 +1589,20 @@ FActiveGameplayEffect::FActiveGameplayEffect(FActiveGameplayEffectHandle InHandl
 }
 
 FActiveGameplayEffect::FActiveGameplayEffect(FActiveGameplayEffect&& Other)
-	:Handle(Other.Handle)
-	,Spec(MoveTemp(Other.Spec))
-	,PredictionKey(Other.PredictionKey)
-	,StartServerWorldTime(Other.StartServerWorldTime)
-	,CachedStartServerWorldTime(Other.CachedStartServerWorldTime)
-	,StartWorldTime(Other.StartWorldTime)
-	,bIsInhibited(Other.bIsInhibited)
-	,bPendingRepOnActiveGC(Other.bPendingRepOnActiveGC)
-	,bPendingRepWhileActiveGC(Other.bPendingRepWhileActiveGC)
-	,IsPendingRemove(Other.IsPendingRemove)
-	,ClientCachedStackCount(0)
-	,OnRemovedDelegate(Other.OnRemovedDelegate)
-	,OnRemoved_InfoDelegate(Other.OnRemoved_InfoDelegate)
-	,PeriodHandle(Other.PeriodHandle)
-	,DurationHandle(Other.DurationHandle)
+	: Handle(Other.Handle)
+	, Spec(MoveTemp(Other.Spec))
+	, PredictionKey(Other.PredictionKey)
+	, StartServerWorldTime(Other.StartServerWorldTime)
+	, CachedStartServerWorldTime(Other.CachedStartServerWorldTime)
+	, StartWorldTime(Other.StartWorldTime)
+	, bIsInhibited(Other.bIsInhibited)
+	, bPendingRepOnActiveGC(Other.bPendingRepOnActiveGC)
+	, bPendingRepWhileActiveGC(Other.bPendingRepWhileActiveGC)
+	, IsPendingRemove(Other.IsPendingRemove)
+	, ClientCachedStackCount(0)
+	, PeriodHandle(Other.PeriodHandle)
+	, DurationHandle(Other.DurationHandle)
+	, EventSet(Other.EventSet)
 {
 
 	ReplicationID = Other.ReplicationID;
@@ -1604,10 +1624,9 @@ FActiveGameplayEffect& FActiveGameplayEffect::operator=(FActiveGameplayEffect&&
 	bPendingRepWhileActiveGC = Other.bPendingRepWhileActiveGC;
 	IsPendingRemove = Other.IsPendingRemove;
 	ClientCachedStackCount = Other.ClientCachedStackCount;
-	OnRemovedDelegate = Other.OnRemovedDelegate;
-	OnRemoved_InfoDelegate = Other.OnRemoved_InfoDelegate;
 	PeriodHandle = Other.PeriodHandle;
 	DurationHandle = Other.DurationHandle;
+	EventSet = Other.EventSet;
 	// Note: purposefully not copying PendingNext pointer.
 
 	ReplicationID = Other.ReplicationID;
@@ -1628,10 +1647,9 @@ FActiveGameplayEffect& FActiveGameplayEffect::operator=(const FActiveGameplayEff
 	bPendingRepWhileActiveGC = Other.bPendingRepWhileActiveGC;
 	IsPendingRemove = Other.IsPendingRemove;
 	ClientCachedStackCount = Other.ClientCachedStackCount;
-	OnRemovedDelegate = Other.OnRemovedDelegate;
-	OnRemoved_InfoDelegate = Other.OnRemoved_InfoDelegate;
 	PeriodHandle = Other.PeriodHandle;
 	DurationHandle = Other.DurationHandle;
+	EventSet = Other.EventSet;
 	PendingNext = Other.PendingNext;
 
 	ReplicationID = Other.ReplicationID;
@@ -1646,21 +1664,24 @@ void FActiveGameplayEffect::CheckOngoingTagRequirements(const FGameplayTagContai
 
 	if (bIsInhibited != bShouldBeInhibited)
 	{
-		// All OnDirty callbacks must be inhibited until we update this entire GameplayEffect.
-		FScopedAggregatorOnDirtyBatch	AggregatorOnDirtyBatcher;
+		{
+			// All OnDirty callbacks must be inhibited until we update this entire GameplayEffect.
+			FScopedAggregatorOnDirtyBatch	AggregatorOnDirtyBatcher;
 
-		// Important to set this prior to adding or removing, so that any delegates that are triggered can query accurately against this GE
-		bIsInhibited = bShouldBeInhibited;
+			// Important to set this prior to adding or removing, so that any delegates that are triggered can query accurately against this GE
+			bIsInhibited = bShouldBeInhibited;
 
-		if (bShouldBeInhibited)
-		{
-			// Remove our ActiveGameplayEffects modifiers with our Attribute Aggregators
-			OwningContainer.RemoveActiveGameplayEffectGrantedTagsAndModifiers(*this, bInvokeGameplayCueEvents);
-		}
-		else
-		{
-			OwningContainer.AddActiveGameplayEffectGrantedTagsAndModifiers(*this, bInvokeGameplayCueEvents);
+			if (bShouldBeInhibited)
+			{
+				// Remove our ActiveGameplayEffects modifiers with our Attribute Aggregators
+				OwningContainer.RemoveActiveGameplayEffectGrantedTagsAndModifiers(*this, bInvokeGameplayCueEvents);
+			}
+			else
+			{
+				OwningContainer.AddActiveGameplayEffectGrantedTagsAndModifiers(*this, bInvokeGameplayCueEvents);
+			}
 		}
+		EventSet.OnInhibitionChanged.Broadcast(Handle, bIsInhibited);
 	}
 }
 
@@ -2051,6 +2072,10 @@ FAggregatorRef& FActiveGameplayEffectsContainer::FindOrCreateAttributeAggregator
 	{
 		NewAttributeAggregator->OnDirty.AddUObject(Owner, &UAbilitySystemComponent::OnAttributeAggregatorDirty, Attribute, false);
 		NewAttributeAggregator->OnDirtyRecursive.AddUObject(Owner, &UAbilitySystemComponent::OnAttributeAggregatorDirty, Attribute, true);
+
+		// Callback in case the set wants to do something
+		const UAttributeSet* Set = Owner->GetAttributeSubobject(Attribute.GetAttributeSetClass());
+		Set->OnAttributeAggregatorCreated(Attribute, NewAttributeAggregator);
 	}
 
 	return AttributeAggregatorMap.Add(Attribute, FAggregatorRef(NewAttributeAggregator));
@@ -2193,13 +2218,13 @@ void FActiveGameplayEffectsContainer::OnStackCountChange(FActiveGameplayEffect&
 
 	Owner->NotifyTagMap_StackCountChange(ActiveEffect.Spec.DynamicGrantedTags);
 
-	ActiveEffect.OnStackChangeDelegate.Broadcast(ActiveEffect.Handle, ActiveEffect.Spec.StackCount, OldStackCount);
+	ActiveEffect.EventSet.OnStackChanged.Broadcast(ActiveEffect.Handle, ActiveEffect.Spec.StackCount, OldStackCount);
 }
 
 /** Called when the duration or starttime of an AGE has changed */
 void FActiveGameplayEffectsContainer::OnDurationChange(FActiveGameplayEffect& Effect)
 {
-	Effect.OnTimeChangeDelegate.Broadcast(Effect.Handle, Effect.StartWorldTime, Effect.GetDuration());
+	Effect.EventSet.OnTimeChanged.Broadcast(Effect.Handle, Effect.StartWorldTime, Effect.GetDuration());
 	Owner->OnGameplayEffectDurationChange(Effect);
 }
 
@@ -2714,7 +2739,7 @@ FActiveGameplayEffect* FActiveGameplayEffectsContainer::ApplyGameplayEffectSpec(
 		// @todo: If dynamic asset tags differ (which they shouldn't), we'll actually have to diff them
 		// and cause a removal and add of only the ones that have changed. For now, ensure on this happening and come
 		// back to this later.
-		ensureMsgf(ExistingSpec.DynamicAssetTags == Spec.DynamicAssetTags, TEXT("While adding a stack of the gameplay effect: %s, the old stack and the new application had different dynamic asset tags, which is currently not resolved properly!"), *Spec.Def->GetName());
+		ensureMsgf(ExistingSpec.DynamicAssetTags == Spec.DynamicAssetTags, TEXT("While adding a stack of the gameplay effect: %s, the old stack and the new application had different dynamic asset tags, which is currently not resolved properly! Existing: %s. New: %s"), *Spec.Def->GetName(), *ExistingSpec.DynamicAssetTags.ToStringSimple(), *Spec.DynamicAssetTags.ToStringSimple() );
 
 		ExistingStackableGE->Spec = Spec;
 		ExistingStackableGE->Spec.StackCount = NewStackCount;
@@ -3039,6 +3064,9 @@ void FActiveGameplayEffectsContainer::AddActiveGameplayEffectGrantedTagsAndModif
 			// would regrant the ability every time the stack was applied
 			if (AbilitySpecDef.AssignedHandle.IsValid() == false)
 			{
+				// Copy over SetByCaller Magnitudes (we can't do this in ::Initialize since these are set afterwards by the caller)
+				AbilitySpecDef.SetByCallerTagMagnitudes = Effect.Spec.SetByCallerTagMagnitudes;
+
 				Owner->GiveAbility( FGameplayAbilitySpec(AbilitySpecDef, Effect.Spec.GetLevel(), Effect.Handle) );
 
 				ABILITY_LOG(Display, TEXT("::AddActiveGameplayEffectGrantedTagsAndModifiers granted ability %s (Handle %s) from GE %s (Handle: %s)"), *GetNameSafe(AbilitySpecDef.Ability), *AbilitySpecDef.AssignedHandle.ToString(), *Effect.GetDebugString(), *Effect.Handle.ToString());
@@ -3240,8 +3268,9 @@ void FActiveGameplayEffectsContainer::InternalOnActiveGameplayEffectRemoved(FAct
 		ABILITY_LOG(Warning, TEXT("InternalOnActiveGameplayEffectRemoved called with no GameplayEffect: %s"), *Effect.Handle.ToString());
 	}
 
-	Effect.OnRemovedDelegate.Broadcast();
-	Effect.OnRemoved_InfoDelegate.Broadcast(GameplayEffectRemovalInfo);
+	Effect.EventSet.DEPRECATED_OnEffectRemoved.Broadcast();
+	Effect.EventSet.OnEffectRemoved.Broadcast(GameplayEffectRemovalInfo);
+
 	OnActiveGameplayEffectRemovedDelegate.Broadcast(Effect);
 }
 
@@ -3901,15 +3930,15 @@ TArray<FActiveGameplayEffectHandle> FActiveGameplayEffectsContainer::GetActiveEf
 	return ReturnList;
 }
 
-float FActiveGameplayEffectsContainer::GetActiveEffectsEndTime(const FGameplayEffectQuery& Query) const
+float FActiveGameplayEffectsContainer::GetActiveEffectsEndTime(const FGameplayEffectQuery& Query, TArray<AActor*>& Instigators) const
 {
 	float EndTime = 0.f;
 	float Duration = 0.f;
-	GetActiveEffectsEndTimeAndDuration(Query, EndTime, Duration);
+	GetActiveEffectsEndTimeAndDuration(Query, EndTime, Duration, Instigators);
 	return EndTime;
 }
 
-bool FActiveGameplayEffectsContainer::GetActiveEffectsEndTimeAndDuration(const FGameplayEffectQuery& Query, float& EndTime, float& Duration) const
+bool FActiveGameplayEffectsContainer::GetActiveEffectsEndTimeAndDuration(const FGameplayEffectQuery& Query, float& EndTime, float& Duration, TArray<AActor*>& Instigators) const
 {
 	bool FoundSomething = false;
 	
@@ -3936,6 +3965,8 @@ bool FActiveGameplayEffectsContainer::GetActiveEffectsEndTimeAndDuration(const F
 			EndTime = ThisEndTime;
 			Duration = Effect.GetDuration();
 		}
+
+		Instigators.AddUnique(Effect.Spec.GetEffectContext().GetOriginalInstigator());
 	}
 	return FoundSomething;
 }
@@ -4384,7 +4415,6 @@ bool FGameplayEffectQuery::Matches(const FGameplayEffectSpec& Spec) const
 {
 	if (Spec.Def == nullptr)
 	{
-		ABILITY_LOG(Error, TEXT("Matches called with no UGameplayEffect def."));
 		return false;
 	}
 
diff --git a/Engine/Plugins/Runtime/GameplayAbilities/Source/GameplayAbilities/Private/GameplayEffectAggregator.cpp b/Engine/Plugins/Runtime/GameplayAbilities/Source/GameplayAbilities/Private/GameplayEffectAggregator.cpp
--- a/Engine/Plugins/Runtime/GameplayAbilities/Source/GameplayAbilities/Private/GameplayEffectAggregator.cpp
+++ b/Engine/Plugins/Runtime/GameplayAbilities/Source/GameplayAbilities/Private/GameplayEffectAggregator.cpp
@@ -5,7 +5,7 @@
 #include "UObject/UObjectIterator.h"
 #include "AbilitySystemComponent.h"
 
-bool FAggregatorMod::Qualifies(const FAggregatorEvaluateParameters& Parameters) const
+void FAggregatorMod::UpdateQualifies(const FAggregatorEvaluateParameters& Parameters) const
 {
 	static const FGameplayTagContainer EmptyTagContainer;
 	const FGameplayTagContainer& SrcTags = Parameters.SourceTags ? *Parameters.SourceTags : EmptyTagContainer;
@@ -18,7 +18,8 @@ bool FAggregatorMod::Qualifies(const FAggregatorEvaluateParameters& Parameters)
 
 	if (Parameters.IncludePredictiveMods == false && IsPredicted)
 	{
-		return false;
+		IsQualified = false;
+		return;
 	}
 
 	if (ActiveHandle.IsValid())
@@ -27,7 +28,8 @@ bool FAggregatorMod::Qualifies(const FAggregatorEvaluateParameters& Parameters)
 		{
 			if (ActiveHandle == HandleToIgnore)
 			{
-				return false;
+				IsQualified = false;
+				return;
 			}
 		}
 	}
@@ -48,14 +50,14 @@ bool FAggregatorMod::Qualifies(const FAggregatorEvaluateParameters& Parameters)
 		}
 	}
 
-	return bSourceMet && bTargetMet && bSourceFilterMet && bTargetFilterMet;
+	IsQualified = bSourceMet && bTargetMet && bSourceFilterMet && bTargetFilterMet;
 }
 
 float FAggregatorModChannel::EvaluateWithBase(float InlineBaseValue, const FAggregatorEvaluateParameters& Parameters) const
 {
 	for (const FAggregatorMod& Mod : Mods[EGameplayModOp::Override])
 	{
-		if (Mod.Qualifies(Parameters))
+		if (Mod.Qualifies())
 		{
 			return Mod.EvaluatedMagnitude;
 		}
@@ -78,7 +80,7 @@ bool FAggregatorModChannel::ReverseEvaluate(float FinalValue, const FAggregatorE
 {
 	for (const FAggregatorMod& Mod : Mods[EGameplayModOp::Override])
 	{
-		if (Mod.Qualifies(Parameters))
+		if (Mod.Qualifies())
 		{
 			// This is the case we can't really handle due to lack of information.
 			ComputedValue = FinalValue;
@@ -143,6 +145,30 @@ void FAggregatorModChannel::AddModsFrom(const FAggregatorModChannel& Other)
 	}
 }
 
+void FAggregatorModChannel::UpdateQualifiesOnAllMods(const FAggregatorEvaluateParameters& Parameters) const
+{
+	for (int32 ModOpIdx = 0; ModOpIdx < ARRAY_COUNT(Mods); ++ModOpIdx)
+	{
+		for (const FAggregatorMod& Mod : Mods[ModOpIdx])
+		{
+			Mod.UpdateQualifies(Parameters);
+		}
+	}
+}
+
+void FAggregatorModChannel::ForEachMod(FAggregatorModInfo& Info, TFunction<void (const FAggregatorModInfo&) > Func) const
+{
+	for (int32 ModOpIdx = 0; ModOpIdx < ARRAY_COUNT(Mods); ++ModOpIdx)
+	{
+		Info.Op = (EGameplayModOp::Type)ModOpIdx;
+		for (const FAggregatorMod& Mod : Mods[ModOpIdx])
+		{
+			Info.Mod = &Mod;
+			Func(Info);
+		}
+	}
+}
+
 void FAggregatorModChannel::GetAllAggregatorMods(EGameplayModEvaluationChannel Channel, OUT TMap<EGameplayModEvaluationChannel, const TArray<FAggregatorMod>* >& OutMods) const
 {
 	OutMods.Add(Channel, Mods);
@@ -169,7 +195,7 @@ float FAggregatorModChannel::SumMods(const TArray<FAggregatorMod>& InMods, float
 
 	for (const FAggregatorMod& Mod : InMods)
 	{
-		if (Mod.Qualifies(Parameters))
+		if (Mod.Qualifies())
 		{
 			Sum += (Mod.EvaluatedMagnitude - Bias);
 		}
@@ -254,6 +280,17 @@ float FAggregatorModChannelContainer::ReverseEvaluate(float FinalValue, const FA
 	return ComputedValue;
 }
 
+void FAggregatorModChannelContainer::EvaluateQualificationForAllMods(const FAggregatorEvaluateParameters& Parameters) const
+{
+	// First run our "Default" qualifies function
+	for (auto& MapIt : ModChannelsMap)
+	{
+		const FAggregatorModChannel& Channel = MapIt.Value;
+		Channel.UpdateQualifiesOnAllMods(Parameters);
+	}
+}
+
+
 void FAggregatorModChannelContainer::RemoveAggregatorMod(const FActiveGameplayEffectHandle& ActiveHandle)
 {
 	if (ActiveHandle.IsValid())
@@ -298,6 +335,17 @@ void FAggregatorModChannelContainer::OnActiveEffectDependenciesSwapped(const TMa
 	}
 }
 
+void FAggregatorModChannelContainer::ForEachMod( TFunction<void (const FAggregatorModInfo&) > Func) const
+{
+	FAggregatorModInfo Info;
+	for (const auto& MapIt : ModChannelsMap)
+	{
+		Info.Channel = MapIt.Key;
+		const FAggregatorModChannel& Channel = MapIt.Value;
+		Channel.ForEachMod(Info, Func);
+	}
+}
+
 FAggregator::~FAggregator()
 {
 	int32 NumRemoved = FScopedAggregatorOnDirtyBatch::DirtyAggregators.Remove(this);
@@ -306,21 +354,25 @@ FAggregator::~FAggregator()
 
 float FAggregator::Evaluate(const FAggregatorEvaluateParameters& Parameters) const
 {
+	EvaluateQualificationForAllMods(Parameters);
 	return ModChannels.EvaluateWithBase(BaseValue, Parameters);
 }
 
 float FAggregator::EvaluateToChannel(const FAggregatorEvaluateParameters& Parameters, EGameplayModEvaluationChannel FinalChannel) const
 {
+	EvaluateQualificationForAllMods(Parameters);
 	return ModChannels.EvaluateWithBaseToChannel(BaseValue, Parameters, FinalChannel);
 }
 
 float FAggregator::EvaluateWithBase(float InlineBaseValue, const FAggregatorEvaluateParameters& Parameters) const
 {
+	EvaluateQualificationForAllMods(Parameters);
 	return ModChannels.EvaluateWithBase(InlineBaseValue, Parameters);
 }
 
 float FAggregator::ReverseEvaluate(float FinalValue, const FAggregatorEvaluateParameters& Parameters) const
 {
+	EvaluateQualificationForAllMods(Parameters);
 	return ModChannels.ReverseEvaluate(FinalValue, Parameters);
 }
 
@@ -342,6 +394,18 @@ float FAggregator::EvaluateContribution(const FAggregatorEvaluateParameters& Par
 	return 0.f;
 }
 
+void FAggregator::EvaluateQualificationForAllMods(const FAggregatorEvaluateParameters& Parameters) const
+{
+	// First run our "Default" qualifies function
+	ModChannels.EvaluateQualificationForAllMods(Parameters);
+
+	// Then run custom func
+	if (EvaluationMetaData && EvaluationMetaData->CustomQualifiesFunc)
+	{
+		EvaluationMetaData->CustomQualifiesFunc(Parameters, this);
+	}
+}
+
 float FAggregator::GetBaseValue() const
 {
 	return BaseValue;
@@ -481,6 +545,11 @@ void FAggregator::OnActiveEffectDependenciesSwapped(const TMap<FActiveGameplayEf
 	}
 }
 
+void FAggregator::ForEachMod( TFunction<void (const FAggregatorModInfo&) > Func ) const
+{
+	ModChannels.ForEachMod(Func);
+}
+
 void FAggregator::TakeSnapshotOf(const FAggregator& AggToSnapshot)
 {
 	BaseValue = AggToSnapshot.BaseValue;
diff --git a/Engine/Plugins/Runtime/GameplayAbilities/Source/GameplayAbilities/Private/GameplayEffectExecutionCalculation.cpp b/Engine/Plugins/Runtime/GameplayAbilities/Source/GameplayAbilities/Private/GameplayEffectExecutionCalculation.cpp
--- a/Engine/Plugins/Runtime/GameplayAbilities/Source/GameplayAbilities/Private/GameplayEffectExecutionCalculation.cpp
+++ b/Engine/Plugins/Runtime/GameplayAbilities/Source/GameplayAbilities/Private/GameplayEffectExecutionCalculation.cpp
@@ -200,7 +200,7 @@ bool FGameplayEffectCustomExecutionParameters::AttemptGetCapturedAttributeAggreg
 	return false;
 }
 
-bool FGameplayEffectCustomExecutionParameters::AttemptGatherAttributeMods(const FGameplayEffectAttributeCaptureDefinition& InCaptureDef, OUT TMap<EGameplayModEvaluationChannel, const TArray<FAggregatorMod>*>& OutModMap) const
+bool FGameplayEffectCustomExecutionParameters::AttemptGatherAttributeMods(const FGameplayEffectAttributeCaptureDefinition& InCaptureDef, const FAggregatorEvaluateParameters& InEvalParams, OUT TMap<EGameplayModEvaluationChannel, const TArray<FAggregatorMod>*>& OutModMap) const
 {
 	check(OwningSpec);
 
@@ -214,7 +214,7 @@ bool FGameplayEffectCustomExecutionParameters::AttemptGatherAttributeMods(const
 		const FGameplayEffectAttributeCaptureSpec* CaptureSpec = OwningSpec->CapturedRelevantAttributes.FindCaptureSpecByDefinition(InCaptureDef, true);
 		if (CaptureSpec)
 		{
-			return CaptureSpec->AttemptGatherAttributeMods(OutModMap);
+			return CaptureSpec->AttemptGatherAttributeMods(InEvalParams, OutModMap);
 		}
 	}
 
@@ -224,7 +224,7 @@ bool FGameplayEffectCustomExecutionParameters::AttemptGatherAttributeMods(const
 bool FGameplayEffectCustomExecutionParameters::ForEachQualifiedAttributeMod(const FGameplayEffectAttributeCaptureDefinition& InCaptureDef, const FAggregatorEvaluateParameters& InEvalParams, TFunction< void(EGameplayModEvaluationChannel, EGameplayModOp::Type, const FAggregatorMod&) > Func) const
 {
 	TMap<EGameplayModEvaluationChannel, const TArray<FAggregatorMod>*> ModMap;
-	if (AttemptGatherAttributeMods(InCaptureDef, ModMap))
+	if (AttemptGatherAttributeMods(InCaptureDef, InEvalParams, ModMap))
 	{
 		for (auto It : ModMap)
 		{
@@ -234,7 +234,7 @@ bool FGameplayEffectCustomExecutionParameters::ForEachQualifiedAttributeMod(cons
 				const TArray<FAggregatorMod>& CurModArray = ModList[ModOpIdx];
 				for (const FAggregatorMod& AggMod : CurModArray)
 				{
-					if (AggMod.Qualifies(InEvalParams))
+					if (AggMod.Qualifies())
 					{
 						Func(It.Key, (EGameplayModOp::Type)ModOpIdx, AggMod);
 					}
diff --git a/Engine/Plugins/Runtime/GameplayAbilities/Source/GameplayAbilities/Public/Abilities/Tasks/AbilityTask_WaitAttributeChange.h b/Engine/Plugins/Runtime/GameplayAbilities/Source/GameplayAbilities/Public/Abilities/Tasks/AbilityTask_WaitAttributeChange.h
--- a/Engine/Plugins/Runtime/GameplayAbilities/Source/GameplayAbilities/Public/Abilities/Tasks/AbilityTask_WaitAttributeChange.h
+++ b/Engine/Plugins/Runtime/GameplayAbilities/Source/GameplayAbilities/Public/Abilities/Tasks/AbilityTask_WaitAttributeChange.h
@@ -45,11 +45,11 @@ class GAMEPLAYABILITIES_API UAbilityTask_WaitAttributeChange : public UAbilityTa
 
 	/** Wait until an attribute changes. */
 	UFUNCTION(BlueprintCallable, Category="Ability|Tasks", meta = (HidePin = "OwningAbility", DefaultToSelf = "OwningAbility", BlueprintInternalUseOnly = "TRUE"))
-	static UAbilityTask_WaitAttributeChange* WaitForAttributeChange(UGameplayAbility* OwningAbility, FGameplayAttribute Attribute, FGameplayTag WithSrcTag, FGameplayTag WithoutSrcTag, bool TriggerOnce=true);
+	static UAbilityTask_WaitAttributeChange* WaitForAttributeChange(UGameplayAbility* OwningAbility, FGameplayAttribute Attribute, FGameplayTag WithSrcTag, FGameplayTag WithoutSrcTag, bool TriggerOnce=true, AActor* OptionalExternalOwner = nullptr);
 
 	/** Wait until an attribute changes to pass a given test. */
 	UFUNCTION(BlueprintCallable, Category = "Ability|Tasks", meta = (HidePin = "OwningAbility", DefaultToSelf = "OwningAbility", BlueprintInternalUseOnly = "TRUE"))
-	static UAbilityTask_WaitAttributeChange* WaitForAttributeChangeWithComparison(UGameplayAbility* OwningAbility, FGameplayAttribute InAttribute, FGameplayTag InWithTag, FGameplayTag InWithoutTag, TEnumAsByte<EWaitAttributeChangeComparison::Type> InComparisonType, float InComparisonValue, bool TriggerOnce=true);
+	static UAbilityTask_WaitAttributeChange* WaitForAttributeChangeWithComparison(UGameplayAbility* OwningAbility, FGameplayAttribute InAttribute, FGameplayTag InWithTag, FGameplayTag InWithoutTag, TEnumAsByte<EWaitAttributeChangeComparison::Type> InComparisonType, float InComparisonValue, bool TriggerOnce=true, AActor* OptionalExternalOwner = nullptr);
 
 	FGameplayTag WithTag;
 	FGameplayTag WithoutTag;
@@ -61,5 +61,10 @@ class GAMEPLAYABILITIES_API UAbilityTask_WaitAttributeChange : public UAbilityTa
 
 protected:
 
+	UPROPERTY()
+	UAbilitySystemComponent* ExternalOwner;
+
+	UAbilitySystemComponent* GetFocusedASC();
+
 	virtual void OnDestroy(bool AbilityEnded) override;
 };
diff --git a/Engine/Plugins/Runtime/GameplayAbilities/Source/GameplayAbilities/Public/Abilities/Tasks/AbilityTask_WaitAttributeChangeRatioThreshold.h b/Engine/Plugins/Runtime/GameplayAbilities/Source/GameplayAbilities/Public/Abilities/Tasks/AbilityTask_WaitAttributeChangeRatioThreshold.h
--- a/Engine/Plugins/Runtime/GameplayAbilities/Source/GameplayAbilities/Public/Abilities/Tasks/AbilityTask_WaitAttributeChangeRatioThreshold.h
+++ b/Engine/Plugins/Runtime/GameplayAbilities/Source/GameplayAbilities/Public/Abilities/Tasks/AbilityTask_WaitAttributeChangeRatioThreshold.h
@@ -31,7 +31,7 @@ class GAMEPLAYABILITIES_API UAbilityTask_WaitAttributeChangeRatioThreshold : pub
 
 	/** Wait on attribute ratio change meeting a comparison threshold. */
 	UFUNCTION(BlueprintCallable, Category="Ability|Tasks", meta = (HidePin = "OwningAbility", DefaultToSelf = "OwningAbility", BlueprintInternalUseOnly = "TRUE"))
-	static UAbilityTask_WaitAttributeChangeRatioThreshold* WaitForAttributeChangeRatioThreshold(UGameplayAbility* OwningAbility, FGameplayAttribute AttributeNumerator, FGameplayAttribute AttributeDenominator, TEnumAsByte<EWaitAttributeChangeComparison::Type> ComparisonType, float ComparisonValue, bool bTriggerOnce);
+	static UAbilityTask_WaitAttributeChangeRatioThreshold* WaitForAttributeChangeRatioThreshold(UGameplayAbility* OwningAbility, FGameplayAttribute AttributeNumerator, FGameplayAttribute AttributeDenominator, TEnumAsByte<EWaitAttributeChangeComparison::Type> ComparisonType, float ComparisonValue, bool bTriggerOnce, AActor* OptionalExternalOwner = nullptr);
 
 	FGameplayAttribute AttributeNumerator;
 	FGameplayAttribute AttributeDenominator;
@@ -50,6 +50,11 @@ protected:
 	/** Timer used when either numerator or denominator attribute is changed to delay checking of ratio to avoid false positives (MaxHealth changed before Health updates accordingly) */
 	FTimerHandle CheckAttributeTimer;
 
+	UPROPERTY()
+	UAbilitySystemComponent* ExternalOwner;
+
+	UAbilitySystemComponent* GetFocusedASC();
+
 	void OnAttributeChange();
 	void OnRatioChange();
 
diff --git a/Engine/Plugins/Runtime/GameplayAbilities/Source/GameplayAbilities/Public/Abilities/Tasks/AbilityTask_WaitAttributeChangeThreshold.h b/Engine/Plugins/Runtime/GameplayAbilities/Source/GameplayAbilities/Public/Abilities/Tasks/AbilityTask_WaitAttributeChangeThreshold.h
--- a/Engine/Plugins/Runtime/GameplayAbilities/Source/GameplayAbilities/Public/Abilities/Tasks/AbilityTask_WaitAttributeChangeThreshold.h
+++ b/Engine/Plugins/Runtime/GameplayAbilities/Source/GameplayAbilities/Public/Abilities/Tasks/AbilityTask_WaitAttributeChangeThreshold.h
@@ -29,7 +29,7 @@ class GAMEPLAYABILITIES_API UAbilityTask_WaitAttributeChangeThreshold : public U
 
 	/** Wait on attribute change meeting a comparison threshold. */
 	UFUNCTION(BlueprintCallable, Category="Ability|Tasks", meta = (HidePin = "OwningAbility", DefaultToSelf = "OwningAbility", BlueprintInternalUseOnly = "TRUE"))
-	static UAbilityTask_WaitAttributeChangeThreshold* WaitForAttributeChangeThreshold(UGameplayAbility* OwningAbility, FGameplayAttribute Attribute, TEnumAsByte<EWaitAttributeChangeComparison::Type> ComparisonType, float ComparisonValue, bool bTriggerOnce);
+	static UAbilityTask_WaitAttributeChangeThreshold* WaitForAttributeChangeThreshold(UGameplayAbility* OwningAbility, FGameplayAttribute Attribute, TEnumAsByte<EWaitAttributeChangeComparison::Type> ComparisonType, float ComparisonValue, bool bTriggerOnce, AActor* OptionalExternalOwner = nullptr);
 
 	FGameplayAttribute Attribute;
 	TEnumAsByte<EWaitAttributeChangeComparison::Type> ComparisonType;
@@ -41,6 +41,11 @@ protected:
 
 	bool bMatchedComparisonLastAttributeChange;
 
+	UPROPERTY()
+	UAbilitySystemComponent* ExternalOwner;
+
+	UAbilitySystemComponent* GetFocusedASC();
+
 	virtual void OnDestroy(bool AbilityEnded) override;
 
 	bool DoesValuePassComparison(float Value) const;
diff --git a/Engine/Plugins/Runtime/GameplayAbilities/Source/GameplayAbilities/Public/AbilitySystemComponent.h b/Engine/Plugins/Runtime/GameplayAbilities/Source/GameplayAbilities/Public/AbilitySystemComponent.h
--- a/Engine/Plugins/Runtime/GameplayAbilities/Source/GameplayAbilities/Public/AbilitySystemComponent.h
+++ b/Engine/Plugins/Runtime/GameplayAbilities/Source/GameplayAbilities/Public/AbilitySystemComponent.h
@@ -586,15 +586,13 @@ class GAMEPLAYABILITIES_API UAbilitySystemComponent : public UGameplayTasksCompo
 	// Additional Helper Functions
 	// --------------------------------------------
 	
+	FOnGivenActiveGameplayEffectRemoved& OnAnyGameplayEffectRemovedDelegate();
 	DEPRECATED(4.17, "Use OnGameplayEffectRemoved_InfoDelegate (the delegate signature has changed)")
 	FOnActiveGameplayEffectRemoved* OnGameplayEffectRemovedDelegate(FActiveGameplayEffectHandle Handle);
 
+	FActiveGameplayEffectEvents* GetActiveEffectEventSet(FActiveGameplayEffectHandle Handle);
 	FOnActiveGameplayEffectRemoved_Info* OnGameplayEffectRemoved_InfoDelegate(FActiveGameplayEffectHandle Handle);
-
-	FOnGivenActiveGameplayEffectRemoved& OnAnyGameplayEffectRemovedDelegate();
-
 	FOnActiveGameplayEffectStackChange* OnGameplayEffectStackChangeDelegate(FActiveGameplayEffectHandle Handle);
-
 	FOnActiveGameplayEffectTimeChange* OnGameplayEffectTimeChangeDelegate(FActiveGameplayEffectHandle Handle);
 
 	UFUNCTION(BlueprintCallable, Category = GameplayEffects, meta=(DisplayName = "ApplyGameplayEffectToTarget"))
@@ -652,12 +650,19 @@ class GAMEPLAYABILITIES_API UAbilitySystemComponent : public UGameplayTasksCompo
 	/** This will give the world time that all effects matching this query will be finished. If multiple effects match, it returns the one that returns last.*/
 	float GetActiveEffectsEndTime(const FGameplayEffectQuery& Query) const
 	{
-		return ActiveGameplayEffects.GetActiveEffectsEndTime(Query);
+		TArray<AActor*> DummyInstigators;
+		return ActiveGameplayEffects.GetActiveEffectsEndTime(Query, DummyInstigators);
+	}
+	
+	float GetActiveEffectsEndTimeWithInstigators(const FGameplayEffectQuery& Query, TArray<AActor*>& Instigators) const
+	{
+		return ActiveGameplayEffects.GetActiveEffectsEndTime(Query, Instigators);
 	}
 
 	bool GetActiveEffectsEndTimeAndDuration(const FGameplayEffectQuery& Query, float& EndTime, float& Duration) const
 	{
-		return ActiveGameplayEffects.GetActiveEffectsEndTimeAndDuration(Query, EndTime, Duration);
+		TArray<AActor*> DummyInstigators;
+		return ActiveGameplayEffects.GetActiveEffectsEndTimeAndDuration(Query, EndTime, Duration, DummyInstigators);
 	}
 
 	void ModifyActiveEffectStartTime(FActiveGameplayEffectHandle Handle, float StartTimeDiff)
@@ -1046,7 +1051,7 @@ class GAMEPLAYABILITIES_API UAbilitySystemComponent : public UGameplayTasksCompo
 	UFUNCTION()
 	virtual void OnRep_ServerDebugString();
 
-	float GetFilteredAttributeValue(const FGameplayAttribute& Attribute, const FGameplayTagRequirements& SourceTags, const FGameplayTagContainer& TargetTags);
+	float GetFilteredAttributeValue(const FGameplayAttribute& Attribute, const FGameplayTagRequirements& SourceTags, const FGameplayTagContainer& TargetTags, const TArray<FActiveGameplayEffectHandle>& HandlesToIgnore = TArray<FActiveGameplayEffectHandle>());
 
 protected:
 
diff --git a/Engine/Plugins/Runtime/GameplayAbilities/Source/GameplayAbilities/Public/AbilitySystemGlobals.h b/Engine/Plugins/Runtime/GameplayAbilities/Source/GameplayAbilities/Public/AbilitySystemGlobals.h
--- a/Engine/Plugins/Runtime/GameplayAbilities/Source/GameplayAbilities/Public/AbilitySystemGlobals.h
+++ b/Engine/Plugins/Runtime/GameplayAbilities/Source/GameplayAbilities/Public/AbilitySystemGlobals.h
@@ -161,6 +161,11 @@ class GAMEPLAYABILITIES_API UAbilitySystemGlobals : public UObject
 
 	// Global Tags
 
+	UPROPERTY()
+	FGameplayTag ActivateFailIsDeadTag; // TryActivate failed due to being dead
+	UPROPERTY(config)
+	FName ActivateFailIsDeadName;
+
 	UPROPERTY()
 	FGameplayTag ActivateFailCooldownTag; // TryActivate failed due to being on cooldown
 	UPROPERTY(config)
@@ -192,6 +197,11 @@ class GAMEPLAYABILITIES_API UAbilitySystemGlobals : public UObject
 
 	virtual void InitGlobalTags()
 	{
+		if (ActivateFailIsDeadName != NAME_None)
+		{
+			ActivateFailIsDeadTag = FGameplayTag::RequestGameplayTag(ActivateFailIsDeadName);
+		}
+
 		if (ActivateFailCooldownName != NAME_None)
 		{
 			ActivateFailCooldownTag = FGameplayTag::RequestGameplayTag(ActivateFailCooldownName);
diff --git a/Engine/Plugins/Runtime/GameplayAbilities/Source/GameplayAbilities/Public/AttributeSet.h b/Engine/Plugins/Runtime/GameplayAbilities/Source/GameplayAbilities/Public/AttributeSet.h
--- a/Engine/Plugins/Runtime/GameplayAbilities/Source/GameplayAbilities/Public/AttributeSet.h
+++ b/Engine/Plugins/Runtime/GameplayAbilities/Source/GameplayAbilities/Public/AttributeSet.h
@@ -16,6 +16,7 @@
 class UAbilitySystemComponent;
 class UAttributeSet;
 struct FGameplayAbilityActorInfo;
+struct FAggregator;
 
 USTRUCT(BlueprintType)
 struct GAMEPLAYABILITIES_API FGameplayAttributeData
@@ -198,6 +199,9 @@ public:
 	 */
 	virtual void PreAttributeBaseChange(const FGameplayAttribute& Attribute, float& NewValue) const { }
 
+	/** Callback for when an FAggregator is created for an attribute in this set. Allows custom setup of FAggregator::EvaluationMetaData */
+	virtual void OnAttributeAggregatorCreated(const FGameplayAttribute& Attribute, FAggregator* NewAggregator) const { }
+
 	/** This signifies the attribute set can be ID'd by name over the network. */
 	void SetNetAddressable();
 
diff --git a/Engine/Plugins/Runtime/GameplayAbilities/Source/GameplayAbilities/Public/GameplayAbilitySpec.h b/Engine/Plugins/Runtime/GameplayAbilities/Source/GameplayAbilities/Public/GameplayAbilitySpec.h
--- a/Engine/Plugins/Runtime/GameplayAbilities/Source/GameplayAbilities/Public/GameplayAbilitySpec.h
+++ b/Engine/Plugins/Runtime/GameplayAbilities/Source/GameplayAbilities/Public/GameplayAbilitySpec.h
@@ -136,6 +136,9 @@ struct FGameplayAbilitySpecDef
 	UPROPERTY(NotReplicated)
 	UObject* SourceObject;
 
+	// SetbyCaller Magnitudes that were passed in to this ability by a GE (GE's that grant abilities). Made available so that 
+	TMap<FGameplayTag, float>	SetByCallerTagMagnitudes;
+
 	/** This handle can be set if the SpecDef is used to create a real FGameplaybilitySpec */
 	UPROPERTY()
 	FGameplayAbilitySpecHandle	AssignedHandle;
@@ -279,6 +282,9 @@ struct GAMEPLAYABILITIES_API FGameplayAbilitySpec : public FFastArraySerializerI
 	UPROPERTY(NotReplicated)
 	FActiveGameplayEffectHandle	GameplayEffectHandle;
 
+	/** Passed on SetByCaller magnitudes if this ability was granted by a GE */
+	TMap<FGameplayTag, float> SetByCallerTagMagnitudes;
+
 	/** Returns the primary instance, used for instance once abilities */
 	UGameplayAbility* GetPrimaryInstance() const;
 
diff --git a/Engine/Plugins/Runtime/GameplayAbilities/Source/GameplayAbilities/Public/GameplayCueManager.h b/Engine/Plugins/Runtime/GameplayAbilities/Source/GameplayAbilities/Public/GameplayCueManager.h
--- a/Engine/Plugins/Runtime/GameplayAbilities/Source/GameplayAbilities/Public/GameplayCueManager.h
+++ b/Engine/Plugins/Runtime/GameplayAbilities/Source/GameplayAbilities/Public/GameplayCueManager.h
@@ -99,6 +99,9 @@ class GAMEPLAYABILITIES_API UGameplayCueManager : public UDataAsset
 	/** Send out any pending cues */
 	virtual void FlushPendingCues();
 
+	/** Broadcasted when ::FlushPendingCues runs: useful for custom batching/gameplay cue handling */
+	FSimpleMulticastDelegate	OnFlushPendingCues;
+
 	virtual void OnCreated();
 
 	virtual void OnEngineInitComplete();
@@ -207,8 +210,8 @@ protected:
 	/** Refreshes the existing, already initialized, object libraries. */
 	void RefreshObjectLibraries();
 
-	/** Internal function to actually init the FGameplayCueObjectLibrary */
-	void InitObjectLibrary(FGameplayCueObjectLibrary& Library);
+	/** Internal function to actually init the FGameplayCueObjectLibrary.  Returns StreamableHandle when asset async loading is requested. */
+	TSharedPtr<FStreamableHandle> InitObjectLibrary(FGameplayCueObjectLibrary& Library);
 
 	virtual TArray<FString> GetAlwaysLoadedGameplayCuePaths();
 
@@ -221,7 +224,9 @@ protected:
 	UPROPERTY(transient)
 	FGameplayCueObjectLibrary EditorGameplayCueObjectLibrary;
 
-	/** Handle to maintain ownership of gameplay cue assets. */
+	/** Handle to maintain ownership of gameplay cue assets.  
+      Note:	Only the latest handle to async load request is cached.
+			If projects require multiple concurrent asynchronous loads, handles returned from InitObjectLibrary should be cached as needed. */
 	TSharedPtr<FStreamableHandle> GameplayCueAssetHandle;
 
 public:		
diff --git a/Engine/Plugins/Runtime/GameplayAbilities/Source/GameplayAbilities/Public/GameplayCueTranslator.h b/Engine/Plugins/Runtime/GameplayAbilities/Source/GameplayAbilities/Public/GameplayCueTranslator.h
--- a/Engine/Plugins/Runtime/GameplayAbilities/Source/GameplayAbilities/Public/GameplayCueTranslator.h
+++ b/Engine/Plugins/Runtime/GameplayAbilities/Source/GameplayAbilities/Public/GameplayCueTranslator.h
@@ -201,6 +201,9 @@ struct GAMEPLAYABILITIES_API FGameplayCueTranslationManager
 	/** Used by the GC editor to enumerate possible translation tags. Never called at runtime. */
 	bool GetTranslatedTags(const FName& ParentTag, TArray<FGameplayCueTranslationEditorInfo>& Children);
 	const TArray<FNameSwapData>& GetNameSwapData() const { return AllNameSwaps; }
+
+	/** Searches, slowly, to see if the passed in tag can be translated from something else. Don't call this at runtime. */
+	FGameplayTag SearchSlowForTranslationParent(FGameplayTag Tag);
 #endif
 
 private:
diff --git a/Engine/Plugins/Runtime/GameplayAbilities/Source/GameplayAbilities/Public/GameplayEffect.h b/Engine/Plugins/Runtime/GameplayAbilities/Source/GameplayAbilities/Public/GameplayEffect.h
--- a/Engine/Plugins/Runtime/GameplayAbilities/Source/GameplayAbilities/Public/GameplayEffect.h
+++ b/Engine/Plugins/Runtime/GameplayAbilities/Source/GameplayAbilities/Public/GameplayEffect.h
@@ -31,6 +31,8 @@ class UGameplayModMagnitudeCalculation;
 struct FActiveGameplayEffectsContainer;
 struct FGameplayEffectModCallbackData;
 struct FGameplayEffectSpec;
+struct FAggregatorEvaluateParameters;
+struct FAggregatorMod;
 
 /** Enumeration outlining the possible gameplay effect magnitude calculation policies. */
 UENUM()
@@ -296,6 +298,9 @@ public:
 	/** Returns the DataName associated with this magnitude if it is set by caller */
 	bool GetSetByCallerDataNameIfPossible(FName& OutDataName) const;
 
+	/** Returns SetByCaller data structure, for inspection purposes */
+	const FSetByCallerFloat& GetSetByCallerFloat() const { return SetByCallerMagnitude; }
+
 	/** Returns the custom magnitude calculation class, if any, for this magnitude. Only applies to CustomMagnitudes */
 	TSubclassOf<UGameplayModMagnitudeCalculation> GetCustomMagnitudeCalculationClass() const;
 
@@ -761,8 +766,8 @@ struct GAMEPLAYABILITIES_API FGameplayEffectAttributeCaptureSpec
 	 */
 	bool AttemptAddAggregatorModsToAggregator(OUT FAggregator& OutAggregatorToAddTo) const;
 	
-	/** Gathers made for a given capture. Note that these mods are unqualified and direct references (not copies). Only use this if you know what you are doing. */
-	bool AttemptGatherAttributeMods(OUT TMap<EGameplayModEvaluationChannel, const TArray<FAggregatorMod>*>& OutModMap) const;
+	/** Gathers made for a given capture. Note all mods are returned but only some will be qualified (use Qualified() func to determine) */
+	bool AttemptGatherAttributeMods(const FAggregatorEvaluateParameters& InEvalParams, OUT TMap<EGameplayModEvaluationChannel, const TArray<FAggregatorMod>*>& OutModMap) const;
 	
 	/** Simple accessor to backing capture definition */
 	const FGameplayEffectAttributeCaptureDefinition& GetBackingDefinition() const;
@@ -908,6 +913,12 @@ struct GAMEPLAYABILITIES_API FGameplayEffectSpec
 	// Initialize the spec as a linked spec. The original spec's context is preserved except for the original GE asset tags, which are stripped out
 	void InitializeFromLinkedSpec(const UGameplayEffect* InDef, const FGameplayEffectSpec& OriginalSpec);
 
+	// Copies SetbyCallerMagnitudes from OriginalSpec into this
+	void CopySetByCallerMagnitudes(const FGameplayEffectSpec& OriginalSpec);
+
+	// Copies SetbuCallerMagnitudes, but only if magnitudes don't exist in our map (slower but preserves data)
+	void MergeSetByCallerMagnitudes(const TMap<FGameplayTag, float>& Magnitudes);
+
 	/**
 	 * Determines if the spec has capture specs with valid captures for all of the specified definitions.
 	 * 
@@ -1085,11 +1096,11 @@ public:
 	UPROPERTY()
 	TArray<FGameplayAbilitySpecDef> GrantedAbilitySpecs;
 
-private:
-
 	/** Map of set by caller magnitudes */
 	TMap<FName, float>			SetByCallerNameMagnitudes;
 	TMap<FGameplayTag, float>	SetByCallerTagMagnitudes;
+
+private:
 	
 	UPROPERTY()
 	FGameplayEffectContextHandle EffectContext; // This tells us how we got here (who / what applied us)
@@ -1258,18 +1269,13 @@ struct GAMEPLAYABILITIES_API FActiveGameplayEffect : public FFastArraySerializer
 	/** Last StackCount that the client had. Used to tell if the stackcount has changed in PostReplicatedChange */
 	int32 ClientCachedStackCount;
 
-	FOnActiveGameplayEffectRemoved OnRemovedDelegate;
-	FOnActiveGameplayEffectRemoved_Info OnRemoved_InfoDelegate;
-
-	FOnActiveGameplayEffectStackChange OnStackChangeDelegate;
-
-	FOnActiveGameplayEffectTimeChange OnTimeChangeDelegate;
-
 	FTimerHandle PeriodHandle;
-
 	FTimerHandle DurationHandle;
 
 	FActiveGameplayEffect* PendingNext;
+	
+	/** All the bindable events for this active effect (bundled to allow easier non-const access to these events via the ASC) */
+	FActiveGameplayEffectEvents EventSet;
 };
 
 DECLARE_DELEGATE_RetVal_OneParam(bool, FActiveGameplayEffectQueryCustomMatch, const FActiveGameplayEffect&);
@@ -1580,7 +1586,7 @@ struct GAMEPLAYABILITIES_API FActiveGameplayEffectsContainer : public FFastArray
 
 	bool NetDeltaSerialize(FNetDeltaSerializeInfo& DeltaParms);
 
-	void Uninitialize();	
+	void Uninitialize();
 
 	// ------------------------------------------------
 
@@ -1594,8 +1600,8 @@ struct GAMEPLAYABILITIES_API FActiveGameplayEffectsContainer : public FFastArray
 
 	TArray<FActiveGameplayEffectHandle> GetActiveEffects(const FGameplayEffectQuery& Query) const;
 
-	float GetActiveEffectsEndTime(const FGameplayEffectQuery& Query) const;
-	bool GetActiveEffectsEndTimeAndDuration(const FGameplayEffectQuery& Query, float& EndTime, float& Duration) const;
+	float GetActiveEffectsEndTime(const FGameplayEffectQuery& Query, TArray<AActor*>& Instigators) const;
+	bool GetActiveEffectsEndTimeAndDuration(const FGameplayEffectQuery& Query, float& EndTime, float& Duration, TArray<AActor*>& Instigators) const;
 
 	/** Returns an array of all of the active gameplay effect handles */
 	TArray<FActiveGameplayEffectHandle> GetAllActiveEffectHandles() const;
diff --git a/Engine/Plugins/Runtime/GameplayAbilities/Source/GameplayAbilities/Public/GameplayEffectAggregator.h b/Engine/Plugins/Runtime/GameplayAbilities/Source/GameplayAbilities/Public/GameplayEffectAggregator.h
--- a/Engine/Plugins/Runtime/GameplayAbilities/Source/GameplayAbilities/Public/GameplayEffectAggregator.h
+++ b/Engine/Plugins/Runtime/GameplayAbilities/Source/GameplayAbilities/Public/GameplayEffectAggregator.h
@@ -7,7 +7,9 @@
 #include "GameplayEffectTypes.h"
 
 struct FGameplayEffectSpec;
+struct FAggregator;
 
+/** Data that is used in aggregator evaluation that is passed from the caller/game code */
 struct GAMEPLAYABILITIES_API FAggregatorEvaluateParameters
 {
 	FAggregatorEvaluateParameters()
@@ -31,6 +33,21 @@ struct GAMEPLAYABILITIES_API FAggregatorEvaluateParameters
 	bool IncludePredictiveMods;
 };
 
+/** Data that is used in aggregator evaluation that is intrinsic to the aggregator itself. Not passed in from the game code (though maybe originally setup once by the game code)   */
+struct GAMEPLAYABILITIES_API FAggregatorEvaluateMetaData
+{
+	typedef TFunction< void(const FAggregatorEvaluateParameters&, const FAggregator*) > FCustomQualifiesFunc;
+
+	FAggregatorEvaluateMetaData(FCustomQualifiesFunc InQualifierFunc) 
+		: CustomQualifiesFunc(InQualifierFunc)
+	{
+
+	}
+
+
+	FCustomQualifiesFunc CustomQualifiesFunc;
+};
+
 struct GAMEPLAYABILITIES_API FAggregatorMod
 {
 	const FGameplayTagRequirements*	SourceTagReqs;
@@ -41,8 +58,26 @@ struct GAMEPLAYABILITIES_API FAggregatorMod
 
 	FActiveGameplayEffectHandle ActiveHandle;	// Handle of the active GameplayEffect we are tied to (if any)
 	bool IsPredicted;
+	
+	bool Qualifies() const { return IsQualified; }
+
+	/** Called to update the Qualifies bool */
+	void UpdateQualifies(const FAggregatorEvaluateParameters& Parameters) const;
+
+	/** Intended to be used by FAggregatorEvaluateMetaData::CustomQualifiesFunc to toggle qualifications of mods */
+	void SetExplicitQualifies(bool NewQualifies) const { IsQualified = NewQualifies; }
+
+private:
+
+	/** This bool is updated by UpdateQualifiees. Think of it as a transient bool, we make a full pass on all qualifiers first (::UpdateQualifies) then while evaluating we check what the bool was set to (::Qualifies) */
+	mutable bool IsQualified;
+};
 
-	bool Qualifies(const FAggregatorEvaluateParameters& Parameters) const;
+struct GAMEPLAYABILITIES_API FAggregatorModInfo
+{
+	EGameplayModEvaluationChannel Channel;
+	EGameplayModOp::Type Op;
+	const FAggregatorMod* Mod;
 };
 
 /** Struct representing an individual aggregation channel/depth. Contains mods of all mod op types. */
@@ -99,6 +134,12 @@ struct GAMEPLAYABILITIES_API FAggregatorModChannel
 	 */
 	void AddModsFrom(const FAggregatorModChannel& Other);
 
+	/** runs UpdateQualifies on all mods */
+	void UpdateQualifiesOnAllMods(const FAggregatorEvaluateParameters& Parameters) const;
+
+	/** Helper function for iterating through all mods within a channel */
+	void ForEachMod(FAggregatorModInfo& Info, TFunction<void (const FAggregatorModInfo&) > Func) const;
+
 	/**
 	 * Populate a mapping of channel to corresponding mods
 	 * 
@@ -192,6 +233,9 @@ struct GAMEPLAYABILITIES_API FAggregatorModChannelContainer
 	 */
 	float ReverseEvaluate(float FinalValue, const FAggregatorEvaluateParameters& Parameters) const;
 
+	/** Calls ::UpdateQualifies on each mod */
+	void EvaluateQualificationForAllMods(const FAggregatorEvaluateParameters& Parameters) const;
+
 	/**
 	 * Removes any mods from every channel matching the specified handle
 	 * 
@@ -206,6 +250,9 @@ struct GAMEPLAYABILITIES_API FAggregatorModChannelContainer
 	 */
 	void AddModsFrom(const FAggregatorModChannelContainer& Other);
 
+	/** Helper function for iterating through all mods within the channel container */
+	void ForEachMod( TFunction<void (const FAggregatorModInfo&) > ) const;
+
 	/**
 	 * Populate a mapping of channel to corresponding mods for debugging purposes
 	 * 
@@ -235,7 +282,8 @@ struct GAMEPLAYABILITIES_API FAggregator : public TSharedFromThis<FAggregator>
 	DECLARE_MULTICAST_DELEGATE_OneParam(FOnAggregatorDirty, FAggregator*);
 
 	FAggregator(float InBaseValue=0.f) 
-		: NetUpdateID(0)
+		: EvaluationMetaData(nullptr)
+		, NetUpdateID(0)
 		, BaseValue(InBaseValue)
 		, BroadcastingDirtyCount(0)
 	{}
@@ -275,6 +323,9 @@ struct GAMEPLAYABILITIES_API FAggregator : public TSharedFromThis<FAggregator>
 	/** Evaluates the contribution from the GE associated with ActiveHandle */
 	float EvaluateContribution(const FAggregatorEvaluateParameters& Parameters, FActiveGameplayEffectHandle ActiveHandle) const;
 
+	/** Calls ::UpdateQualifies on each mod. Useful for when you need to manually inspect aggregators */
+	void EvaluateQualificationForAllMods(const FAggregatorEvaluateParameters& Parameters) const;
+
 	void TakeSnapshotOf(const FAggregator& AggToSnapshot);
 
 	FOnAggregatorDirty OnDirty;
@@ -303,6 +354,12 @@ struct GAMEPLAYABILITIES_API FAggregator : public TSharedFromThis<FAggregator>
 	 */
 	void OnActiveEffectDependenciesSwapped(const TMap<FActiveGameplayEffectHandle, FActiveGameplayEffectHandle>& SwappedDependencies);
 
+	/** Helper function for iterating through all mods within the aggregator */
+	void ForEachMod( TFunction<void (const FAggregatorModInfo&) > ) const;
+
+	/** Custom meta data for the aggregator */
+	FAggregatorEvaluateMetaData* EvaluationMetaData;
+
 	/** NetworkID that we had our last update from. Will only be set on clients and is only used to ensure we don't recompute server base value twice. */
 	int32 NetUpdateID;
 
diff --git a/Engine/Plugins/Runtime/GameplayAbilities/Source/GameplayAbilities/Public/GameplayEffectAggregatorLibrary.h b/Engine/Plugins/Runtime/GameplayAbilities/Source/GameplayAbilities/Public/GameplayEffectAggregatorLibrary.h
new file mode 100644
--- /dev/null
+++ b/Engine/Plugins/Runtime/GameplayAbilities/Source/GameplayAbilities/Public/GameplayEffectAggregatorLibrary.h
@@ -0,0 +1,10 @@
+// Copyright 1998-2017 Epic Games, Inc. All Rights Reserved.
+
+#pragma once
+
+#include "GameplayEffectAggregator.h"
+
+struct GAMEPLAYABILITIES_API FAggregatorEvaluateMetaDataLibrary
+{
+	static FAggregatorEvaluateMetaData MostNegativeMod_AllPostiiveMods;
+};
\ No newline at end of file
diff --git a/Engine/Plugins/Runtime/GameplayAbilities/Source/GameplayAbilities/Public/GameplayEffectExecutionCalculation.h b/Engine/Plugins/Runtime/GameplayAbilities/Source/GameplayAbilities/Public/GameplayEffectExecutionCalculation.h
--- a/Engine/Plugins/Runtime/GameplayAbilities/Source/GameplayAbilities/Public/GameplayEffectExecutionCalculation.h
+++ b/Engine/Plugins/Runtime/GameplayAbilities/Source/GameplayAbilities/Public/GameplayEffectExecutionCalculation.h
@@ -107,12 +107,12 @@ public:
 
 	/** 
 	 *	Returns all modifiers for a given captured def. Note the returned list is a direct reference to the internal list stored in the attribute aggregators (possibly a snapshot copy or possibly the "real" ones).
-	 *	Also note the modifiers returned or NOT qualified! You will still want to qualifiy them on an FAggregatorEvaluateParameters yourself. (Having this function internally do it would onyl be possible if this function
-	 *	returned a copy of the mod list which is something we would like to avoid).
+	 *	Also note that all modifiers are returned, even ones that do not qualify for the given InEvalParams. You will still need to check ::Qualifies() on the mods themselves to see if they "count".
+	 *  (we are returning a direct reference to the internal mod list, so we do not want to remove them or force a copy).
 	 *	
-	 *	Consider using ForEachQualifiedAttributeMod when you want to "do something for every qualifier mod"
+	 *	Consider using ForEachQualifiedAttributeMod when you want to "do something for every qualifier mod".
 	 */
-	bool AttemptGatherAttributeMods(const FGameplayEffectAttributeCaptureDefinition& InCaptureDef, OUT TMap<EGameplayModEvaluationChannel, const TArray<FAggregatorMod>*>& OutModMap) const;
+	bool AttemptGatherAttributeMods(const FGameplayEffectAttributeCaptureDefinition& InCaptureDef, const FAggregatorEvaluateParameters& InEvalParams, OUT TMap<EGameplayModEvaluationChannel, const TArray<FAggregatorMod>*>& OutModMap) const;
 
 	/** Runs given TFunction on every qualifier mod for a given AttributeCaptureDefinition */
 	bool ForEachQualifiedAttributeMod(const FGameplayEffectAttributeCaptureDefinition& InCaptureDef, const FAggregatorEvaluateParameters& InEvalParams, TFunction< void(EGameplayModEvaluationChannel, EGameplayModOp::Type, const FAggregatorMod&) >) const;
diff --git a/Engine/Plugins/Runtime/GameplayAbilities/Source/GameplayAbilities/Public/GameplayEffectTypes.h b/Engine/Plugins/Runtime/GameplayAbilities/Source/GameplayAbilities/Public/GameplayEffectTypes.h
--- a/Engine/Plugins/Runtime/GameplayAbilities/Source/GameplayAbilities/Public/GameplayEffectTypes.h
+++ b/Engine/Plugins/Runtime/GameplayAbilities/Source/GameplayAbilities/Public/GameplayEffectTypes.h
@@ -990,16 +990,22 @@ namespace EGameplayCueEvent
 DECLARE_DELEGATE_OneParam(FOnGameplayAttributeEffectExecuted, struct FGameplayModifierEvaluatedData&);
 
 DECLARE_MULTICAST_DELEGATE_TwoParams(FOnGameplayEffectTagCountChanged, const FGameplayTag, int32);
+DECLARE_MULTICAST_DELEGATE_OneParam(FOnGivenActiveGameplayEffectRemoved, const FActiveGameplayEffect&);
 
 DECLARE_MULTICAST_DELEGATE(FOnActiveGameplayEffectRemoved);
 DECLARE_MULTICAST_DELEGATE_OneParam(FOnActiveGameplayEffectRemoved_Info, const FGameplayEffectRemovalInfo&);
+DECLARE_MULTICAST_DELEGATE_ThreeParams(FOnActiveGameplayEffectStackChange, FActiveGameplayEffectHandle, int32 /*NewStackCount*/, int32 /*PreviousStackCount*/);
+DECLARE_MULTICAST_DELEGATE_ThreeParams(FOnActiveGameplayEffectTimeChange, FActiveGameplayEffectHandle, float /*NewStartTime*/, float /*NewDuration*/);
+DECLARE_MULTICAST_DELEGATE_TwoParams(FOnActiveGameplayEffectInhibitionChanged, FActiveGameplayEffectHandle, bool /*bIsInhibited*/);
 
-DECLARE_MULTICAST_DELEGATE_OneParam(FOnGivenActiveGameplayEffectRemoved, const FActiveGameplayEffect&);
-
-DECLARE_MULTICAST_DELEGATE_ThreeParams(FOnActiveGameplayEffectStackChange, FActiveGameplayEffectHandle, int32, int32);
-
-/** FActiveGameplayEffectHandle that is being effect, the start time, duration of the effect */
-DECLARE_MULTICAST_DELEGATE_ThreeParams(FOnActiveGameplayEffectTimeChange, FActiveGameplayEffectHandle, float, float);
+struct FActiveGameplayEffectEvents
+{
+	FOnActiveGameplayEffectRemoved DEPRECATED_OnEffectRemoved;
+	FOnActiveGameplayEffectRemoved_Info OnEffectRemoved;
+	FOnActiveGameplayEffectStackChange OnStackChanged;
+	FOnActiveGameplayEffectTimeChange OnTimeChanged;
+	FOnActiveGameplayEffectInhibitionChanged OnInhibitionChanged;
+};
 
 // This is deprecated, use FOnGameplayAttributeValueChange
 DECLARE_MULTICAST_DELEGATE_TwoParams(FOnGameplayAttributeChange, float, const FGameplayEffectModCallbackData*);
