Copying //UE4/Dev-Framework to //UE4/Dev-Main (Source: //UE4/Dev-Framework @ 4058146)
============================
  MAJOR FEATURES & CHANGES
============================

Change 4007876 by [REDACTED_PERSON]

	Add Inventory Level and Count, accessed as ItemData. Changed various places to read/write this, and switched Souls to be a proper inventory item instead of a variable on player controller
	The player starts with 0 souls, but I hooked up the + on the souls display to grant 50
	Change the way the Store items are calculated in game instance, GetStoreItems now returns hard pointers so it only loads them once at startup
	Add option to reset save data to the options screen, replaced restore purchases as that makes less sense with the current design

Change 4008251 by [REDACTED_PERSON]

	PR #4668: UE-57857: Calling incorrect super function (Contributed by [REDACTED_PERSON])

	Also addresses #jira UE-57869

Change 4008530 by [REDACTED_PERSON]

	Fix hang on startup when async loading component blueprints from game startup code. The component type registry will now load it's meshes on the next tick instead of on construction, as it caused a recursive load issue

Change 4008694 by [REDACTED_PERSON]

	Add bAllowEngineTick option to FLoadingScreenAttributes. If set, it will run the main engine tick while waiting for a manually disabled loading screen to finish displaying. This allows latent actions such as level streaming to complete before stopping the load movie
	This option is disabled by default because game-specific tick functions may be doing unsafe slate operations

Change 4008698 by [REDACTED_PERSON]

	Fix loading screen on map transfer to work properly. There are now options to have the screen be up until it is taken down, changed the game instance to use that
	This depends on engine tick working from the loading movie, a feature I just added

Change 4008699 by [REDACTED_PERSON]

	Add SaveGame flag to gameplay tags so they can be used for native save systems

Change 4008941 by [REDACTED_PERSON]

	Hook up Fireball using new functions that allow applying an effect container spec from a projectile
	Hook mana cost for player abilities, set to 10 but should be balanced and move to a curvetable. Cooldowns are next
	Rename some ability functions to make them shorter

Change 4008943 by [REDACTED_PERSON]

	Make sure we don't drop LOAD_DeferDependencyLoads when loading data via import text

Change 4010465 by [REDACTED_PERSON]

	Make the setting of bWasActive in OnUnregister consistent with SetTemplate. Fixes cases where a deactivated particle system can restart when renaming the owning actor between levels.

Change 4010508 by [REDACTED_PERSON]

	PR #4660: UE-57775: IsEditorOnly components visible in details panel (Contributed by [REDACTED_PERSON])

Change 4010845 by [REDACTED_PERSON]

	Avoid crashing trying to serialize a subobject that was create outside of a transaction

Change 4012148 by [REDACTED_PERSON]

	PR #4552: Significantly optimized performance when refreshing the components tree in the Actor details panel.

Change 4012393 by mason.seay

	Test BP with 512 components

Change 4015966 by mason.seay

	Updated BP to add split pin debugging

Change 4016110 by [REDACTED_PERSON]

	(4.19.2) PR #4678: Fix crash that occurs when the player controller's view target is in a sublevel instance that was unloaded (Contributed by [REDACTED_PERSON])

Change 4016447 by [REDACTED_PERSON]

	Allow Blueprints that implement a native C++ interface declaring one or more BlueprintNativeEvent methods to be nativized.

	Change summary:
	- UHT: Modified FNativeClassHeaderGenerator::ExportNativeFunctionHeader() to emit a PURE_VIRTUAL() expansion in place of "=0" for all BlueprintNativeEvent C++ implementations implicitly declared within a C++ interface class.

Change 4016463 by [REDACTED_PERSON]

	CIS fix - back out changelist 4016447 (temp)

Change 4017382 by [REDACTED_PERSON]

	Prevent LOAD_DeferDependencyLoads from being dropped when we preload an object in another linker

Change 4020602 by paulo.souza

	Lighting improvements and optmizations

Change 4020638 by paulo.souza

	Icons and launch screens on mobile (Android and iOS)

Change 4021340 by [REDACTED_PERSON]

	Fix Map/Set add comments to be accurate, the return value was removed

Change 4021392 by [REDACTED_PERSON]

Change 4021480 by mason.seay

	Reorganized comments and nodes

Change 4025794 by mason.seay

	Cleared all watches

Change 4026141 by [REDACTED_PERSON]

	Removed redundant NumExistingVerts variable/parameter from multiple places in RecastNavMeshGenerator.cpp #UE4

	In rare cases where NumExistingVerts != 0 the code was actually crashing. Found by UDN user:
	https://udn.unrealengine.com/questions/429286/crash-with-dynamic-navmesh.html

Change 4027427 by [REDACTED_PERSON]

	Avoid crash when a subboject reference in the component instance data cache is cleared by a reference collector

Change 4027434 by [REDACTED_PERSON]

	Clean up rest of ability headers, added struct initializers and UPROPERTY for several that were missing them
	Add a constructor for GameplayAbilitySpec that takes an ability class, which makes more sense than forcing the caller to extract a CDO
	Add explicit warning comment to GameplayAbilityTargetActor about it being not recommended
	Add macros to AttributeSet to declare accessors, a version of which is used by all of the Epic internal games

Change 4028656 by [REDACTED_PERSON]

	Added comments and cleaned up ActionRPG code, done with primary features
	Add DefaultSlottedAbilities to Character, I need to update the blueprints to use this
	Add inventory interface that is used instead of having character explicitly cast to player controller

Change 4029079 by paulo.souza

	Fixes to camera rotation when using the AutoMode + UI changes

Change 4030066 by [REDACTED_PERSON]

	Message (interface) call nodes no longer display the skeleton class name in the node subtitle.

	Change summary:
	- Modified UK2Node_Message::GetNodeTitle() to replace outdated title string formatting with the super class implementation for non-menu title queries.

Change 4031843 by [REDACTED_PERSON]

	Action RPG Game full UI overhaul.

	Goals:
	- new layout and art
	- consolidate view to center of screen
	- make buttons appear more like interactible objects
	- update button placement for reach and usability
	- art pass for consistency of visual language (color, iconography, style)

	Still to do:
	- polish on some of the icons (temp art in several places)
	- audio pass
	- environment pass
	- scripting pass for comments/clarity (although everything looks pretty amazing from what I've seen so far, you guys rock)

Change 4033889 by [REDACTED_PERSON]

	Fixed some watches that were incorrectly displayed as not in scope.

	Blueprint pins on some nodes were incorrectly being displayed as not in scope because they were not directly under the active object being debugged.

Change 4033921 by [REDACTED_PERSON]

	Remove unnecessary cast and unused variable.

Change 4034094 by [REDACTED_PERSON]

	Moved the Blueprint bookmarks feature out from under the experimental settings flag.

Change 4035553 by [REDACTED_PERSON]

	Remove unneeded UFUNCTION declaration

Change 4035588 by [REDACTED_PERSON]

	RPG Game:
	- Fixed a couple weapon icons (from temp art to more final version for review)
	- Created 1st pass audio for Guardian enemies (attack, death, roar, swing)
	- added reeeeeeeaally temp environmental audio (WIP)
	- Started on audio for Spider creature (not in engine yet)

Change 4036698 by [REDACTED_PERSON]

	When blueprint debugging during PIE, step over and out commands no longer cause the mouse pointer to jump back to the game viewport after each step.

	Change summary:
	- Modified FKismetDebugUtilities::IsSingleStepping() to include step out/over state checking.
	- Modified LeaveDebuggingMode() to skip the FocusPIEViewport() call when single-stepping.

Change 4038454 by [REDACTED_PERSON]

	Remove unneeded validation code for old UC state system
	Reinstitute proper rejection of UFUNCTION on function in subclass of same name as a ufunction in a parent class.

Change 4038487 by [REDACTED_PERSON]

	RPG Game:
	- Icon work (still a couple placeholder, but almost done!)
	- Audio pass on Guardian creature
	- started audio on Spider creature (WIP)

Change 4040374 by [REDACTED_PERSON]

	When blueprint debugging during PIE, also keep the mouse pointer from jumping back to the game viewport after choosing to stop play.

	Change summary:
	- Modified LeaveDebuggingMode() to include a pending PIE session exit so that clicking Stop in the BP editor also doesn't cause the cursor to jump.
	- Modified FKismetDebugUtilities::IsSingleStepping() to avoid multiple calls to FKismetDebugUtilitiesData::Get() (per review).

Change 4040727 by [REDACTED_PERSON]

	Ability blueprint fixes
	Refactored melee execution to use the item slots for both enemies and players, the goblin has his melee placed in weapon slot 0
	Added cooldowns for skills and fixed it so melee/hit reacts would not interrupt skills and cause things like infinite slomo
	Added some comments

Change 4040812 by [REDACTED_PERSON]

	Fix errors and warnings in blueprint editor tests.

	This came from a UDN thread (https://udn.unrealengine.com/questions/411330/test-systempromotioneditorblueprinteditor-aka-fblu.html).

Change 4041001 by [REDACTED_PERSON]

	Hook up skill cooldown to ui, bump cooldown to 2 seconds

Change 4041021 by [REDACTED_PERSON]

	PR #4703: UE-46077: Remove warning log about removed class variable (Contributed by [REDACTED_PERSON])

Change 4041038 by [REDACTED_PERSON]

	Remove UFUNCTION macros in overridden functions to fix build errors.

Change 4041671 by [REDACTED_PERSON]

	Added calls to delegates when a periodic effect executes a final time as it is being removed.

	PR #4607: Added missing Call to Delegates (Contributed by [REDACTED_PERSON])

Change 4041792 by [REDACTED_PERSON]

	Execution flow, blueprint call stack, and blueprint watchpoint viewer refactoring into a single Blueprint Debugger tab. Call stack viewer now indicates whether call stack is stale, watch point viewer layout now matches clal stack viewer

Change 4041796 by [REDACTED_PERSON]

	SubAnim instance nodes can now orphan pins as expected, the actual fix for this issue is 3997164

Change 4041886 by [REDACTED_PERSON]

	Editable Blueprint events now add 'const' to array type and reference parameter properties when compiled.

	Change summary:
	- Added UK2Node_EditablePinBase::ShouldUseConstRefParams() to replace explicit node type checks.
	- Removed redundant 'const' pin type flag assignment in FBlueprintGraphArgumentLayout::OnRefCheckStateChanged().
	- Modified FBlueprintGraphArgumentLayout::PinInfoChanged() to apply 'const' to array and reference pin types for event nodes.
	- Moved pin type fixup code out of UK2Node_CustomEvent::Serialize() and into UK2Node_EditablePinBase::Serialize().
	- Bumped object version so pin type fixup only needs to run for older assets when loaded in the editor.

Change 4042215 by [REDACTED_PERSON]

	Copy fix for depth of field in to Dev-Framework

Change 4042732 by [REDACTED_PERSON]

	Put the default value for bEnableGestureRecognizer in to BaseInput.ini to make it easier to see there is an option that can be set

Change 4042796 by [REDACTED_PERSON]

Change 4042943 by [REDACTED_PERSON]

	RPG Game:
	- replaced all existing audio
	- set up audio for all animations / matinee
	- will need some polish when real audio comes in, but placeholder is good reference. :)

Change 4043287 by [REDACTED_PERSON]

	Refactor property handle internals so all object path setting goes through SetValueFromFormattedString and move UseSelected to the property handle instead of the value internal

Change 4043396 by [REDACTED_PERSON]

	Fix crash when mousing over a variable that has been deleted and fix breakpoints on nodes in ForEachLoops being skipped

Change 4043708 by paulo.souza

	Enemy progression intial commit + cleanups

Change 4045083 by [REDACTED_PERSON]

	Don't allow new bookmarks to be added when the name field is empty.

Change 4045504 by [REDACTED_PERSON]

	The search bar is now functional in the Blueprint Bookmarks view.

Change 4045516 by [REDACTED_PERSON]

	Fix incorrect original name display when renaming a bookmark in the Blueprint graph view (popup).

Change 4046707 by [REDACTED_PERSON]

	Action RPG Game

	Guardians:
	- Removed delay before grunts attack (so they don't just stand there anymore)
	- Replaced idle animation with idle animation (was a scream, which they did every time they were idle)

	HUD:
	- Fixed skill meter not animating properly
	- Added pulsing reminder around skill button when it's ready and hasn't been used

	Character: (WIP)
	- Fixed missing anim notify in Attack02
	- Added missing notify (and sound) in a couple attacks
	- reduced forward movement component of first couple attacks in combo move

Change 4046868 by [REDACTED_PERSON]

	Reparent blueprints before replacing references when using the 'delete and replace references' tool

Change 4047012 by [REDACTED_PERSON]

	Action RPG Game: Added new sounds for the abilities, made tiny adjustments to two anims to compensate.

Change 4047018 by [REDACTED_PERSON]

	Action RPG Game: Updated pitch and volume on player roll anim to compensate for new assets

Change 4047089 by paulo.souza

	Action RPG Game: Spider boss now uses the Ability System for ranged attacks + Fixes to enemy animations and physics

Change 4049741 by [REDACTED_PERSON]

	Action RPG Game:

	- Set up Wave intro/outro screen
	- Added a some audio stingers (legal approved, no need to replace)
	- Content (music) file organization

Change 4050235 by [REDACTED_PERSON]

	Action RPG:

	- Set up blocking volumes throughout entire map
	- aligned all volumes on major grid lines
	- turned off collision on all exterior rock meshes
	- full rebuild
	(should improve perf, collision, and pathing)

Change 4050440 by paulo.souza

	Action RPG Game: Fixes to Goblin death and hit animations + Nicer Melee and Skill functions

Change 4050910 by paulo.souza

	Action RPG Game: Changed some collision volumes to ignore camera channel traces to not interfere with the character's camera

Change 4050920 by paulo.souza

	Action RPG Game: Wave start and finish screen animation timing fix/polishing

Change 4050921 by paulo.souza

	Action RPG Game: FIX - Enemies could not follow the player when in auto-play mode

Change 4052161 by [REDACTED_PERSON]

	Added player character efforts.
	Adjusted soundcues for VO that plays during slow downs.
	Added anims to support different sounds for mana/health potions

Change 4052932 by [REDACTED_PERSON]

	Add context menu so that we can restore blueprint debugger tabs that have been closed, moved Blueprint Debugger related code out of BlueprintEditorModule as it is now quite significant

Change 4053179 by [REDACTED_PERSON]

	Action RPG Game:

	- New front end (background, logo, buttons, animations)
	- Updated HUD/UI with new art to match updated front end.

Change 4053187 by [REDACTED_PERSON]

	Add method to invoke dynamic force feedback effects from native code without misusing the latent action mechanism.
	Fix latent dynamic force feedback effects not updating their values when instructed to.

Change 4053423 by [REDACTED_PERSON]

	Added Guardian footsteps and concurrency rules for them.
	Added new spawn sound and variant for Guardian, with concurrency rules to keep them in check.
	Added sword swings, adjusted volume per anim.
	Added power up for Firewave.
	Added Player Character footsteps.
	Added whoosh for slo-mo meteors.

Change 4053769 by [REDACTED_PERSON]

	Remove associated local bookmarks when Blueprint assets are deleted.

	Change summary:
	- Added a UBlueprint::BeginDestroy() override (WITH_EDITOR only).
	- Added FBlueprintEditorUtils::RemoveAllLocalBookmarks().

Change 4053771 by [REDACTED_PERSON]

	CIS fix (failed P4 resolve)

Change 4053849 by [REDACTED_PERSON]

	Spider large steps added, adjusted all anims and added them in the anims they weren't in.
	Character collapse added.
	Began work on Intro audio (creature sounds and timing)

Change 4054042 by [REDACTED_PERSON]

	Added Health and Mana cues, they now have seperate anims per item.
	Added all Guardian VO, setup sequences and anims with matching audio.
	Hammer and Axe swings added.
	Level up cue added, adjusted anim.
	Guardian swings and impacts added

Change 4054375 by [REDACTED_PERSON]

	Ensure only that instanced IsEditorOnly components are displayed in the IWCE window

Change 4054518 by [REDACTED_PERSON]

	For now, ignore older bookmark nodes that don't have a corresponding map entry during BP asset deletion.

Change 4054777 by [REDACTED_PERSON]

Change 4054796 by [REDACTED_PERSON]

	Improved watch window.
	 - shows watches from multiple blueprints.
	 - better indication of instances being debugged vs watches that aren't currently valid

Change 4055112 by [REDACTED_PERSON]

	PR #4273: Expose AIController public properties to BP (Contributed by [REDACTED_PERSON])

Change 4055126 by [REDACTED_PERSON]

	Fix shadow variable

Change 4055253 by paulo.souza

	Action RPG Game - Fixes:
	Player can die properly;
	Should not be able to buy Souls;
	Margins for the iPhoneX notch;

Change 4055279 by [REDACTED_PERSON]

	Added a helper function to make it easier to query containers for the presence of a single tag.

	PR #4620: FGameplayTagQuery match single tag shortcut (Contributed by [REDACTED_PERSON])

Change 4055511 by [REDACTED_PERSON]

	Fix it so the Primary Asset load BP nodes can be safely called from a loop like path Async Load nodes. They now take WorldContextObjects, which should automatically convert
	Add UBlueprintAsyncActionBase::RegisterWithGameInstance, when called the action will not be garbage collected until the GameInstance goes away or it is unregistered

Change 4055981 by [REDACTED_PERSON]

	Spider completed

Change 4056011 by [REDACTED_PERSON]

	RPG Game:
	- Fixed textures that weren't power of 2 for mobile
	- Updated main menu screens with better lighting/resolution
	- lighting tweaks to main level
	- Gameplay balance tweaks (should be a bit more difficult now)
	     - more enemies per wave
	     - tighter distribution of enemy levels
	- Differentiated enemies:
	     - Lvl 1 enemies are smaller w/ red effects
	     - Lvl 2 enemies are same size with yellow effects
	     - Lvl 3 enemies are larger with purplish effects
	- Added effects to lvl 3 enemy's weapon (torch)
	- Fixed color distrubution and transparency across buttons on the HUD
	- Fixed button text eating input from buttons
	- maybe some other stuff I forgot. :P

Change 4056192 by [REDACTED_PERSON]

	Fix failure to propagate LOAD_DeferDependencyLoads when loading via FindImportedObject or StaticLoadObjectInternal

Change 4056224 by [REDACTED_PERSON]

	Revert CL 4040812 for this file only. This change was not meant to be checked in.

Change 4056239 by [REDACTED_PERSON]

	Components correctly display again.
	Sprite components of Instanced components do appear. Can't solve that for now.

Change 4056390 by [REDACTED_PERSON]

	Call UGameUserSettings::SetToDefaults() after we've created the instance. This makes sure that classes that overrode this function will have the correct version called.

Change 4056397 by [REDACTED_PERSON]

	Fix several minor issues with the watch window.
	 - Switched to more user friendly names for the instances being debugged
	 - Support copy and paste of multiple lines in the watch window
	 - Deselect whatever was currently selected when we use the hyperlink to jump to the object being debugged.

Change 4056410 by [REDACTED_PERSON]

	Core: Added FUNC_Const to FUNC_FuncInherit

Change 4056515 by [REDACTED_PERSON]

	Fix crash on load during serialization of function entry nodes if the generated class is not yet available.

Change 4056530 by [REDACTED_PERSON]

	Set up soundclasses for all soundcues.
	PSMs for Potions, Abilities, Slomo, and Enemy

Change 4056552 by [REDACTED_PERSON]

Change 4056554 by [REDACTED_PERSON]

	Add a test for property iterator, reorganized the property path helpers test so it shares the structure and is enabled for cooked builds

Change 4056558 by paulo.souza

	Action RPG:
	- Fixed weapon switching bug
	- Added more time to play the game (added per wave)
	- AnimBP now resets to idle animation when in Inventory mode

Change 4056634 by [REDACTED_PERSON]

	Stop error spam about loading null items

Change 4056638 by [REDACTED_PERSON]

	Cleaned up GameInstance handling of loading screens
	Delete some unused assets and consolidate a physical material

Change 4056640 by [REDACTED_PERSON]

	PR #4119: Expose bClientSimulatingViewTarget to BP (Contributed by [REDACTED_PERSON])

Change 4056641 by [REDACTED_PERSON]

	PR #4128: Marked APawn::LastHitBy as BlueprintReadOnly (Contributed by [REDACTED_PERSON])

Change 4056642 by [REDACTED_PERSON]

	PR #4339: Fix a typo in a comment in UPlayerInput::ProcessInputStack (Contributed by [REDACTED_PERSON])

Change 4056644 by [REDACTED_PERSON]

	PR #4462: Fixed a typo in name validation error messages where the name was already in use (Contributed by [REDACTED_PERSON])

Change 4056645 by [REDACTED_PERSON]

	PR #4635: UE-57273: Only call PostProcessWorldToScreen if ProjectWorldToScreen was successful (Contributed by [REDACTED_PERSON])

Change 4056646 by [REDACTED_PERSON]

	Blueprints: Prevent struct properties with an Identical type trait (e.g., FGameplayTagContainer) from showing up as different in a BP diff even if they were unmodified
	PR #4687: (Contributed by [REDACTED_PERSON])

Change 4056659 by [REDACTED_PERSON]

	PR #4244: Fixed TargetPoint's Arrow component being too small to see (Contributed by [REDACTED_PERSON])

Change 4056662 by [REDACTED_PERSON]

	PR #4690: Dirty sprites when double-clicking to change the UV region (Contributed by [REDACTED_PERSON], modified slightly)

Change 4056664 by [REDACTED_PERSON]

	PR #4126: Allow CanRestartPlayer to be BlueprintCallable (Contributed by [REDACTED_PERSON])

Change 4056665 by [REDACTED_PERSON]

	PR #4641: UE-57415: Clamp value for time dilation (Contributed by [REDACTED_PERSON])

Change 4056696 by [REDACTED_PERSON]

	PR #4127: Marked PlayerCanRestart in GameMode as BlueprintCallable (Contributed by [REDACTED_PERSON])

Change 4056716 by [REDACTED_PERSON]

	PR #4192: Fix adding new collision or rendering shapes (box/sphere) being at the wrong position when a sprite is not at the origin in UV space (Contributed by [REDACTED_PERSON], with minor edits)

Change 4056720 by [REDACTED_PERSON]

	PR #4718: Fixed collision generation for tile maps with non-orthogonal projections (Contributed by [REDACTED_PERSON])

Change 4056723 by [REDACTED_PERSON]

	PR #4583: [Paper2D] Fixed yellow tint in tilemap editor & made tile grid color customizable (Contributed by [REDACTED_PERSON])

Change 4056744 by paulo.souza

	Action RPG:
	- Fixed null referenced assets
	- Reinstated the "Add Souls" button (for QA)
	- Reduced some UI images max cook resolution

Change 4056745 by [REDACTED_PERSON]

	UI and Ambient sounds added

Change 4057038 by [REDACTED_PERSON]

	RPG Game:
	- Fixed broken title screen

Change 4057043 by [REDACTED_PERSON]

	RPG Game:
	- Lowered footstep volume

Change 4057071 by [REDACTED_PERSON]

	RPG Game: fixed broken logo/title widget

Change 4057079 by [REDACTED_PERSON]

	Blueprints: Fixing a static analysis error in the watch window

Change 4057112 by [REDACTED_PERSON]

	RPG Game: updated logo (downsized from 2048 to 1024 and improved quality)

Change 4057201 by [REDACTED_PERSON]

	RPG Game: removed music pitch bending from slomo effect (kept ducking) as it sounded very odd in certain circumstances.

Change 4057245 by [REDACTED_PERSON]

	RPG Game: Lowered pitch of sword swing

Change 4057443 by [REDACTED_PERSON]

	Property counts will be different in cooked and uncooked builds due to the editor only properties

Change 4057515 by [REDACTED_PERSON]

	Action RPG:
	- Replaced background image in main  menu with much higher quality art
	- Removed dynamic spotlight that was causing perf hitch in main map
	- Added slight animation to damage number pops
	- Audio tweaks

Change 4020341 by [REDACTED_PERSON]

	(Revised) Allow Blueprints that implement a native C++ interface declaring one or more BlueprintNativeEvent methods to be nativized.

	Change summary:
	- Restored 4016447.
	- UHT: Modified FNativeClassHeaderGenerator::ExportNativeFunctionHeader() to construct a TEnumAsByte as the return value for non-class Enum types when emitting the PURE_VIRTUAL() syntax for BPNE interface methods.
	- Removed existing occurrences of explicit BPNE interface PVM stub implementations as these would otherwise conflict with the PURE_VIRTUAL() expansion.

Change 4024137 by [REDACTED_PERSON]

	Clean up AbilitySystemComponent and GameplayAbility headers. Improved comments, reorganized functions, added virtual to useful places, and removed some dead functions
	Renamed EReplicationMode to EGameplayEffectReplicationMode as the old name was too general for a global enum
	Added UGameplayAbility::GetAbilitySystemComponentFromActorInfo
	Added UAbilitySystemComponent::AddGameplayEventTagContainerDelegate to allow binding a delegate to a gameplay event using a tag container allowing non-exact matches. Added option to AbilityTask_WaitGameplayEvent to allow non exact tags
	Fixed ActionRPG sample and internal games for changes. ActionRPG now only has game-specific ability system code

Change 4035540 by [REDACTED_PERSON]

	Make UWidget::IsHovered virtual

Change 4043467 by [REDACTED_PERSON]

	These are now stored as /game/foo instead of foo and use the in-editor UI instead of the platform directory UI

diff --git a/Engine/Config/BaseEditorPerProjectUserSettings.ini b/Engine/Config/BaseEditorPerProjectUserSettings.ini
--- a/Engine/Config/BaseEditorPerProjectUserSettings.ini
+++ b/Engine/Config/BaseEditorPerProjectUserSettings.ini
@@ -141,7 +141,6 @@ bToolbarCustomization=False
 bBehaviorTreeEditor=False
 bEnableFindAndReplaceReferences =False
 bExampleLayersAndBlends = True
-bEnableBlueprintGraphBookmarks=False
 
 [/Script/UnrealEd.EditorLoadingSavingSettings]
 ; True if we should automatically load a default level at start up
diff --git a/Engine/Config/BaseInput.ini b/Engine/Config/BaseInput.ini
--- a/Engine/Config/BaseInput.ini
+++ b/Engine/Config/BaseInput.ini
@@ -9,6 +9,7 @@ bShowConsoleOnFourFingerTap=true
 bAltEnterTogglesFullscreen=true
 bF11TogglesFullscreen=true
 bRequireCtrlToNavigateAutoComplete=False
+bEnableGestureRecognizer=false
 +ConsoleKeys=Tilde
 
 ;-----------------------------------------------------------------------------------------
diff --git a/Engine/Plugins/2D/Paper2D/Source/Paper2D/Classes/PaperTileMap.h b/Engine/Plugins/2D/Paper2D/Source/Paper2D/Classes/PaperTileMap.h
--- a/Engine/Plugins/2D/Paper2D/Source/Paper2D/Classes/PaperTileMap.h
+++ b/Engine/Plugins/2D/Paper2D/Source/Paper2D/Classes/PaperTileMap.h
@@ -119,6 +119,34 @@ public:
 	/** The background color displayed in the tile map editor */
 	UPROPERTY(Category=Setup, EditAnywhere, meta=(HideAlphaChannel))
 	FLinearColor BackgroundColor;
+
+	/** The color of the tile grid */
+	UPROPERTY(Category=Setup, EditAnywhere, meta=(HideAlphaChannel))
+	FLinearColor TileGridColor;
+
+	/** The color of the multi tile grid */
+	UPROPERTY(Category=Setup, EditAnywhere, meta=(HideAlphaChannel))
+	FLinearColor MultiTileGridColor;
+
+	/** Number of tiles the multi tile grid spans horizontally. 0 removes vertical lines */
+	UPROPERTY(Category=Setup, EditAnywhere, meta=(ClampMin=0))
+	int32 MultiTileGridWidth;
+
+	/** Number of tiles the multi tile grid spans vertically. 0 removes horizontal lines */
+	UPROPERTY(Category=Setup, EditAnywhere, meta=(ClampMin=0))
+	int32 MultiTileGridHeight;
+
+	/** Number of tiles the multi tile grid is shifted to the right */
+	UPROPERTY(Category=Setup, EditAnywhere)
+	int32 MultiTileGridOffsetX;
+
+	/** Number of tiles the multi tile grid is shifted downwards */
+	UPROPERTY(Category=Setup, EditAnywhere)
+	int32 MultiTileGridOffsetY;
+
+	/** The color of the layer grid */
+	UPROPERTY(Category=Setup, EditAnywhere, meta=(HideAlphaChannel))
+	FLinearColor LayerGridColor;
 #endif
 
 	/** The naming index to start at when trying to create a new layer */
diff --git a/Engine/Plugins/2D/Paper2D/Source/Paper2D/Classes/PaperTileMapComponent.h b/Engine/Plugins/2D/Paper2D/Source/Paper2D/Classes/PaperTileMapComponent.h
--- a/Engine/Plugins/2D/Paper2D/Source/Paper2D/Classes/PaperTileMapComponent.h
+++ b/Engine/Plugins/2D/Paper2D/Source/Paper2D/Classes/PaperTileMapComponent.h
@@ -97,6 +97,14 @@ public:
 	// Should this component show an outline around the first layer when the component is not selected?
 	UPROPERTY(Category=Rendering, EditAnywhere)
 	bool bShowOutlineWhenUnselected;
+
+	// Should this component show a tile grid when the component is not selected?
+	UPROPERTY(Category=Rendering, EditAnywhere)
+	bool bShowPerTileGridWhenUnselected;
+
+	// Should this component show an outline around each layer when the component is not selected?
+	UPROPERTY(Category=Rendering, EditAnywhere)
+	bool bShowPerLayerGridWhenUnselected;
 #endif
 
 protected:
diff --git a/Engine/Plugins/2D/Paper2D/Source/Paper2D/Private/PaperSprite.cpp b/Engine/Plugins/2D/Paper2D/Source/Paper2D/Private/PaperSprite.cpp
--- a/Engine/Plugins/2D/Paper2D/Source/Paper2D/Private/PaperSprite.cpp
+++ b/Engine/Plugins/2D/Paper2D/Source/Paper2D/Private/PaperSprite.cpp
@@ -343,6 +343,9 @@ void UPaperSprite::ExtractSourceRegionFromTexturePoint(const FVector2D& SourcePo
 		{
 			if (Dimension.X > 0 && Dimension.Y > 0)
 			{
+				// Mark the package as dirty
+				Modify();
+
 				SourceUV = FVector2D(Origin.X, Origin.Y);
 				SourceDimension = FVector2D(Dimension.X, Dimension.Y);
 				PostEditChange();
diff --git a/Engine/Plugins/2D/Paper2D/Source/Paper2D/Private/PaperTileLayer.cpp b/Engine/Plugins/2D/Paper2D/Source/Paper2D/Private/PaperTileLayer.cpp
--- a/Engine/Plugins/2D/Paper2D/Source/Paper2D/Private/PaperTileLayer.cpp
+++ b/Engine/Plugins/2D/Paper2D/Source/Paper2D/Private/PaperTileLayer.cpp
@@ -239,7 +239,32 @@ void UPaperTileLayer::AugmentBodySetup(UBodySetup* ShapeBodySetup, float RenderS
 					const int32 Flags = CellInfo.GetFlagsAsIndex();
 
 					const FTransform& LocalTransform = TilePermutationTransforms[Flags];
-					const FVector2D CellOffset(TileWidth * CellX, TileHeight * -CellY);
+					
+					// When building the collison we must consider setting the correct offsets for the projection mode.
+					// This must be inline with the way you are making any changes too the 2D grid.
+					FVector ProjectionOffset;
+					switch (TileMap->ProjectionMode)
+					{
+					case ETileMapProjectionMode::Orthogonal: // General offset used by default.
+					default:
+						ProjectionOffset.X = TileWidth * CellX;
+						ProjectionOffset.Y = TileHeight * -CellY;
+					break;
+
+					case ETileMapProjectionMode::IsometricDiamond: // Offset to compensate for the Isometric Scatter of the grid.
+						ProjectionOffset.X = (CellX * TileWidth * 0.5f) + -(CellY * TileWidth * 0.5f);
+						ProjectionOffset.Y = -(CellX * TileHeight * 0.5f) - (CellY * TileHeight * 0.5f);
+					break;
+
+					case ETileMapProjectionMode::IsometricStaggered: // Isometric & Hexagonal share the same offsets.
+					case ETileMapProjectionMode::HexagonalStaggered:
+						ProjectionOffset.X = (CellX * TileWidth) + CellY % 2 * (TileWidth * 0.5f);
+						ProjectionOffset.Y = -CellY * (TileHeight * 0.5f);
+					break;
+					}
+
+					const FVector2D CellOffset(ProjectionOffset.X, ProjectionOffset.Y);
+					
 					CollisionBuilder.SetCellOffset(CellOffset, LocalTransform);
 
 					CollisionBuilder.ProcessGeometry(CellMetadata->CollisionData);
diff --git a/Engine/Plugins/2D/Paper2D/Source/Paper2D/Private/PaperTileMap.cpp b/Engine/Plugins/2D/Paper2D/Source/Paper2D/Private/PaperTileMap.cpp
--- a/Engine/Plugins/2D/Paper2D/Source/Paper2D/Private/PaperTileMap.cpp
+++ b/Engine/Plugins/2D/Paper2D/Source/Paper2D/Private/PaperTileMap.cpp
@@ -40,6 +40,13 @@ UPaperTileMap::UPaperTileMap(const FObjectInitializer& ObjectInitializer)
 #if WITH_EDITORONLY_DATA
 	SelectedLayerIndex = INDEX_NONE;
 	BackgroundColor = FColor(55, 55, 55);
+	TileGridColor = FLinearColor::Gray;
+	MultiTileGridColor = FLinearColor::Red;
+	MultiTileGridWidth = 0;
+	MultiTileGridHeight = 0;
+	MultiTileGridOffsetX = 0;
+	MultiTileGridOffsetY = 0;
+	LayerGridColor = FLinearColor::White;
 #endif
 
 	LayerNameIndex = 0;
diff --git a/Engine/Plugins/2D/Paper2D/Source/Paper2D/Private/PaperTileMapComponent.cpp b/Engine/Plugins/2D/Paper2D/Source/Paper2D/Private/PaperTileMapComponent.cpp
--- a/Engine/Plugins/2D/Paper2D/Source/Paper2D/Private/PaperTileMapComponent.cpp
+++ b/Engine/Plugins/2D/Paper2D/Source/Paper2D/Private/PaperTileMapComponent.cpp
@@ -40,7 +40,9 @@ UPaperTileMapComponent::UPaperTileMapComponent(const FObjectInitializer& ObjectI
 
 #if WITH_EDITORONLY_DATA
 	bShowPerTileGridWhenSelected = true;
+	bShowPerTileGridWhenUnselected = false;
 	bShowPerLayerGridWhenSelected = true;
+	bShowPerLayerGridWhenUnselected = false;
 	bShowOutlineWhenUnselected = true;
 #endif
 
diff --git a/Engine/Plugins/2D/Paper2D/Source/Paper2D/Private/PaperTileMapRenderSceneProxy.cpp b/Engine/Plugins/2D/Paper2D/Source/Paper2D/Private/PaperTileMapRenderSceneProxy.cpp
--- a/Engine/Plugins/2D/Paper2D/Source/Paper2D/Private/PaperTileMapRenderSceneProxy.cpp
+++ b/Engine/Plugins/2D/Paper2D/Source/Paper2D/Private/PaperTileMapRenderSceneProxy.cpp
@@ -18,8 +18,10 @@ DECLARE_CYCLE_STAT(TEXT("Tile Map Editor Grid"), STAT_TileMap_EditorWireDrawing,
 FPaperTileMapRenderSceneProxy::FPaperTileMapRenderSceneProxy(const UPaperTileMapComponent* InComponent)
 	: FPaperRenderSceneProxy(InComponent)
 #if WITH_EDITOR
-	, bShowPerTileGrid(false)
-	, bShowPerLayerGrid(false)
+	, bShowPerTileGridWhenSelected(false)
+	, bShowPerTileGridWhenUnselected(false)
+	, bShowPerLayerGridWhenSelected(false)
+	, bShowPerLayerGridWhenUnselected(false)
 	, bShowOutlineWhenUnselected(false)
 #endif
 	, TileMap(nullptr)
@@ -38,8 +40,10 @@ FPaperTileMapRenderSceneProxy::FPaperTileMapRenderSceneProxy(const UPaperTileMap
 	MaterialRelevance = InComponent->GetMaterialRelevance(GetScene().GetFeatureLevel());
 
 #if WITH_EDITORONLY_DATA
-	bShowPerTileGrid = InComponent->bShowPerTileGridWhenSelected;
-	bShowPerLayerGrid = InComponent->bShowPerLayerGridWhenSelected;
+	bShowPerTileGridWhenSelected = InComponent->bShowPerTileGridWhenSelected;
+	bShowPerTileGridWhenUnselected = InComponent->bShowPerTileGridWhenUnselected;
+	bShowPerLayerGridWhenSelected = InComponent->bShowPerLayerGridWhenSelected;
+	bShowPerLayerGridWhenUnselected = InComponent->bShowPerLayerGridWhenUnselected;
 	bShowOutlineWhenUnselected = InComponent->bShowOutlineWhenUnselected;
 #endif
 }
@@ -78,8 +82,9 @@ void FPaperTileMapRenderSceneProxy::DrawBoundsForLayer(FPrimitiveDrawInterface*
 	PDI->DrawLine(BL, TL, Color, SDPG_Foreground, 0.0f, WireDepthBias);
 }
 
-void FPaperTileMapRenderSceneProxy::DrawNormalGridLines(FPrimitiveDrawInterface* PDI, const FLinearColor& Color, int32 LayerIndex) const
+void FPaperTileMapRenderSceneProxy::DrawNormalGridLines(FPrimitiveDrawInterface* PDI, const FLinearColor& PerTileColor, const FLinearColor& MultiTileColor, int32 MultiTileGridWidth, int32 MultiTileGridHeight, int32 MultiTileGridOffsetX, int32 MultiTileGridOffsetY, int32 LayerIndex) const
 {
+	FLinearColor Color = PerTileColor;
 	const FMatrix& LocalToWorldMat = GetLocalToWorld();
 	const uint8 DPG = SDPG_Foreground;//GetDepthPriorityGroup(View);
 
@@ -92,6 +97,15 @@ void FPaperTileMapRenderSceneProxy::DrawNormalGridLines(FPrimitiveDrawInterface*
 		X = TileMap->MapWidth;
 		const FVector End(TileMap->GetTilePositionInLocalSpace(X, Y, LayerIndex));
 
+		if (MultiTileGridHeight > 0 && int32(Y - MultiTileGridOffsetY) % int32(MultiTileGridHeight) == 0)
+		{
+			Color = MultiTileColor;
+		}
+		else
+		{
+			Color = PerTileColor;
+		}
+
 		PDI->DrawLine(LocalToWorldMat.TransformPosition(Start), LocalToWorldMat.TransformPosition(End), Color, DPG, 0.0f, WireDepthBias);
 	}
 
@@ -104,12 +118,22 @@ void FPaperTileMapRenderSceneProxy::DrawNormalGridLines(FPrimitiveDrawInterface*
 		Y = TileMap->MapHeight;
 		const FVector End(TileMap->GetTilePositionInLocalSpace(X, Y, LayerIndex));
 
+		if (MultiTileGridWidth > 0 && int32(X - MultiTileGridOffsetX) % int32(MultiTileGridWidth) == 0)
+		{
+			Color = MultiTileColor;
+		}
+		else
+		{
+			Color = PerTileColor;
+		}
+
 		PDI->DrawLine(LocalToWorldMat.TransformPosition(Start), LocalToWorldMat.TransformPosition(End), Color, DPG, 0.0f, WireDepthBias);
 	}
 }
 
-void FPaperTileMapRenderSceneProxy::DrawStaggeredGridLines(FPrimitiveDrawInterface* PDI, const FLinearColor& Color, int32 LayerIndex) const
+void FPaperTileMapRenderSceneProxy::DrawStaggeredGridLines(FPrimitiveDrawInterface* PDI, const FLinearColor& PerTileColor, const FLinearColor& MultiTileColor, int32 MultiTileGridWidth, int32 MultiTileGridHeight, int32 MultiTileGridOffsetX, int32 MultiTileGridOffsetY, int32 LayerIndex) const
 {
+	FLinearColor Color = PerTileColor;
 	TArray<FVector> Poly;
 	Poly.Empty(4);
 
@@ -162,6 +186,15 @@ void FPaperTileMapRenderSceneProxy::DrawStaggeredGridLines(FPrimitiveDrawInterfa
 		TileMap->GetTilePolygon(XBottom, YBottom, LayerIndex, Poly);
 		const FVector LSB = Poly[2];
 
+		if (MultiTileGridHeight > 0 && int32(X - MultiTileGridOffsetY) % int32(MultiTileGridHeight) == 0)
+		{
+			Color = MultiTileColor;
+		}
+		else
+		{
+			Color = PerTileColor;
+		}
+
 		PDI->DrawLine(LocalToWorldMat.TransformPosition(LSA), LocalToWorldMat.TransformPosition(LSB), Color, DPG, 0.0f, WireDepthBias);
 	}
 
@@ -192,12 +225,22 @@ void FPaperTileMapRenderSceneProxy::DrawStaggeredGridLines(FPrimitiveDrawInterfa
 		TileMap->GetTilePolygon(XBottom, YBottom, LayerIndex, Poly);
 		const FVector LSB = Poly[3];
 
+		if (MultiTileGridWidth > 0 && int32(X - MultiTileGridOffsetX) % int32(MultiTileGridWidth) == 0)
+		{
+			Color = MultiTileColor;
+		}
+		else
+		{
+			Color = PerTileColor;
+		}
+
 		PDI->DrawLine(LocalToWorldMat.TransformPosition(LSA), LocalToWorldMat.TransformPosition(LSB), Color, DPG, 0.0f, WireDepthBias);
 	}
 }
 
-void FPaperTileMapRenderSceneProxy::DrawHexagonalGridLines(FPrimitiveDrawInterface* PDI, const FLinearColor& Color, int32 LayerIndex) const
+void FPaperTileMapRenderSceneProxy::DrawHexagonalGridLines(FPrimitiveDrawInterface* PDI, const FLinearColor& PerTileColor, const FLinearColor& MultiTileColor, int32 MultiTileGridWidth, int32 MultiTileGridHeight, int32 MultiTileGridOffsetX, int32 MultiTileGridOffsetY, int32 LayerIndex) const
 {
+	FLinearColor Color = FLinearColor::White;
 	//@TODO: This isn't very efficient
 	const FMatrix& LocalToWorldMat = GetLocalToWorld();
 	const uint8 DPG = SDPG_Foreground;//GetDepthPriorityGroup(View);
@@ -292,55 +335,105 @@ void FPaperTileMapRenderSceneProxy::GetDynamicMeshElements(const TArray<const FS
 
 				const uint8 DPG = SDPG_Foreground;//GetDepthPriorityGroup(View);
 
-				// Draw separation wires if selected
-				const FLinearColor OverrideColor = GetSelectionColor(FLinearColor::White, bShowAsSelected, IsHovered(), /*bUseOverlayIntensity=*/ false);
 
 				// Draw the debug outline
 				if (bEffectivelySelected)
 				{
 					const int32 SelectedLayerIndex = (OnlyLayerIndex != INDEX_NONE) ? OnlyLayerIndex : TileMap->SelectedLayerIndex;
 
-					if (bShowPerLayerGrid)
+					// Draw separation wires if selected
+					const FLinearColor OverrideColor = GetSelectionColor(FLinearColor::White, bShowAsSelected, IsHovered(), /*bUseOverlayIntensity=*/ false);
+
+					if (bShowPerLayerGridWhenSelected)
 					{
 						if (OnlyLayerIndex == INDEX_NONE)
 						{
 							// Draw a bound for every layer but the selected one (and even that one if the per-tile grid is off)
 							for (int32 LayerIndex = 0; LayerIndex < TileMap->TileLayers.Num(); ++LayerIndex)
 							{
-								if ((LayerIndex != SelectedLayerIndex) || !bShowPerTileGrid)
+								if ((LayerIndex != SelectedLayerIndex) || !bShowPerTileGridWhenSelected)
 								{
 									DrawBoundsForLayer(PDI, OverrideColor, LayerIndex);
 								}
 							}
 						}
-						else if (!bShowPerTileGrid)
+						else if (!bShowPerTileGridWhenSelected)
 						{
 							DrawBoundsForLayer(PDI, OverrideColor, OnlyLayerIndex);
 						}
 					}
 
-					if (bShowPerTileGrid && (SelectedLayerIndex != INDEX_NONE))
+					if (bShowPerTileGridWhenSelected && (SelectedLayerIndex != INDEX_NONE))
 					{
 						switch (TileMap->ProjectionMode)
 						{
 						default:
 						case ETileMapProjectionMode::Orthogonal:
 						case ETileMapProjectionMode::IsometricDiamond:
-							DrawNormalGridLines(PDI, OverrideColor, SelectedLayerIndex);
+							DrawNormalGridLines(PDI, OverrideColor, OverrideColor, 0, 0, 0, 0, SelectedLayerIndex);
 							break;
 						case ETileMapProjectionMode::IsometricStaggered:
-							DrawStaggeredGridLines(PDI, OverrideColor, SelectedLayerIndex);
+							DrawStaggeredGridLines(PDI, OverrideColor, OverrideColor, 0, 0, 0, 0, SelectedLayerIndex);
 							break;
 						case ETileMapProjectionMode::HexagonalStaggered:
-							DrawHexagonalGridLines(PDI, OverrideColor, SelectedLayerIndex);
+							DrawHexagonalGridLines(PDI, OverrideColor, OverrideColor, 0, 0, 0, 0, SelectedLayerIndex);
 							break;
 						}
 					}
 				}
-				else if (View->Family->EngineShowFlags.Grid && bShowOutlineWhenUnselected)
+				else if (View->Family->EngineShowFlags.Grid)
 				{
-					// Draw a layer rectangle even when not selected, so you can see where the tile map is in the editor
-					DrawBoundsForLayer(PDI, GetWireframeColor(), /*LayerIndex=*/ (OnlyLayerIndex != INDEX_NONE) ? OnlyLayerIndex : 0);
+					const int32 SelectedLayerIndex = (OnlyLayerIndex != INDEX_NONE) ? OnlyLayerIndex : TileMap->SelectedLayerIndex;
+
+					if (bShowOutlineWhenUnselected)
+					{
+						// Draw a layer rectangle even when not selected, so you can see where the tile map is in the editor
+						DrawBoundsForLayer(PDI, GetWireframeColor(), /*LayerIndex=*/ (OnlyLayerIndex != INDEX_NONE) ? OnlyLayerIndex : 0);
+					}
+					if (bShowPerLayerGridWhenUnselected)
+					{
+						const FLinearColor LayerGridColor = TileMap->LayerGridColor;
+
+						if (OnlyLayerIndex == INDEX_NONE)
+						{
+							// Draw a bound for every layer but the selected one (and even that one if the per-tile grid is off)
+							for (int32 LayerIndex = 0; LayerIndex < TileMap->TileLayers.Num(); ++LayerIndex)
+							{
+								if ((LayerIndex != SelectedLayerIndex) || !bShowPerTileGridWhenUnselected)
+								{
+									DrawBoundsForLayer(PDI, LayerGridColor, LayerIndex);
+								}
+							}
+						}
+						else if (!bShowPerTileGridWhenUnselected)
+						{
+							DrawBoundsForLayer(PDI, LayerGridColor, OnlyLayerIndex);
+						}
+					}
+					if (bShowPerTileGridWhenUnselected)
+					{
+						const FLinearColor TileGridColor = TileMap->TileGridColor;
+						const FLinearColor MultiTileGridColor = TileMap->MultiTileGridColor;
+						int32 MultiTileGridWidth = TileMap->MultiTileGridWidth;
+						int32 MultiTileGridHeight = TileMap->MultiTileGridHeight;
+						int32 MultiTileGridOffsetX = TileMap->MultiTileGridOffsetX;
+						int32 MultiTileGridOffsetY = TileMap->MultiTileGridOffsetY;
+
+						switch (TileMap->ProjectionMode)
+						{
+						default:
+						case ETileMapProjectionMode::Orthogonal:
+						case ETileMapProjectionMode::IsometricDiamond:
+							DrawNormalGridLines(PDI, TileGridColor, MultiTileGridColor, MultiTileGridWidth, MultiTileGridHeight, MultiTileGridOffsetX, MultiTileGridOffsetY, SelectedLayerIndex);
+							break;
+						case ETileMapProjectionMode::IsometricStaggered:
+							DrawStaggeredGridLines(PDI, TileGridColor, MultiTileGridColor, MultiTileGridWidth, MultiTileGridHeight, MultiTileGridOffsetX, MultiTileGridOffsetY, SelectedLayerIndex);
+							break;
+						case ETileMapProjectionMode::HexagonalStaggered:
+							DrawHexagonalGridLines(PDI, TileGridColor, MultiTileGridColor, MultiTileGridWidth, MultiTileGridHeight, MultiTileGridOffsetX, MultiTileGridOffsetY, SelectedLayerIndex);
+							break;
+						}
+					}
 				}
 #endif
 			}
diff --git a/Engine/Plugins/2D/Paper2D/Source/Paper2D/Private/PaperTileMapRenderSceneProxy.h b/Engine/Plugins/2D/Paper2D/Source/Paper2D/Private/PaperTileMapRenderSceneProxy.h
--- a/Engine/Plugins/2D/Paper2D/Source/Paper2D/Private/PaperTileMapRenderSceneProxy.h
+++ b/Engine/Plugins/2D/Paper2D/Source/Paper2D/Private/PaperTileMapRenderSceneProxy.h
@@ -32,15 +32,17 @@ protected:
 	FPaperTileMapRenderSceneProxy(const UPaperTileMapComponent* InComponent);
 
 	void DrawBoundsForLayer(FPrimitiveDrawInterface* PDI, const FLinearColor& Color, int32 LayerIndex) const;
-	void DrawNormalGridLines(FPrimitiveDrawInterface* PDI, const FLinearColor& Color, int32 LayerIndex) const;
-	void DrawStaggeredGridLines(FPrimitiveDrawInterface* PDI, const FLinearColor& Color, int32 LayerIndex) const;
-	void DrawHexagonalGridLines(FPrimitiveDrawInterface* PDI, const FLinearColor& Color, int32 LayerIndex) const;
+	void DrawNormalGridLines(FPrimitiveDrawInterface* PDI, const FLinearColor& PerTileColor, const FLinearColor& MultiTileColor, int32 MultiTileGridWidth, int32 MultiTileGridHeight, int32 MultiTileGridOffsetX, int32 MultiTileGridOffsetY, int32 LayerIndex) const;
+	void DrawStaggeredGridLines(FPrimitiveDrawInterface* PDI, const FLinearColor& PerTileColor, const FLinearColor& MultiTileColor, int32 MultiTileGridWidth, int32 MultiTileGridHeight, int32 MultiTileGridOffsetX, int32 MultiTileGridOffsetY, int32 LayerIndex) const;
+	void DrawHexagonalGridLines(FPrimitiveDrawInterface* PDI, const FLinearColor& PerTileColor, const FLinearColor& MultiTileColor, int32 MultiTileGridWidth, int32 MultiTileGridHeight, int32 MultiTileGridOffsetX, int32 MultiTileGridOffsetY, int32 LayerIndex) const;
 
 protected:
 
 #if WITH_EDITOR
-	bool bShowPerTileGrid;
-	bool bShowPerLayerGrid;
+	bool bShowPerTileGridWhenSelected;
+	bool bShowPerTileGridWhenUnselected;
+	bool bShowPerLayerGridWhenSelected;
+	bool bShowPerLayerGridWhenUnselected;
 	bool bShowOutlineWhenUnselected;
 #endif
 
diff --git a/Engine/Plugins/2D/Paper2D/Source/Paper2DEditor/Private/SpriteEditor/SpriteEditorViewportClient.cpp b/Engine/Plugins/2D/Paper2D/Source/Paper2DEditor/Private/SpriteEditor/SpriteEditorViewportClient.cpp
--- a/Engine/Plugins/2D/Paper2D/Source/Paper2DEditor/Private/SpriteEditor/SpriteEditorViewportClient.cpp
+++ b/Engine/Plugins/2D/Paper2D/Source/Paper2DEditor/Private/SpriteEditor/SpriteEditorViewportClient.cpp
@@ -633,10 +633,10 @@ void FSpriteEditorViewportClient::Tick(float DeltaSeconds)
 		}
 
 		const FVector2D BoxSize(Sprite->GetSourceSize());
-		const FVector2D BoxLocation(Sprite->GetSourceUV() + (BoxSize * 0.5f));
+		const FVector2D BoxLocation(Sprite->GetSourceUV());
 		FBox2D SpriteBounds(ForceInitToZero);
-		SpriteBounds.Min = BoxSize - BoxLocation * 0.5f;
-		SpriteBounds.Max = BoxSize + BoxLocation * 0.5f;
+		SpriteBounds.Min = BoxLocation;
+		SpriteBounds.Max = BoxLocation + BoxSize;
 
 		if (FSpriteGeometryEditMode* GeometryEditMode = ModeTools->GetActiveModeTyped<FSpriteGeometryEditMode>(FSpriteGeometryEditMode::EM_SpriteGeometry))
 		{
diff --git a/Engine/Plugins/2D/Paper2D/Source/Paper2DEditor/Private/TileMapEditing/EdModeTileMap.cpp b/Engine/Plugins/2D/Paper2D/Source/Paper2DEditor/Private/TileMapEditing/EdModeTileMap.cpp
--- a/Engine/Plugins/2D/Paper2D/Source/Paper2DEditor/Private/TileMapEditing/EdModeTileMap.cpp
+++ b/Engine/Plugins/2D/Paper2D/Source/Paper2DEditor/Private/TileMapEditing/EdModeTileMap.cpp
@@ -170,7 +170,9 @@ void FEdModeTileMap::Enter()
 	CursorPreviewComponent->TileMap->InitializeNewEmptyTileMap();
 	CursorPreviewComponent->TranslucencySortPriority = 99999;
 	CursorPreviewComponent->bShowPerTileGridWhenSelected = false;
+	CursorPreviewComponent->bShowPerTileGridWhenUnselected = false;
 	CursorPreviewComponent->bShowPerLayerGridWhenSelected = false;
+	CursorPreviewComponent->bShowPerLayerGridWhenUnselected = false;
 	CursorPreviewComponent->bShowOutlineWhenUnselected = false;
 	CursorPreviewComponent->UpdateBounds();
 	CursorPreviewComponent->AddToRoot();
diff --git a/Engine/Plugins/2D/Paper2D/Source/Paper2DEditor/Private/TileMapEditing/TileMapEditorSettings.cpp b/Engine/Plugins/2D/Paper2D/Source/Paper2DEditor/Private/TileMapEditing/TileMapEditorSettings.cpp
--- a/Engine/Plugins/2D/Paper2D/Source/Paper2DEditor/Private/TileMapEditing/TileMapEditorSettings.cpp
+++ b/Engine/Plugins/2D/Paper2D/Source/Paper2DEditor/Private/TileMapEditing/TileMapEditorSettings.cpp
@@ -8,5 +8,12 @@
 UTileMapEditorSettings::UTileMapEditorSettings()
 	: DefaultBackgroundColor(55, 55, 55)
 	, bShowGridByDefault(false)
+	, DefaultTileGridColor(150, 150, 150)
+	, DefaultMultiTileGridColor(255, 0, 0)
+	, DefaultMultiTileGridWidth(0)
+	, DefaultMultiTileGridHeight(0)
+	, DefaultMultiTileGridOffsetX(0)
+	, DefaultMultiTileGridOffsetY(0)
+	, DefaultLayerGridColor(255, 255, 0)
 {
 }
diff --git a/Engine/Plugins/2D/Paper2D/Source/Paper2DEditor/Private/TileMapEditing/TileMapEditorSettings.h b/Engine/Plugins/2D/Paper2D/Source/Paper2DEditor/Private/TileMapEditing/TileMapEditorSettings.h
--- a/Engine/Plugins/2D/Paper2D/Source/Paper2DEditor/Private/TileMapEditing/TileMapEditorSettings.h
+++ b/Engine/Plugins/2D/Paper2D/Source/Paper2DEditor/Private/TileMapEditing/TileMapEditorSettings.h
@@ -21,6 +21,34 @@ public:
 	FColor DefaultBackgroundColor;
 
 	/** Should the grid be shown by default when the editor is opened? */
-	UPROPERTY(config, EditAnywhere, Category=Background)
+	UPROPERTY(config, EditAnywhere, Category=Grid)
 	bool bShowGridByDefault;
+
+	/** Default per tile grid color for new tile map assets */
+	UPROPERTY(config, EditAnywhere, Category=Grid, meta=(HideAlphaChannel))
+	FColor DefaultTileGridColor;
+
+	/** Default multi tile grid color for new tile map assets */
+	UPROPERTY(config, EditAnywhere, Category=Grid, meta=(HideAlphaChannel))
+	FColor DefaultMultiTileGridColor;
+
+	/** Default number of tiles the multi tile grid spans horizontally for new tile map assets. 0 removes vertical lines */
+	UPROPERTY(config, EditAnywhere, Category=Grid)
+	int32 DefaultMultiTileGridWidth;
+
+	/** Default number of tiles the multi tile grid spans vertically for new tile map assets. 0 removes horizontal lines */
+	UPROPERTY(config, EditAnywhere, Category=Grid)
+	int32 DefaultMultiTileGridHeight;
+
+	/** Default number of tiles the multi tile grid is shifted to the right for new tile map assets */
+	UPROPERTY(config, EditAnywhere, Category=Grid)
+	int32 DefaultMultiTileGridOffsetX;
+
+	/** Default number of tiles the multi tile grid is shifted downwards for new tile map assets */
+	UPROPERTY(config, EditAnywhere, Category=Grid)
+	int32 DefaultMultiTileGridOffsetY;
+
+	/** Default layer grid color for new tile map assets */
+	UPROPERTY(config, EditAnywhere, Category=Grid, meta=(HideAlphaChannel))
+	FColor DefaultLayerGridColor;
 };
diff --git a/Engine/Plugins/2D/Paper2D/Source/Paper2DEditor/Private/TileMapEditing/TileMapEditorViewportClient.cpp b/Engine/Plugins/2D/Paper2D/Source/Paper2DEditor/Private/TileMapEditing/TileMapEditorViewportClient.cpp
--- a/Engine/Plugins/2D/Paper2D/Source/Paper2DEditor/Private/TileMapEditing/TileMapEditorViewportClient.cpp
+++ b/Engine/Plugins/2D/Paper2D/Source/Paper2DEditor/Private/TileMapEditing/TileMapEditorViewportClient.cpp
@@ -52,7 +52,6 @@ FTileMapEditorViewportClient::FTileMapEditorViewportClient(TWeakPtr<FTileMapEdit
 		RenderTileMapComponent = NewObject<UPaperTileMapComponent>();
 		UPaperTileMap* TileMap = GetTileMapBeingEdited();
 		RenderTileMapComponent->TileMap = TileMap;
-		RenderTileMapComponent->SelectionOverrideDelegate = UPrimitiveComponent::FSelectionOverride::CreateLambda([](const UPrimitiveComponent*) { return true; });
 
 		PreviewScene->AddComponent(RenderTileMapComponent, FTransform::Identity);
 	}
@@ -238,25 +237,25 @@ FLinearColor FTileMapEditorViewportClient::GetBackgroundColor() const
 void FTileMapEditorViewportClient::ToggleShowTileGrid()
 {
 	FComponentReregisterContext RefreshComponentHelper(RenderTileMapComponent);
-	RenderTileMapComponent->bShowPerTileGridWhenSelected = !RenderTileMapComponent->bShowPerTileGridWhenSelected;
+	RenderTileMapComponent->bShowPerTileGridWhenUnselected = !RenderTileMapComponent->bShowPerTileGridWhenUnselected;
 	Invalidate();
 }
 
 bool FTileMapEditorViewportClient::IsShowTileGridChecked() const
 {
-	return RenderTileMapComponent->bShowPerTileGridWhenSelected;
+	return RenderTileMapComponent->bShowPerTileGridWhenUnselected;
 }
 
 void FTileMapEditorViewportClient::ToggleShowLayerGrid()
 {
 	FComponentReregisterContext RefreshComponentHelper(RenderTileMapComponent);
-	RenderTileMapComponent->bShowPerLayerGridWhenSelected = !RenderTileMapComponent->bShowPerLayerGridWhenSelected;
+	RenderTileMapComponent->bShowPerLayerGridWhenUnselected = !RenderTileMapComponent->bShowPerLayerGridWhenUnselected;
 	Invalidate();
 }
 
 bool FTileMapEditorViewportClient::IsShowLayerGridChecked() const
 {
-	return RenderTileMapComponent->bShowPerLayerGridWhenSelected;
+	return RenderTileMapComponent->bShowPerLayerGridWhenUnselected;
 }
 
 void FTileMapEditorViewportClient::ToggleShowMeshEdges()
diff --git a/Engine/Plugins/Runtime/GameplayAbilities/Source/GameplayAbilities/Private/Abilities/GameplayAbility.cpp b/Engine/Plugins/Runtime/GameplayAbilities/Source/GameplayAbilities/Private/Abilities/GameplayAbility.cpp
--- a/Engine/Plugins/Runtime/GameplayAbilities/Source/GameplayAbilities/Private/Abilities/GameplayAbility.cpp
+++ b/Engine/Plugins/Runtime/GameplayAbilities/Source/GameplayAbilities/Private/Abilities/GameplayAbility.cpp
@@ -10,12 +10,6 @@
 #include "Abilities/Tasks/AbilityTask.h"
 #include "GameplayCue_Types.h"
 
-// --------------------------------------------------------------------------------------------------------------------------------------------------------
-//
-//	UGameplayAbility
-//
-// --------------------------------------------------------------------------------------------------------------------------------------------------------
-
 UGameplayAbility::UGameplayAbility(const FObjectInitializer& ObjectInitializer)
 	: Super(ObjectInitializer)
 {
@@ -71,20 +65,30 @@ UGameplayAbility::UGameplayAbility(const FObjectInitializer& ObjectInitializer)
 	ScopeLockCount = 0;
 }
 
+UWorld* UGameplayAbility::GetWorld() const
+{
+	if (!IsInstantiated())
+	{
+		// If we are a CDO, we must return nullptr instead of calling Outer->GetWorld() to fool UObject::ImplementsGetWorld.
+		return nullptr;
+	}
+	return GetOuter()->GetWorld();
+}
+
 int32 UGameplayAbility::GetFunctionCallspace(UFunction* Function, void* Parameters, FFrame* Stack)
 {
 	if (HasAnyFlags(RF_ClassDefaultObject))
 	{
 		return FunctionCallspace::Local;
 	}
-	check(GetOuter() != NULL);
+	check(GetOuter() != nullptr);
 	return GetOuter()->GetFunctionCallspace(Function, Parameters, Stack);
 }
 
 bool UGameplayAbility::CallRemoteFunction(UFunction* Function, void* Parameters, FOutParmRec* OutParms, FFrame* Stack)
 {
 	check(!HasAnyFlags(RF_ClassDefaultObject));
-	check(GetOuter() != NULL);
+	check(GetOuter() != nullptr);
 
 	AActor* Owner = CastChecked<AActor>(GetOuter());
 	UNetDriver* NetDriver = Owner->GetNetDriver();
@@ -97,7 +101,6 @@ bool UGameplayAbility::CallRemoteFunction(UFunction* Function, void* Parameters,
 	return false;
 }
 
-// TODO: polymorphic payload
 void UGameplayAbility::SendGameplayEvent(FGameplayTag EventTag, FGameplayEventData Payload)
 {
 	UAbilitySystemComponent* AbilitySystemComponent = CurrentActorInfo->AbilitySystemComponent.Get();
@@ -114,7 +117,7 @@ void UGameplayAbility::PostNetInit()
 	 *  This may need to be updated further if we start having abilities live on different outers than player AbilitySystemComponents.
 	 */
 	
-	if (CurrentActorInfo == NULL)
+	if (CurrentActorInfo == nullptr)
 	{
 		AActor* OwnerActor = Cast<AActor>(GetOuter());
 		if (ensure(OwnerActor))
@@ -919,6 +922,15 @@ USkeletalMeshComponent* UGameplayAbility::GetOwningComponentFromActorInfo() cons
 	return CurrentActorInfo->SkeletalMeshComponent.Get();
 }
 
+UAbilitySystemComponent* UGameplayAbility::GetAbilitySystemComponentFromActorInfo() const
+{
+	if (!ensure(CurrentActorInfo))
+	{
+		return nullptr;
+	}
+	return CurrentActorInfo->AbilitySystemComponent.Get();
+}
+
 FGameplayEffectSpecHandle UGameplayAbility::MakeOutgoingGameplayEffectSpec(TSubclassOf<UGameplayEffect> GameplayEffectClass, float Level) const
 {
 	check(CurrentActorInfo && CurrentActorInfo->AbilitySystemComponent.IsValid());
@@ -1023,8 +1035,6 @@ void UGameplayAbility::K2_EndAbility()
 	EndAbility(CurrentSpecHandle, CurrentActorInfo, CurrentActivationInfo, bReplicateEndAbility, bWasCancelled);
 }
 
-// --------------------------------------------------------------------
-
 void UGameplayAbility::MontageJumpToSection(FName SectionName)
 {
 	check(CurrentActorInfo);
@@ -1050,7 +1060,7 @@ void UGameplayAbility::MontageStop(float OverrideBlendOutTime)
 	check(CurrentActorInfo);
 
 	UAbilitySystemComponent* AbilitySystemComponent = CurrentActorInfo->AbilitySystemComponent.Get();
-	if (AbilitySystemComponent != NULL)
+	if (AbilitySystemComponent != nullptr)
 	{
 		// We should only stop the current montage if we are the animating ability
 		if (AbilitySystemComponent->IsAnimatingAbility(this))
@@ -1071,8 +1081,6 @@ UAnimMontage* UGameplayAbility::GetCurrentMontage() const
 	return CurrentMontage;
 }
 
-// --------------------------------------------------------------------
-
 FGameplayAbilityTargetingLocationInfo UGameplayAbility::MakeTargetLocationInfoFromOwnerActor()
 {
 	FGameplayAbilityTargetingLocationInfo ReturnLocation;
@@ -1092,8 +1100,6 @@ FGameplayAbilityTargetingLocationInfo UGameplayAbility::MakeTargetLocationInfoFr
 	return ReturnLocation;
 }
 
-//----------------------------------------------------------------------
-
 UGameplayTasksComponent* UGameplayAbility::GetGameplayTasksComponent(const UGameplayTask& Task) const
 {
 	return GetCurrentActorInfo() ? GetCurrentActorInfo()->AbilitySystemComponent.Get() : nullptr;
@@ -1392,10 +1398,10 @@ void UGameplayAbility::RemoveGrantedByEffect()
 
 UObject* UGameplayAbility::GetSourceObject(FGameplayAbilitySpecHandle Handle, const FGameplayAbilityActorInfo* ActorInfo) const
 {
-	if (ActorInfo != NULL)
+	if (ActorInfo != nullptr)
 	{
 		UAbilitySystemComponent* AbilitySystemComponent = ActorInfo->AbilitySystemComponent.Get();
-		if (AbilitySystemComponent != NULL)
+		if (AbilitySystemComponent != nullptr)
 		{
 			FGameplayAbilitySpec* AbilitySpec = AbilitySystemComponent->FindAbilitySpecFromHandle(Handle);
 			if (AbilitySpec)
@@ -1496,6 +1502,11 @@ bool UGameplayAbility::HasAuthorityOrPredictionKey(const FGameplayAbilityActorIn
 	return ActorInfo->AbilitySystemComponent->HasAuthorityOrPredictionKey(ActivationInfo);
 }
 
+bool UGameplayAbility::IsInstantiated() const
+{
+	return !HasAllFlags(RF_ClassDefaultObject);
+}
+
 void UGameplayAbility::OnGiveAbility(const FGameplayAbilityActorInfo* ActorInfo, const FGameplayAbilitySpec& Spec)
 {
 	SetCurrentActorInfo(Spec.Handle, ActorInfo);
@@ -1512,8 +1523,6 @@ void UGameplayAbility::OnAvatarSet(const FGameplayAbilityActorInfo* ActorInfo, c
 	// Projects may want to initiate passives or do other "BeginPlay" type of logic here.
 }
 
-// -------------------------------------------------------
-
 FActiveGameplayEffectHandle UGameplayAbility::BP_ApplyGameplayEffectToOwner(TSubclassOf<UGameplayEffect> GameplayEffectClass, int32 GameplayEffectLevel, int32 Stacks)
 {
 	check(CurrentActorInfo);
@@ -1563,8 +1572,6 @@ FActiveGameplayEffectHandle UGameplayAbility::ApplyGameplayEffectSpecToOwner(con
 	return FActiveGameplayEffectHandle();
 }
 
-// -------------------------------
-
 TArray<FActiveGameplayEffectHandle> UGameplayAbility::BP_ApplyGameplayEffectToTarget(FGameplayAbilityTargetDataHandle Target, TSubclassOf<UGameplayEffect> GameplayEffectClass, int32 GameplayEffectLevel, int32 Stacks)
 {
 	return ApplyGameplayEffectToTarget(CurrentSpecHandle, CurrentActorInfo, CurrentActivationInfo, Target, GameplayEffectClass, GameplayEffectLevel, Stacks);
@@ -1623,6 +1630,29 @@ TArray<FActiveGameplayEffectHandle> UGameplayAbility::ApplyGameplayEffectSpecToT
 	return EffectHandles;
 }
 
+void UGameplayAbility::SetCurrentActorInfo(const FGameplayAbilitySpecHandle Handle, const FGameplayAbilityActorInfo* ActorInfo) const
+{
+	if (IsInstantiated())
+	{
+		CurrentActorInfo = ActorInfo;
+		CurrentSpecHandle = Handle;
+	}
+}
+
+void UGameplayAbility::SetCurrentActivationInfo(const FGameplayAbilityActivationInfo ActivationInfo)
+{
+	if (IsInstantiated())
+	{
+		CurrentActivationInfo = ActivationInfo;
+	}
+}
+
+void UGameplayAbility::SetCurrentInfo(const FGameplayAbilitySpecHandle Handle, const FGameplayAbilityActorInfo* ActorInfo, const FGameplayAbilityActivationInfo ActivationInfo)
+{
+	SetCurrentActorInfo(Handle, ActorInfo);
+	SetCurrentActivationInfo(ActivationInfo);
+}
+
 void UGameplayAbility::IncrementListLock() const
 {
 	++ScopeLockCount;
diff --git a/Engine/Plugins/Runtime/GameplayAbilities/Source/GameplayAbilities/Private/Abilities/GameplayAbilitySet.cpp b/Engine/Plugins/Runtime/GameplayAbilities/Source/GameplayAbilities/Private/Abilities/GameplayAbilitySet.cpp
--- a/Engine/Plugins/Runtime/GameplayAbilities/Source/GameplayAbilities/Private/Abilities/GameplayAbilitySet.cpp
+++ b/Engine/Plugins/Runtime/GameplayAbilities/Source/GameplayAbilities/Private/Abilities/GameplayAbilitySet.cpp
@@ -23,7 +23,7 @@ void UGameplayAbilitySet::GiveAbilities(UAbilitySystemComponent* AbilitySystemCo
 	{
 		if (BindInfo.GameplayAbilityClass)
 		{
-			AbilitySystemComponent->GiveAbility(FGameplayAbilitySpec(BindInfo.GameplayAbilityClass->GetDefaultObject<UGameplayAbility>(), 1, (int32)BindInfo.Command));
+			AbilitySystemComponent->GiveAbility(FGameplayAbilitySpec(BindInfo.GameplayAbilityClass, 1, (int32)BindInfo.Command));
 		}
 	}
 }
diff --git a/Engine/Plugins/Runtime/GameplayAbilities/Source/GameplayAbilities/Private/Abilities/Tasks/AbilityTask_WaitGameplayEvent.cpp b/Engine/Plugins/Runtime/GameplayAbilities/Source/GameplayAbilities/Private/Abilities/Tasks/AbilityTask_WaitGameplayEvent.cpp
--- a/Engine/Plugins/Runtime/GameplayAbilities/Source/GameplayAbilities/Private/Abilities/Tasks/AbilityTask_WaitGameplayEvent.cpp
+++ b/Engine/Plugins/Runtime/GameplayAbilities/Source/GameplayAbilities/Private/Abilities/Tasks/AbilityTask_WaitGameplayEvent.cpp
@@ -12,12 +12,13 @@ UAbilityTask_WaitGameplayEvent::UAbilityTask_WaitGameplayEvent(const FObjectInit
 
 }
 
-UAbilityTask_WaitGameplayEvent* UAbilityTask_WaitGameplayEvent::WaitGameplayEvent(UGameplayAbility* OwningAbility, FGameplayTag Tag, AActor* OptionalExternalTarget, bool OnlyTriggerOnce)
+UAbilityTask_WaitGameplayEvent* UAbilityTask_WaitGameplayEvent::WaitGameplayEvent(UGameplayAbility* OwningAbility, FGameplayTag Tag, AActor* OptionalExternalTarget, bool OnlyTriggerOnce, bool OnlyMatchExact)
 {
 	UAbilityTask_WaitGameplayEvent* MyObj = NewAbilityTask<UAbilityTask_WaitGameplayEvent>(OwningAbility);
 	MyObj->Tag = Tag;
 	MyObj->SetExternalTarget(OptionalExternalTarget);
 	MyObj->OnlyTriggerOnce = OnlyTriggerOnce;
+	MyObj->OnlyMatchExact = OnlyMatchExact;
 
 	return MyObj;
 }
@@ -26,20 +27,34 @@ void UAbilityTask_WaitGameplayEvent::Activate()
 {
 	UAbilitySystemComponent* ASC = GetTargetASC();
 	if (ASC)
-	{	
-		MyHandle = ASC->GenericGameplayEventCallbacks.FindOrAdd(Tag).AddUObject(this, &UAbilityTask_WaitGameplayEvent::GameplayEventCallback);
+	{
+		if (OnlyMatchExact)
+		{
+			MyHandle = ASC->GenericGameplayEventCallbacks.FindOrAdd(Tag).AddUObject(this, &UAbilityTask_WaitGameplayEvent::GameplayEventCallback);
+		}
+		else
+		{
+			MyHandle = ASC->AddGameplayEventTagContainerDelegate(FGameplayTagContainer(Tag), FGameplayEventTagMulticastDelegate::FDelegate::CreateUObject(this, &UAbilityTask_WaitGameplayEvent::GameplayEventContainerCallback));
+		}	
 	}
 
 	Super::Activate();
 }
 
 void UAbilityTask_WaitGameplayEvent::GameplayEventCallback(const FGameplayEventData* Payload)
+{
+	GameplayEventContainerCallback(Tag, Payload);
+}
+
+void UAbilityTask_WaitGameplayEvent::GameplayEventContainerCallback(FGameplayTag MatchingTag, const FGameplayEventData* Payload)
 {
 	if (ShouldBroadcastAbilityTaskDelegates())
 	{
-		EventReceived.Broadcast(*Payload);
+		FGameplayEventData TempPayload = *Payload;
+		TempPayload.EventTag = MatchingTag;
+		EventReceived.Broadcast(MoveTemp(TempPayload));
 	}
-	if(OnlyTriggerOnce)
+	if (OnlyTriggerOnce)
 	{
 		EndTask();
 	}
@@ -68,8 +83,16 @@ void UAbilityTask_WaitGameplayEvent::OnDestroy(bool AbilityEnding)
 {
 	UAbilitySystemComponent* ASC = GetTargetASC();
 	if (ASC && MyHandle.IsValid())
-	{	
-		ASC->GenericGameplayEventCallbacks.FindOrAdd(Tag).Remove(MyHandle);
+	{
+		if (OnlyMatchExact)
+		{
+			ASC->GenericGameplayEventCallbacks.FindOrAdd(Tag).Remove(MyHandle);
+		}
+		else
+		{
+			ASC->RemoveGameplayEventTagContainerDelegate(FGameplayTagContainer(Tag), MyHandle);
+		}
+		
 	}
 
 	Super::OnDestroy(AbilityEnding);
diff --git a/Engine/Plugins/Runtime/GameplayAbilities/Source/GameplayAbilities/Private/AbilitySystemBlueprintLibrary.cpp b/Engine/Plugins/Runtime/GameplayAbilities/Source/GameplayAbilities/Private/AbilitySystemBlueprintLibrary.cpp
--- a/Engine/Plugins/Runtime/GameplayAbilities/Source/GameplayAbilities/Private/AbilitySystemBlueprintLibrary.cpp
+++ b/Engine/Plugins/Runtime/GameplayAbilities/Source/GameplayAbilities/Private/AbilitySystemBlueprintLibrary.cpp
@@ -39,14 +39,14 @@ bool UAbilitySystemBlueprintLibrary::IsValid(FGameplayAttribute Attribute)
 	return Attribute.IsValid();
 }
 
-float UAbilitySystemBlueprintLibrary::GetFloatAttribute(const class AActor* Actor, FGameplayAttribute Attribute, bool& bSuccessfullyFoundAttribute)
+float UAbilitySystemBlueprintLibrary::GetFloatAttribute(const  AActor* Actor, FGameplayAttribute Attribute, bool& bSuccessfullyFoundAttribute)
 {
 	const UAbilitySystemComponent* const AbilitySystem = UAbilitySystemGlobals::GetAbilitySystemComponentFromActor(Actor);
 
 	return GetFloatAttributeFromAbilitySystemComponent(AbilitySystem, Attribute, bSuccessfullyFoundAttribute);
 }
 
-float UAbilitySystemBlueprintLibrary::GetFloatAttributeFromAbilitySystemComponent(const class UAbilitySystemComponent* AbilitySystem, FGameplayAttribute Attribute, bool& bSuccessfullyFoundAttribute)
+float UAbilitySystemBlueprintLibrary::GetFloatAttributeFromAbilitySystemComponent(const  UAbilitySystemComponent* AbilitySystem, FGameplayAttribute Attribute, bool& bSuccessfullyFoundAttribute)
 {
 	bSuccessfullyFoundAttribute = true;
 
diff --git a/Engine/Plugins/Runtime/GameplayAbilities/Source/GameplayAbilities/Private/AbilitySystemComponent.cpp b/Engine/Plugins/Runtime/GameplayAbilities/Source/GameplayAbilities/Private/AbilitySystemComponent.cpp
--- a/Engine/Plugins/Runtime/GameplayAbilities/Source/GameplayAbilities/Private/AbilitySystemComponent.cpp
+++ b/Engine/Plugins/Runtime/GameplayAbilities/Source/GameplayAbilities/Private/AbilitySystemComponent.cpp
@@ -20,7 +20,6 @@ DEFINE_LOG_CATEGORY(LogAbilitySystemComponent);
 
 DECLARE_CYCLE_STAT(TEXT("AbilitySystemComp ApplyGameplayEffectSpecToTarget"), STAT_AbilitySystemComp_ApplyGameplayEffectSpecToTarget, STATGROUP_AbilitySystem);
 
-
 #define LOCTEXT_NAMESPACE "AbilitySystemComponent"
 
 /** Enable to log out all render state create, destroy and updatetransform events */
@@ -37,7 +36,7 @@ UAbilitySystemComponent::UAbilitySystemComponent(const FObjectInitializer& Objec
 							// if we don't have this, UpdateShouldTick() fails to have any effect
 							// because we'll be receiving ticks but bIsActive starts as false
 	
-	CachedIsNetSimulated = false;
+	bCachedIsNetSimulated = false;
 	UserAbilityActivationInhibited = false;
 
 	GenericConfirmInputID = INDEX_NONE;
@@ -45,12 +44,12 @@ UAbilitySystemComponent::UAbilitySystemComponent(const FObjectInitializer& Objec
 
 	bSuppressGrantAbility = false;
 	bSuppressGameplayCues = false;
-	bPendingMontagerep = false;
+	bPendingMontageRep = false;
 	bIsNetDirty = true;
 
 	AbilityLastActivatedTime = 0.f;
 
-	ReplicationMode = EReplicationMode::Full;
+	ReplicationMode = EGameplayEffectReplicationMode::Full;
 
 	ClientActivateAbilityFailedStartTime = 0.f;
 	ClientActivateAbilityFailedCountRecent = 0;
@@ -58,7 +57,7 @@ UAbilitySystemComponent::UAbilitySystemComponent(const FObjectInitializer& Objec
 
 const UAttributeSet* UAbilitySystemComponent::InitStats(TSubclassOf<class UAttributeSet> Attributes, const UDataTable* DataTable)
 {
-	const UAttributeSet* AttributeObj = NULL;
+	const UAttributeSet* AttributeObj = nullptr;
 	if (Attributes)
 	{
 		AttributeObj = GetOrCreateAttributeSubobject(Attributes);
@@ -112,7 +111,7 @@ const UAttributeSet* UAbilitySystemComponent::GetAttributeSubobject(const TSubcl
 			return Set;
 		}
 	}
-	return NULL;
+	return nullptr;
 }
 
 bool UAbilitySystemComponent::HasAttributeSetForAttribute(FGameplayAttribute Attribute) const
@@ -187,17 +186,35 @@ void UAbilitySystemComponent::BeginPlay()
 
 void UAbilitySystemComponent::CacheIsNetSimulated()
 {
-	CachedIsNetSimulated = IsNetSimulating();
+	bCachedIsNetSimulated = IsNetSimulating();
 	ActiveGameplayEffects.OwnerIsNetAuthority = IsOwnerActorAuthoritative();
 }
 
-// ---------------------------------------------------------
-
 const FActiveGameplayEffect* UAbilitySystemComponent::GetActiveGameplayEffect(const FActiveGameplayEffectHandle Handle) const
 {
 	return ActiveGameplayEffects.GetActiveGameplayEffect(Handle);
 }
 
+const FGameplayTagContainer* UAbilitySystemComponent::GetGameplayEffectSourceTagsFromHandle(FActiveGameplayEffectHandle Handle) const
+{
+	return ActiveGameplayEffects.GetGameplayEffectSourceTagsFromHandle(Handle);
+}
+
+const FGameplayTagContainer* UAbilitySystemComponent::GetGameplayEffectTargetTagsFromHandle(FActiveGameplayEffectHandle Handle) const
+{
+	return ActiveGameplayEffects.GetGameplayEffectTargetTagsFromHandle(Handle);
+}
+
+void UAbilitySystemComponent::CaptureAttributeForGameplayEffect(OUT FGameplayEffectAttributeCaptureSpec& OutCaptureSpec)
+{
+	// Verify the capture is happening on an attribute the component actually has a set for; if not, can't capture the value
+	const FGameplayAttribute& AttributeToCapture = OutCaptureSpec.BackingDefinition.AttributeToCapture;
+	if (AttributeToCapture.IsValid() && (AttributeToCapture.IsSystemAttribute() || GetAttributeSubobject(AttributeToCapture.GetAttributeSetClass())))
+	{
+		ActiveGameplayEffects.CaptureAttributeForGameplayEffect(OutCaptureSpec);
+	}
+}
+
 bool UAbilitySystemComponent::HasNetworkAuthorityToApplyGameplayEffect(FPredictionKey PredictionKey) const
 {
 	return (IsOwnerActorAuthoritative() || PredictionKey.IsValidForMorePrediction());
@@ -507,8 +524,6 @@ const FGameplayEffectAttributeCaptureDefinition& UAbilitySystemComponent::GetInc
 	return IncomingDurationCapture;
 }
 
-// ------------------------------------------------------------------------
-
 void UAbilitySystemComponent::ResetTagMap()
 {
 	GameplayTagCountContainer.Reset();
@@ -523,8 +538,6 @@ void UAbilitySystemComponent::NotifyTagMap_StackCountChange(const FGameplayTagCo
 	}
 }
 
-// ------------------------------------------------------------------------
-
 FActiveGameplayEffectHandle UAbilitySystemComponent::ApplyGameplayEffectSpecToTarget(const FGameplayEffectSpec &Spec, UAbilitySystemComponent *Target, FPredictionKey PredictionKey)
 {
 	SCOPE_CYCLE_COUNTER(STAT_AbilitySystemComp_ApplyGameplayEffectSpecToTarget);
@@ -1016,7 +1029,7 @@ void UAbilitySystemComponent::InvokeGameplayCueEvent(const FGameplayEffectSpecFo
 
 	if (!Spec.Def)
 	{
-		ABILITY_LOG(Warning, TEXT("InvokeGameplayCueEvent Actor %s that has no gameplay effect!"), ActorAvatar ? *ActorAvatar->GetName() : TEXT("NULL"));
+		ABILITY_LOG(Warning, TEXT("InvokeGameplayCueEvent Actor %s that has no gameplay effect!"), ActorAvatar ? *ActorAvatar->GetName() : TEXT("nullptr"));
 		return;
 	}
 	
@@ -1080,6 +1093,21 @@ void UAbilitySystemComponent::ExecuteGameplayCue(const FGameplayTag GameplayCueT
 	UAbilitySystemGlobals::Get().GetGameplayCueManager()->InvokeGameplayCueExecuted_WithParams(this, GameplayCueTag, ScopedPredictionKey, GameplayCueParameters);
 }
 
+void UAbilitySystemComponent::AddGameplayCue(const FGameplayTag GameplayCueTag, FGameplayEffectContextHandle EffectContext)
+{
+	AddGameplayCue_Internal(GameplayCueTag, EffectContext, ActiveGameplayCues);
+}
+
+void UAbilitySystemComponent::AddGameplayCue(const FGameplayTag GameplayCueTag, const FGameplayCueParameters& GameplayCueParameters)
+{
+	AddGameplayCue_Internal(GameplayCueTag, GameplayCueParameters, ActiveGameplayCues);
+}
+
+void UAbilitySystemComponent::AddGameplayCue_MinimalReplication(const FGameplayTag GameplayCueTag, FGameplayEffectContextHandle EffectContext)
+{
+	AddGameplayCue_Internal(GameplayCueTag, EffectContext, MinimalReplicationGameplayCues);
+}
+
 void UAbilitySystemComponent::AddGameplayCue_Internal(const FGameplayTag GameplayCueTag, FGameplayEffectContextHandle& EffectContext, FActiveGameplayCueContainer& GameplayCueContainer)
 {
 	if (EffectContext.IsValid() == false)
@@ -1106,7 +1134,7 @@ void UAbilitySystemComponent::AddGameplayCue_Internal(const FGameplayTag Gamepla
 		// GE in his AGE array. Generate a prediction key for him, which he will look for on the _Implementation function and ignore.
 		{
 			FPredictionKey PredictionKeyForRPC = ScopedPredictionKey;
-			if (GameplayCueContainer.bMinimalReplication && (ReplicationMode == EReplicationMode::Mixed) && ScopedPredictionKey.IsValidKey() == false)
+			if (GameplayCueContainer.bMinimalReplication && (ReplicationMode == EGameplayEffectReplicationMode::Mixed) && ScopedPredictionKey.IsValidKey() == false)
 			{
 				PredictionKeyForRPC = FPredictionKey::CreateNewServerInitiatedKey(this);
 			}
@@ -1132,6 +1160,16 @@ void UAbilitySystemComponent::AddGameplayCue_Internal(const FGameplayTag Gamepla
 	}
 }
 
+void UAbilitySystemComponent::RemoveGameplayCue(const FGameplayTag GameplayCueTag)
+{
+	RemoveGameplayCue_Internal(GameplayCueTag, ActiveGameplayCues);
+}
+
+void UAbilitySystemComponent::RemoveGameplayCue_MinimalReplication(const FGameplayTag GameplayCueTag)
+{
+	RemoveGameplayCue_Internal(GameplayCueTag, MinimalReplicationGameplayCues);
+}
+
 void UAbilitySystemComponent::RemoveGameplayCue_Internal(const FGameplayTag GameplayCueTag, FActiveGameplayCueContainer& GameplayCueContainer)
 {
 	if (IsOwnerActorAuthoritative())
@@ -1173,8 +1211,6 @@ void UAbilitySystemComponent::NetMulticast_InvokeGameplayCueExecuted_FromSpec_Im
 	}
 }
 
-// -----------
-
 void UAbilitySystemComponent::NetMulticast_InvokeGameplayCueExecuted_Implementation(const FGameplayTag GameplayCueTag, FPredictionKey PredictionKey, FGameplayEffectContextHandle EffectContext)
 {
 	if (IsOwnerActorAuthoritative() || PredictionKey.IsLocalClientKey() == false)
@@ -1194,8 +1230,6 @@ void UAbilitySystemComponent::NetMulticast_InvokeGameplayCuesExecuted_Implementa
 	}
 }
 
-// -----------
-
 void UAbilitySystemComponent::NetMulticast_InvokeGameplayCueExecuted_WithParams_Implementation(const FGameplayTag GameplayCueTag, FPredictionKey PredictionKey, FGameplayCueParameters GameplayCueParameters)
 {
 	if (IsOwnerActorAuthoritative() || PredictionKey.IsLocalClientKey() == false)
@@ -1215,8 +1249,6 @@ void UAbilitySystemComponent::NetMulticast_InvokeGameplayCuesExecuted_WithParams
 	}
 }
 
-// -----------
-
 void UAbilitySystemComponent::NetMulticast_InvokeGameplayCueAdded_Implementation(const FGameplayTag GameplayCueTag, FPredictionKey PredictionKey, FGameplayEffectContextHandle EffectContext)
 {
 	if (IsOwnerActorAuthoritative() || PredictionKey.IsLocalClientKey() == false)
@@ -1230,7 +1262,7 @@ void UAbilitySystemComponent::NetMulticast_InvokeGameplayCueAdded_WithParams_Imp
 	// If server generated prediction key and auto proxy, skip this message. 
 	// This is an RPC from mixed replication mode code, we will get the "real" message from our OnRep on the autonomous proxy
 	// See UAbilitySystemComponent::AddGameplayCue_Internal for more info.
-	bool bIsMixedReplicationFromServer = (ReplicationMode == EReplicationMode::Mixed && PredictionKey.IsServerInitiatedKey() && AbilityActorInfo->IsLocallyControlledPlayer());
+	bool bIsMixedReplicationFromServer = (ReplicationMode == EGameplayEffectReplicationMode::Mixed && PredictionKey.IsServerInitiatedKey() && AbilityActorInfo->IsLocallyControlledPlayer());
 
 	if (IsOwnerActorAuthoritative() || (PredictionKey.IsLocalClientKey() == false && !bIsMixedReplicationFromServer))
 	{
@@ -1238,8 +1270,6 @@ void UAbilitySystemComponent::NetMulticast_InvokeGameplayCueAdded_WithParams_Imp
 	}
 }
 
-// -----------
-
 void UAbilitySystemComponent::NetMulticast_InvokeGameplayCueAddedAndWhileActive_FromSpec_Implementation(const FGameplayEffectSpecForRPC& Spec, FPredictionKey PredictionKey)
 {
 	if (IsOwnerActorAuthoritative() || PredictionKey.IsLocalClientKey() == false)
@@ -1258,7 +1288,6 @@ void UAbilitySystemComponent::NetMulticast_InvokeGameplayCueAddedAndWhileActive_
 	}
 }
 
-	
 void UAbilitySystemComponent::NetMulticast_InvokeGameplayCuesAddedAndWhileActive_WithParams_Implementation(const FGameplayTagContainer GameplayCueTags, FPredictionKey PredictionKey, FGameplayCueParameters GameplayCueParameters)
 {
 	if (IsOwnerActorAuthoritative() || PredictionKey.IsLocalClientKey() == false)
@@ -1271,8 +1300,6 @@ void UAbilitySystemComponent::NetMulticast_InvokeGameplayCuesAddedAndWhileActive
 	}
 }
 
-// ----------------------------------------------------------------------------------------
-
 TArray<float> UAbilitySystemComponent::GetActiveEffectsTimeRemaining(const FGameplayEffectQuery& Query) const
 {
 	return ActiveGameplayEffects.GetActiveEffectsTimeRemaining(Query);
@@ -1293,6 +1320,28 @@ TArray<FActiveGameplayEffectHandle> UAbilitySystemComponent::GetActiveEffects(co
 	return ActiveGameplayEffects.GetActiveEffects(Query);
 }
 
+float UAbilitySystemComponent::GetActiveEffectsEndTime(const FGameplayEffectQuery& Query) const
+{
+	TArray<AActor*> DummyInstigators;
+	return ActiveGameplayEffects.GetActiveEffectsEndTime(Query, DummyInstigators);
+}
+
+float UAbilitySystemComponent::GetActiveEffectsEndTimeWithInstigators(const FGameplayEffectQuery& Query, TArray<AActor*>& Instigators) const
+{
+	return ActiveGameplayEffects.GetActiveEffectsEndTime(Query, Instigators);
+}
+
+bool UAbilitySystemComponent::GetActiveEffectsEndTimeAndDuration(const FGameplayEffectQuery& Query, float& EndTime, float& Duration) const
+{
+	TArray<AActor*> DummyInstigators;
+	return ActiveGameplayEffects.GetActiveEffectsEndTimeAndDuration(Query, EndTime, Duration, DummyInstigators);
+}
+
+void UAbilitySystemComponent::ModifyActiveEffectStartTime(FActiveGameplayEffectHandle Handle, float StartTimeDiff)
+{
+	ActiveGameplayEffects.ModifyActiveEffectStartTime(Handle, StartTimeDiff);
+}
+
 int32 UAbilitySystemComponent::RemoveActiveEffectsWithTags(const FGameplayTagContainer Tags)
 {
 	if (IsOwnerActorAuthoritative())
@@ -1340,8 +1389,6 @@ int32 UAbilitySystemComponent::RemoveActiveEffects(const FGameplayEffectQuery& Q
 	return 0;
 }
 
-// ---------------------------------------------------------------------------------------
-
 void UAbilitySystemComponent::GetLifetimeReplicatedProps(TArray< FLifetimeProperty > & OutLifetimeProps) const
 {	
 	DOREPLIFETIME(UAbilitySystemComponent, SpawnedAttributes);
@@ -1424,7 +1471,7 @@ void UAbilitySystemComponent::PreNetReceive()
 	// Even though the value is also cached in OnRegister and BeginPlay, clients may
 	// receive properties before OnBeginPlay, so this ensures the role is correct
 	// for that case.
-	if (!CachedIsNetSimulated)
+	if (!bCachedIsNetSimulated)
 	{
 		CacheIsNetSimulated();
 	}
@@ -1441,7 +1488,7 @@ bool UAbilitySystemComponent::HasAuthorityOrPredictionKey(const FGameplayAbility
 	return ((ActivationInfo->ActivationMode == EGameplayAbilityActivationMode::Authority) || CanPredict());
 }
 
-void UAbilitySystemComponent::SetReplicationMode(EReplicationMode NewReplicationMode)
+void UAbilitySystemComponent::SetReplicationMode(EGameplayEffectReplicationMode NewReplicationMode)
 {
 	ReplicationMode = NewReplicationMode;
 }
@@ -1460,7 +1507,7 @@ IAbilitySystemReplicationProxyInterface* UAbilitySystemComponent::GetReplication
 void UAbilitySystemComponent::OnPredictiveGameplayCueCatchup(FGameplayTag Tag)
 {
 	// Remove it
-	RemoveOneTagCount_NoReturn(Tag);
+	RemoveLooseGameplayTag(Tag);
 
 	if (HasMatchingGameplayTag(Tag) == 0)
 	{
@@ -1480,15 +1527,16 @@ void UAbilitySystemComponent::ReinvokeActiveGameplayCues()
 	}
 }
 
-// ---------------------------------------------------------------------------------------
-
 void UAbilitySystemComponent::PrintAllGameplayEffects() const
 {
 	ABILITY_LOG(Log, TEXT("Owner: %s. Avatar: %s"), *GetOwner()->GetName(), *AbilityActorInfo->AvatarActor->GetName());
 	ActiveGameplayEffects.PrintAllGameplayEffects();
 }
 
-// ------------------------------------------------------------------------
+bool UAbilitySystemComponent::IsOwnerActorAuthoritative() const
+{
+	return !bCachedIsNetSimulated;
+}
 
 void UAbilitySystemComponent::OnAttributeAggregatorDirty(FAggregator* Aggregator, FGameplayAttribute Attribute, bool bFromRecursiveCall)
 {
@@ -1569,8 +1617,6 @@ void UAbilitySystemComponent::DebugCyclicAggregatorBroadcasts(FAggregator* Aggre
 	ActiveGameplayEffects.DebugCyclicAggregatorBroadcasts(Aggregator);
 }
 
-// ------------------------------------------------------------------------
-
 void UAbilitySystemComponent::OnRep_ClientDebugString()
 {
 	ABILITY_LOG(Display, TEXT(" "));
@@ -2377,5 +2423,4 @@ void UAbilitySystemComponent::Debug_Internal(FAbilitySystemComponentDebugInfo& I
 	Info.YL = MaxCharHeight;
 }
 
-
 #undef LOCTEXT_NAMESPACE
diff --git a/Engine/Plugins/Runtime/GameplayAbilities/Source/GameplayAbilities/Private/AbilitySystemComponent_Abilities.cpp b/Engine/Plugins/Runtime/GameplayAbilities/Source/GameplayAbilities/Private/AbilitySystemComponent_Abilities.cpp
--- a/Engine/Plugins/Runtime/GameplayAbilities/Source/GameplayAbilities/Private/AbilitySystemComponent_Abilities.cpp
+++ b/Engine/Plugins/Runtime/GameplayAbilities/Source/GameplayAbilities/Private/AbilitySystemComponent_Abilities.cpp
@@ -151,7 +151,7 @@ void UAbilitySystemComponent::InitAbilityActorInfo(AActor* InOwnerActor, AActor*
 		TagTable->RegisterResponseForEvents(this);
 	}
 
-	if (bPendingMontagerep)
+	if (bPendingMontageRep)
 	{
 		OnRep_ReplicatedAnimMontage();
 	}
@@ -193,8 +193,8 @@ void UAbilitySystemComponent::ClearActorInfo()
 {
 	check(AbilityActorInfo.IsValid());
 	AbilityActorInfo->ClearActorInfo();
-	OwnerActor = NULL;
-	AvatarActor = NULL;
+	OwnerActor = nullptr;
+	AvatarActor = nullptr;
 }
 
 void UAbilitySystemComponent::OnRep_OwningActor()
@@ -203,7 +203,7 @@ void UAbilitySystemComponent::OnRep_OwningActor()
 
 	if (OwnerActor != AbilityActorInfo->OwnerActor || AvatarActor != AbilityActorInfo->AvatarActor)
 	{
-		if (OwnerActor != NULL)
+		if (OwnerActor != nullptr)
 		{
 			InitAbilityActorInfo(OwnerActor, AvatarActor);
 		}
@@ -680,6 +680,7 @@ void UAbilitySystemComponent::NotifyAbilityEnded(FGameplayAbilitySpecHandle Hand
 
 	// Broadcast that the ability ended
 	AbilityEndedCallbacks.Broadcast(Ability);
+	OnAbilityEnded.Broadcast(FAbilityEndedData(Ability, Handle, false, bWasCancelled));
 	
 	/** If this is instanced per execution, mark pending kill and remove it from our instanced lists if we are the authority */
 	if (Ability->GetInstancingPolicy() == EGameplayAbilityInstancingPolicy::InstancedPerExecution)
@@ -714,8 +715,6 @@ void UAbilitySystemComponent::NotifyAbilityEnded(FGameplayAbilitySpecHandle Hand
 			MarkAbilitySpecDirty(*Spec);
 		}
 	}
-
-	OnAbilityEnded.Broadcast(FAbilityEndedData(Ability, Handle, false, bWasCancelled));
 }
 
 void UAbilitySystemComponent::CancelAbility(UGameplayAbility* Ability)
@@ -1008,7 +1007,7 @@ bool UAbilitySystemComponent::TryActivateAbility(FGameplayAbilitySpecHandle Abil
 	const FGameplayAbilityActorInfo* ActorInfo = AbilityActorInfo.Get();
 
 	// make sure the ActorInfo and then Actor on that FGameplayAbilityActorInfo are valid, if not bail out.
-	if (ActorInfo == NULL || !ActorInfo->OwnerActor.IsValid() || !ActorInfo->AvatarActor.IsValid())
+	if (ActorInfo == nullptr || !ActorInfo->OwnerActor.IsValid() || !ActorInfo->AvatarActor.IsValid())
 	{
 		return false;
 	}
@@ -1124,7 +1123,7 @@ bool UAbilitySystemComponent::InternalTryActivateAbility(FGameplayAbilitySpecHan
 	const FGameplayAbilityActorInfo* ActorInfo = AbilityActorInfo.Get();
 
 	// make sure the ActorInfo and then Actor on that FGameplayAbilityActorInfo are valid, if not bail out.
-	if (ActorInfo == NULL || !ActorInfo->OwnerActor.IsValid() || !ActorInfo->AvatarActor.IsValid())
+	if (ActorInfo == nullptr || !ActorInfo->OwnerActor.IsValid() || !ActorInfo->AvatarActor.IsValid())
 	{
 		return false;
 	}
@@ -1548,7 +1547,6 @@ void UAbilitySystemComponent::RemoteEndOrCancelAbility(FGameplayAbilitySpecHandl
 	}
 }
 
-/** Force cancels the ability and does not replicate this to the other side. This should be called when the ability is cancelled by the other side. */
 void UAbilitySystemComponent::ForceCancelAbilityDueToReplication(UGameplayAbility* Instance)
 {
 	check(Instance);
@@ -1845,10 +1843,6 @@ bool UAbilitySystemComponent::TriggerAbilityFromGameplayEvent(FGameplayAbilitySp
 	return false;
 }
 
-// ----------------------------------------------------------------------------------------------------------------------------------------------------
-//								Input 
-// ----------------------------------------------------------------------------------------------------------------------------------------------------
-
 bool UAbilitySystemComponent::GetUserAbilityActivationInhibited() const
 {
 	return UserAbilityActivationInhibited;
@@ -1910,9 +1904,58 @@ int32 UAbilitySystemComponent::HandleGameplayEvent(FGameplayTag EventTag, const
 		Delegate->Broadcast(Payload);
 	}
 
+	// Iterate with integers in case it changes due to callbacks
+	for (int32 Index = 0; Index < GameplayEventTagContainerDelegates.Num(); Index++)
+	{
+		TPair<FGameplayTagContainer, FGameplayEventTagMulticastDelegate>& SearchPair = GameplayEventTagContainerDelegates[Index];
+		if (SearchPair.Key.IsEmpty() || EventTag.MatchesAny(SearchPair.Key))
+		{
+			SearchPair.Value.Broadcast(EventTag, Payload);
+		}
+	}
+
 	return TriggeredCount;
 }
 
+FDelegateHandle UAbilitySystemComponent::AddGameplayEventTagContainerDelegate(const FGameplayTagContainer& TagFilter, const FGameplayEventTagMulticastDelegate::FDelegate& Delegate)
+{
+	TPair<FGameplayTagContainer, FGameplayEventTagMulticastDelegate>* FoundPair = nullptr;
+
+	for (TPair<FGameplayTagContainer, FGameplayEventTagMulticastDelegate>& SearchPair : GameplayEventTagContainerDelegates)
+	{
+		if (TagFilter == SearchPair.Key)
+		{
+			FoundPair = &SearchPair;
+			break;
+		}
+	}
+
+	if (!FoundPair)
+	{
+		FoundPair = new(GameplayEventTagContainerDelegates) TPair<FGameplayTagContainer, FGameplayEventTagMulticastDelegate>(TagFilter, FGameplayEventTagMulticastDelegate());
+	}
+
+	return FoundPair->Value.Add(Delegate);
+}
+
+void UAbilitySystemComponent::RemoveGameplayEventTagContainerDelegate(const FGameplayTagContainer& TagFilter, FDelegateHandle DelegateHandle)
+{
+	// Look for and remove delegate, remove from array if no more delegates are bound
+	for (int32 Index = 0; Index < GameplayEventTagContainerDelegates.Num(); Index++)
+	{
+		TPair<FGameplayTagContainer, FGameplayEventTagMulticastDelegate>& SearchPair = GameplayEventTagContainerDelegates[Index];
+		if (TagFilter == SearchPair.Key)
+		{
+			SearchPair.Value.Remove(DelegateHandle);
+			if (!SearchPair.Value.IsBound())
+			{
+				GameplayEventTagContainerDelegates.RemoveAt(Index);
+			}
+			break;
+		}
+	}
+}
+
 void UAbilitySystemComponent::MonitoredTagChanged(const FGameplayTag Tag, int32 NewCount)
 {
 	int32 TriggeredCount = 0;
@@ -1977,11 +2020,6 @@ bool UAbilitySystemComponent::HasNetworkAuthorityToActivateTriggeredAbility(cons
 	return false;
 }
 
-// ----------------------------------------------------------------------------------------------------------------------------------------------------
-//								Input 
-// ----------------------------------------------------------------------------------------------------------------------------------------------------
-
-
 void UAbilitySystemComponent::BindToInputComponent(UInputComponent* InputComponent)
 {
 	static const FName ConfirmBindName(TEXT("AbilityConfirm"));
@@ -2433,10 +2471,10 @@ void UAbilitySystemComponent::OnRep_ReplicatedAnimMontage()
 	if (AnimInstance == nullptr || !IsReadyForReplicatedMontage())
 	{
 		// We can't handle this yet
-		bPendingMontagerep = true;
+		bPendingMontageRep = true;
 		return;
 	}
-	bPendingMontagerep = false;
+	bPendingMontageRep = false;
 
 	if (!AbilityActorInfo->IsLocallyControlled())
 	{
@@ -2568,8 +2606,8 @@ void UAbilitySystemComponent::ClearAnimatingAbility(UGameplayAbility* Ability)
 {
 	if (LocalAnimMontageInfo.AnimatingAbility == Ability)
 	{
-		Ability->SetCurrentMontage(NULL);
-		LocalAnimMontageInfo.AnimatingAbility = NULL;
+		Ability->SetCurrentMontage(nullptr);
+		LocalAnimMontageInfo.AnimatingAbility = nullptr;
 	}
 }
 
@@ -2813,10 +2851,6 @@ UGameplayAbility* UAbilitySystemComponent::GetAnimatingAbility()
 	return LocalAnimMontageInfo.AnimatingAbility;
 }
 
-// -----------------------------------------------------------------------------
-
-
-
 void UAbilitySystemComponent::ConfirmAbilityTargetData(FGameplayAbilitySpecHandle AbilityHandle, FPredictionKey AbilityOriginalPredictionKey, const FGameplayAbilityTargetDataHandle& TargetData, const FGameplayTag& ApplicationTag)
 {
 	FAbilityReplicatedDataCache* CachedData = AbilityTargetDataMap.Find(FGameplayAbilitySpecHandleAndPredictionKey(AbilityHandle, AbilityOriginalPredictionKey));
@@ -2879,8 +2913,6 @@ FAbilityReplicatedData UAbilitySystemComponent::GetReplicatedDataOfGenericReplic
 	return ReturnData;
 }
 
-// --------------------------------------------------------------------------
-
 void UAbilitySystemComponent::ServerSetReplicatedEvent_Implementation(EAbilityGenericReplicatedEvent::Type EventType, FGameplayAbilitySpecHandle AbilityHandle, FPredictionKey AbilityOriginalPredictionKey,  FPredictionKey CurrentPredictionKey)
 {
 	FScopedPredictionWindow ScopedPrediction(this, CurrentPredictionKey);
@@ -2948,15 +2980,11 @@ bool UAbilitySystemComponent::ServerSetReplicatedEventWithPayload_Validate(EAbil
 	return true;
 }
 
-// -------
-
 void UAbilitySystemComponent::ClientSetReplicatedEvent_Implementation(EAbilityGenericReplicatedEvent::Type EventType, FGameplayAbilitySpecHandle AbilityHandle, FPredictionKey AbilityOriginalPredictionKey)
 {
 	InvokeReplicatedEvent(EventType, AbilityHandle, AbilityOriginalPredictionKey, ScopedPredictionKey);
 }
 
-// -------
-
 void UAbilitySystemComponent::ServerSetReplicatedTargetData_Implementation(FGameplayAbilitySpecHandle AbilityHandle, FPredictionKey AbilityOriginalPredictionKey, const FGameplayAbilityTargetDataHandle& ReplicatedTargetDataHandle, FGameplayTag ApplicationTag, FPredictionKey CurrentPredictionKey)
 {
 	FScopedPredictionWindow ScopedPrediction(this, CurrentPredictionKey);
@@ -2987,8 +3015,6 @@ bool UAbilitySystemComponent::ServerSetReplicatedTargetData_Validate(FGameplayAb
 	return true;
 }
 
-// -------
-
 void UAbilitySystemComponent::ServerSetReplicatedTargetDataCancelled_Implementation(FGameplayAbilitySpecHandle AbilityHandle, FPredictionKey AbilityOriginalPredictionKey, FPredictionKey CurrentPredictionKey)
 {
 	FScopedPredictionWindow ScopedPrediction(this, CurrentPredictionKey);
@@ -3102,8 +3128,6 @@ FSimpleMulticastDelegate& UAbilitySystemComponent::AbilityReplicatedEventDelegat
 	return AbilityTargetDataMap.FindOrAdd(FGameplayAbilitySpecHandleAndPredictionKey(AbilityHandle, AbilityOriginalPredictionKey)).GenericEvents[EventType].Delegate;
 }
 
-// ----------------
-
 int32 AbilitySystemLogServerRPCBatching = 0;
 static FAutoConsoleVariableRef CVarAbilitySystemLogServerRPCBatching(TEXT("AbilitySystem.ServerRPCBatching.Log"), AbilitySystemLogServerRPCBatching, TEXT(""), ECVF_Default	);
 
@@ -3281,5 +3305,3 @@ void UAbilitySystemComponent::CallServerEndAbility(FGameplayAbilitySpecHandle Ab
 }
 
 #undef LOCTEXT_NAMESPACE
-
-
diff --git a/Engine/Plugins/Runtime/GameplayAbilities/Source/GameplayAbilities/Private/AttributeSet.cpp b/Engine/Plugins/Runtime/GameplayAbilities/Source/GameplayAbilities/Private/AttributeSet.cpp
--- a/Engine/Plugins/Runtime/GameplayAbilities/Source/GameplayAbilities/Private/AttributeSet.cpp
+++ b/Engine/Plugins/Runtime/GameplayAbilities/Source/GameplayAbilities/Private/AttributeSet.cpp
@@ -324,6 +324,11 @@ bool UAttributeSet::IsNameStableForNetworking() const
 	return bNetAddressable || Super::IsNameStableForNetworking();
 }
 
+bool UAttributeSet::IsSupportedForNetworking() const
+{
+	return true;
+}
+
 void UAttributeSet::SetNetAddressable()
 {
 	bNetAddressable = true;
diff --git a/Engine/Plugins/Runtime/GameplayAbilities/Source/GameplayAbilities/Private/GameplayAbilityTypes.cpp b/Engine/Plugins/Runtime/GameplayAbilities/Source/GameplayAbilities/Private/GameplayAbilityTypes.cpp
--- a/Engine/Plugins/Runtime/GameplayAbilities/Source/GameplayAbilities/Private/GameplayAbilityTypes.cpp
+++ b/Engine/Plugins/Runtime/GameplayAbilities/Source/GameplayAbilities/Private/GameplayAbilityTypes.cpp
@@ -158,9 +158,9 @@ UGameplayAbility* FGameplayAbilitySpec::GetPrimaryInstance() const
 	return nullptr;
 }
 
-bool FGameplayAbilitySpec::ShouldReplicatedAbilitySpec() const
+bool FGameplayAbilitySpec::ShouldReplicateAbilitySpec() const
 {
-	if (Ability && Ability->ShouldReplicatedAbilitySpec(*this))
+	if (Ability && Ability->ShouldReplicateAbilitySpec(*this))
 	{
 		return true;
 	}
@@ -196,6 +196,32 @@ void FGameplayAbilitySpecContainer::RegisterWithOwner(UAbilitySystemComponent* I
 
 // ----------------------------------------------------
 
+FGameplayAbilitySpec::FGameplayAbilitySpec(UGameplayAbility* InAbility, int32 InLevel, int32 InInputID, UObject* InSourceObject)
+	: Ability(InAbility)
+	, Level(InLevel)
+	, InputID(InInputID)
+	, SourceObject(InSourceObject)
+	, ActiveCount(0)
+	, InputPressed(false)
+	, RemoveAfterActivation(false)
+	, PendingRemove(false)
+{
+	Handle.GenerateNewHandle();
+}
+
+FGameplayAbilitySpec::FGameplayAbilitySpec(TSubclassOf<UGameplayAbility> InAbilityClass, int32 InLevel, int32 InInputID, UObject* InSourceObject)
+	: Ability(InAbilityClass ? InAbilityClass.GetDefaultObject() : nullptr)
+	, Level(InLevel)
+	, InputID(InInputID)
+	, SourceObject(InSourceObject)
+	, ActiveCount(0)
+	, InputPressed(false)
+	, RemoveAfterActivation(false)
+	, PendingRemove(false)
+{
+	Handle.GenerateNewHandle();
+}
+
 FGameplayAbilitySpec::FGameplayAbilitySpec(FGameplayAbilitySpecDef& InDef, int32 InGameplayEffectLevel, FActiveGameplayEffectHandle InGameplayEffectHandle)
 	: Ability(InDef.Ability ? InDef.Ability->GetDefaultObject<UGameplayAbility>() : nullptr)
 	, InputID(InDef.InputID)
diff --git a/Engine/Plugins/Runtime/GameplayAbilities/Source/GameplayAbilities/Private/GameplayCueInterface.cpp b/Engine/Plugins/Runtime/GameplayAbilities/Source/GameplayAbilities/Private/GameplayCueInterface.cpp
--- a/Engine/Plugins/Runtime/GameplayAbilities/Source/GameplayAbilities/Private/GameplayCueInterface.cpp
+++ b/Engine/Plugins/Runtime/GameplayAbilities/Source/GameplayAbilities/Private/GameplayCueInterface.cpp
@@ -295,7 +295,7 @@ bool FActiveGameplayCueContainer::HasCue(const FGameplayTag& Tag) const
 
 bool FActiveGameplayCueContainer::NetDeltaSerialize(FNetDeltaSerializeInfo & DeltaParms)
 {
-	if (bMinimalReplication && (Owner && Owner->ReplicationMode == EReplicationMode::Full))
+	if (bMinimalReplication && (Owner && Owner->ReplicationMode == EGameplayEffectReplicationMode::Full))
 	{
 		return false;
 	}
diff --git a/Engine/Plugins/Runtime/GameplayAbilities/Source/GameplayAbilities/Private/GameplayEffect.cpp b/Engine/Plugins/Runtime/GameplayAbilities/Source/GameplayAbilities/Private/GameplayEffect.cpp
--- a/Engine/Plugins/Runtime/GameplayAbilities/Source/GameplayAbilities/Private/GameplayEffect.cpp
+++ b/Engine/Plugins/Runtime/GameplayAbilities/Source/GameplayAbilities/Private/GameplayEffect.cpp
@@ -1,4 +1,4 @@
-// Copyright 1998-2018 Epic Games, Inc. All Rights Reserved.
+﻿// Copyright 1998-2018 Epic Games, Inc. All Rights Reserved.
 
 #include "GameplayEffect.h"
 #include "TimerManager.h"
@@ -86,15 +86,6 @@ void UGameplayEffect::PostLoad()
 
 	UpdateInheritedTagProperties();
 
-	for (FGameplayAbilitySpecDef& Def : GrantedAbilities)
-	{
-		if (Def.Level != INDEX_NONE)
-		{
-			Def.LevelScalableFloat.SetValue(Def.Level);
-			Def.Level = INDEX_NONE;
-		}
-	}
-
 	HasGrantedApplicationImmunityQuery = !GrantedApplicationImmunityQuery.IsEmpty();
 
 #if WITH_EDITOR
@@ -2374,7 +2365,7 @@ bool FActiveGameplayEffectsContainer::HandleActiveGameplayEffectStackOverflow(co
 
 bool FActiveGameplayEffectsContainer::ShouldUseMinimalReplication()
 {
-	return IsNetAuthority() && (Owner->ReplicationMode == EReplicationMode::Minimal || Owner->ReplicationMode == EReplicationMode::Mixed);
+	return IsNetAuthority() && (Owner->ReplicationMode == EGameplayEffectReplicationMode::Minimal || Owner->ReplicationMode == EGameplayEffectReplicationMode::Mixed);
 }
 
 void FActiveGameplayEffectsContainer::SetBaseAttributeValueFromReplication(FGameplayAttribute Attribute, float ServerValue)
@@ -3671,12 +3662,12 @@ bool FActiveGameplayEffectsContainer::NetDeltaSerialize(FNetDeltaSerializeInfo&
 {
 	if (Owner)
 	{
-		EReplicationMode ReplicationMode = Owner->ReplicationMode;
-		if (ReplicationMode == EReplicationMode::Minimal)
+		EGameplayEffectReplicationMode ReplicationMode = Owner->ReplicationMode;
+		if (ReplicationMode == EGameplayEffectReplicationMode::Minimal)
 		{
 			return false;
 		}
-		else if (ReplicationMode == EReplicationMode::Mixed)
+		else if (ReplicationMode == EGameplayEffectReplicationMode::Mixed)
 		{
 			if (UPackageMapClient* Client = Cast<UPackageMapClient>(DeltaParms.Map))
 			{
@@ -3822,6 +3813,14 @@ void FActiveGameplayEffectsContainer::CheckDuration(FActiveGameplayEffectHandle
 
 						ExecuteActiveEffectsFrom(Effect.Spec);
 
+						// Invoke Delegates for periodic effects being executed
+						UAbilitySystemComponent* SourceASC = Effect.Spec.GetContext().GetInstigatorAbilitySystemComponent();
+						Owner->OnPeriodicGameplayEffectExecuteOnSelf(SourceASC, Effect.Spec, Handle);
+						if (SourceASC)
+						{
+							SourceASC->OnPeriodicGameplayEffectExecuteOnTarget(Owner, Effect.Spec, Handle);
+						}
+
 						// The above call to ExecuteActiveEffectsFrom could cause this effect to be explicitly removed
 						// (for example it could kill the owner and cause the effect to be wiped via death).
 						// In that case, we need to early out instead of possibly continueing to the below calls to InternalRemoveActiveGameplayEffect
diff --git a/Engine/Plugins/Runtime/GameplayAbilities/Source/GameplayAbilities/Private/GameplayEffectTypes.cpp b/Engine/Plugins/Runtime/GameplayAbilities/Source/GameplayAbilities/Private/GameplayEffectTypes.cpp
--- a/Engine/Plugins/Runtime/GameplayAbilities/Source/GameplayAbilities/Private/GameplayEffectTypes.cpp
+++ b/Engine/Plugins/Runtime/GameplayAbilities/Source/GameplayAbilities/Private/GameplayEffectTypes.cpp
@@ -761,7 +761,7 @@ bool FGameplayCueParameters::NetSerialize(FArchive& Ar, class UPackageMap* Map,
 	return true;
 }
 
-bool FGameplayCueParameters::IsInstigatorLocallyControlled() const
+bool FGameplayCueParameters::IsInstigatorLocallyControlled(AActor* FallbackActor) const
 {
 	if (EffectContext.IsValid())
 	{
@@ -772,6 +772,15 @@ bool FGameplayCueParameters::IsInstigatorLocallyControlled() const
 	if (!Pawn)
 	{
 		Pawn = Cast<APawn>(EffectCauser.Get());
+		if (!Pawn && FallbackActor != nullptr)
+		{
+			// Fallback to passed in actor
+			Pawn = Cast<APawn>(FallbackActor);
+			if (!Pawn)
+			{
+				Pawn = FallbackActor->GetInstigator<APawn>();
+			}
+		}
 	}
 	if (Pawn)
 	{
