Copying //UE4/Dev-Framework to //UE4/Dev-Main (Source: //UE4/Dev-Framework @ 4058146)

diff --git a/Engine/Plugins/Runtime/GameplayAbilities/Source/GameplayAbilities/Public/Abilities/GameplayAbility.h b/Engine/Plugins/Runtime/GameplayAbilities/Source/GameplayAbilities/Public/Abilities/GameplayAbility.h
--- a/Engine/Plugins/Runtime/GameplayAbilities/Source/GameplayAbilities/Public/Abilities/GameplayAbility.h
+++ b/Engine/Plugins/Runtime/GameplayAbilities/Source/GameplayAbilities/Public/Abilities/GameplayAbility.h
@@ -20,8 +20,6 @@
 class UGameplayAbility;
 class UGameplayTask;
 class UGameplayTasksComponent;
-struct EventData;
-struct FAbilityEndedData;
 
 /**
  * UGameplayAbility
@@ -48,8 +46,6 @@
  *			-'Giving' abilities (that can be used) to actors
  *	
  *	
- *	 
- *	
  *	The intention is for programmers to create these non instanced abilities in C++. Designers can then
  *	extend them as data assets (E.g., they can change default properties, they cannot implement blueprint graphs).
  *	
@@ -70,12 +66,8 @@
  *	
  *	
  *	To support state or event replication, an ability must be instanced. This can be done with the InstancingPolicy property.
- *	
- *
- *	
  */
 
-
 /** Notification delegate definition for when the gameplay ability ends */
 DECLARE_MULTICAST_DELEGATE_OneParam(FOnGameplayAbilityEnded, UGameplayAbility*);
 
@@ -88,7 +80,7 @@
 /** Used to delay execution until we leave a critical section */
 DECLARE_DELEGATE(FPostLockDelegate);
 
-/** TriggerData */
+/** Structure that defines how an ability will be triggered by external events */
 USTRUCT()
 struct FAbilityTriggerData
 {
@@ -107,9 +99,7 @@
 	TEnumAsByte<EGameplayAbilityTriggerSource::Type> TriggerSource;
 };
 
-/**
- *	Abilities define custom gameplay logic that can be activated by players or external game logic.
- */
+/** Abilities define custom gameplay logic that can be activated by players or external game logic */
 UCLASS(Blueprintable)
 class GAMEPLAYABILITIES_API UGameplayAbility : public UObject, public IGameplayTaskOwnerInterface
 {
@@ -142,46 +132,49 @@
 	//		EndAbility()			- The ability has ended. This is intended to be called by the ability to end itself.
 	//	
 	// ----------------------------------------------------------------------------------------------------------------
-		
 
-	/** Returns true if this ability can be activated right now. Has no side effects */
-	virtual bool CanActivateAbility(const FGameplayAbilitySpecHandle Handle, const FGameplayAbilityActorInfo* ActorInfo, const FGameplayTagContainer* SourceTags = nullptr, const FGameplayTagContainer* TargetTags = nullptr, OUT FGameplayTagContainer* OptionalRelevantTags = nullptr) const;
+	// --------------------------------------
+	//	Accessors
+	// --------------------------------------
 
-	/** Returns true if this ability can be triggered right now. Has no side effects */
-	virtual bool ShouldAbilityRespondToEvent(const FGameplayAbilityActorInfo* ActorInfo, const FGameplayEventData* Payload) const;
-	
-	/** Returns the time in seconds remaining on the currently active cooldown. */
-	virtual float GetCooldownTimeRemaining(const FGameplayAbilityActorInfo* ActorInfo) const;
-
-	/** Returns the time in seconds remaining on the currently active cooldown and the original duration for this cooldown. */
-	virtual void GetCooldownTimeRemainingAndDuration(FGameplayAbilitySpecHandle Handle, const FGameplayAbilityActorInfo* ActorInfo, float& TimeRemaining, float& CooldownDuration) const;
-
-	virtual const FGameplayTagContainer* GetCooldownTags() const;
-	
-	/** Returns true if none of the ability's tags are blocked and if it doesn't have a "Blocking" tag and has all "Required" tags. */
-	bool DoesAbilitySatisfyTagRequirements(const UAbilitySystemComponent& AbilitySystemComponent, const FGameplayTagContainer* SourceTags = nullptr, const FGameplayTagContainer* TargetTags = nullptr, OUT FGameplayTagContainer* OptionalRelevantTags = nullptr) const;
-
-	/** Takes in the ability spec and checks if we should allow replication on the ability spec, this will NOT stop replication of the ability UObject just the spec inside the UAbilitySystemComponenet ActivatableAbilities for this ability */
-	virtual bool ShouldReplicatedAbilitySpec(const FGameplayAbilitySpec& AbilitySpec) const
-	{
-		return true;
-	}
-
+	/** Returns how the ability is instanced when executed. This limits what an ability can do in its implementation. */
 	EGameplayAbilityInstancingPolicy::Type GetInstancingPolicy() const
 	{
 		return InstancingPolicy;
 	}
 
+	/** How an ability replicates state/events to everyone on the network */
 	EGameplayAbilityReplicationPolicy::Type GetReplicationPolicy() const
 	{
 		return ReplicationPolicy;
 	}
 
+	/** How does an ability execute on the network. Does a client "ask and predict", "ask and wait", "don't ask (just do it)" */
 	EGameplayAbilityNetExecutionPolicy::Type GetNetExecutionPolicy() const
 	{
 		return NetExecutionPolicy;
 	}
 
+	/** Returns the actor info associated with this ability, has cached pointers to useful objects */
+	UFUNCTION(BlueprintCallable, Category=Ability)
+	FGameplayAbilityActorInfo GetActorInfo() const;
+
+	/** Returns the actor that owns this ability, which may not have a physical location */
+	UFUNCTION(BlueprintCallable, Category = Ability)
+	AActor* GetOwningActorFromActorInfo() const;
+
+	/** Returns the physical actor that is executing this ability. May be null */
+	UFUNCTION(BlueprintCallable, Category = Ability)
+	AActor* GetAvatarActorFromActorInfo() const;
+
+	/** Convenience method for abilities to get skeletal mesh component - useful for aiming abilities */
+	UFUNCTION(BlueprintCallable, DisplayName = "GetSkeletalMeshComponentFromActorInfo", Category = Ability)
+	USkeletalMeshComponent* GetOwningComponentFromActorInfo() const;
+
+	/** Returns the AbilitySystemComponent that is activating this ability */
+	UFUNCTION(BlueprintCallable, Category = Ability)
+	UAbilitySystemComponent* GetAbilitySystemComponentFromActorInfo() const;
+
 	/** Gets the current actor info bound to this ability - can only be called on instanced abilities. */
 	const FGameplayAbilityActorInfo* GetCurrentActorInfo() const
 	{
@@ -196,6 +189,7 @@
 		return CurrentActivationInfo;
 	}
 
+	/** Gets the current activation info bound to this ability - can only be called on instanced abilities. */
 	FGameplayAbilityActivationInfo& GetCurrentActivationInfoRef()
 	{
 		check(IsInstantiated());
@@ -216,32 +210,47 @@
 	UFUNCTION(BlueprintCallable, Category = Ability)
 	FGameplayEffectContextHandle GetGrantedByEffectContext() const;
 
-	/** Removes the GameplayEffect that granted this ability. Can only be called on instanced abilities. */
+	/** Generates a GameplayEffectContextHandle from our owner and an optional TargetData.*/
 	UFUNCTION(BlueprintCallable, Category = Ability)
-	void RemoveGrantedByEffect();
+	virtual FGameplayEffectContextHandle GetContextFromOwner(FGameplayAbilityTargetDataHandle OptionalTargetData) const;
 
 	/** Returns an effect context, given a specified actor info */
 	virtual FGameplayEffectContextHandle MakeEffectContext(const FGameplayAbilitySpecHandle Handle, const FGameplayAbilityActorInfo *ActorInfo) const;
 
-	virtual UWorld* GetWorld() const override
-	{
-		if (!IsInstantiated())
-		{
-			// If we are a CDO, we must return nullptr instead of calling Outer->GetWorld() to fool UObject::ImplementsGetWorld.
-			return nullptr;
-		}
-		return GetOuter()->GetWorld();
-	}
+	/** Convenience method for abilities to get outgoing gameplay effect specs (for example, to pass on to projectiles to apply to whoever they hit) */
+	UFUNCTION(BlueprintCallable, Category=Ability)
+	FGameplayEffectSpecHandle MakeOutgoingGameplayEffectSpec(TSubclassOf<UGameplayEffect> GameplayEffectClass, float Level=1.f) const;
 
-	int32 GetFunctionCallspace(UFunction* Function, void* Parameters, FFrame* Stack) override;
+	/** Native version of above function */
+	virtual FGameplayEffectSpecHandle MakeOutgoingGameplayEffectSpec(const FGameplayAbilitySpecHandle Handle, const FGameplayAbilityActorInfo* ActorInfo, const FGameplayAbilityActivationInfo ActivationInfo, TSubclassOf<UGameplayEffect> GameplayEffectClass, float Level = 1.f) const;
 
-	bool CallRemoteFunction(UFunction* Function, void* Parameters, FOutParmRec* OutParms, FFrame* Stack) override;
+	/** Add the Ability's tags to the given GameplayEffectSpec. This is likely to be overridden per project. */
+	virtual void ApplyAbilityTagsToGameplayEffectSpec(FGameplayEffectSpec& Spec, FGameplayAbilitySpec* AbilitySpec) const;
 
-	void PostNetInit();
-
 	/** Returns true if the ability is currently active */
 	bool IsActive() const;
 
+	/** Is this ability triggered from TriggerData (or is it triggered explicitly through input/game code) */
+	bool IsTriggered() const;
+
+	/** Is this ability running on a a predicting client, this is false in single player */
+	bool IsPredictingClient() const;
+
+	/** True if this is on the server and is being executed for a non-local player, false in single player */
+	bool IsForRemoteClient() const;
+
+	/** True if the owning actor is locally controlled, true in single player */
+	bool IsLocallyControlled() const;
+
+	/** True if this is the server or single player */
+	bool HasAuthority(const FGameplayAbilityActivationInfo* ActivationInfo) const;
+
+	/** True if we are authority or we have a valid prediciton key that is expected to work */
+	bool HasAuthorityOrPredictionKey(const FGameplayAbilityActorInfo* ActorInfo, const FGameplayAbilityActivationInfo* ActivationInfo) const;
+
+	/** True if this has been instanced, always true for blueprints */
+	bool IsInstantiated() const;
+
 	/** Notification that the ability has ended.  Set using TryActivateAbility. */
 	FOnGameplayAbilityEnded OnGameplayAbilityEnded;
 
@@ -254,41 +263,118 @@
 	/** Used by the ability state task to handle when a state is ended */
 	FOnGameplayAbilityStateEnded OnGameplayAbilityStateEnded;
 
-	/** This ability has these tags */
-	UPROPERTY(EditDefaultsOnly, Category = Tags)
-	FGameplayTagContainer AbilityTags;
-
 	/** Callback for when this ability has been confirmed by the server */
 	FGenericAbilityDelegate	OnConfirmDelegate;
 
-	/** Is this ability triggered from TriggerData (or is it triggered explicitly through input/game code) */
-	bool IsTriggered() const;
+	// --------------------------------------
+	//	CanActivateAbility
+	// --------------------------------------
 
-	bool IsPredictingClient() const;
+	/** Returns true if this ability can be activated right now. Has no side effects */
+	virtual bool CanActivateAbility(const FGameplayAbilitySpecHandle Handle, const FGameplayAbilityActorInfo* ActorInfo, const FGameplayTagContainer* SourceTags = nullptr, const FGameplayTagContainer* TargetTags = nullptr, OUT FGameplayTagContainer* OptionalRelevantTags = nullptr) const;
 
-	bool IsForRemoteClient() const;
+	/** Returns true if this ability can be triggered right now. Has no side effects */
+	virtual bool ShouldAbilityRespondToEvent(const FGameplayAbilityActorInfo* ActorInfo, const FGameplayEventData* Payload) const;
 
-	bool IsLocallyControlled() const;
+	/** Returns true if an an ability should be activated */
+	virtual bool ShouldActivateAbility(ENetRole Role) const;
 
-	bool HasAuthority(const FGameplayAbilityActivationInfo* ActivationInfo) const;
+	/** Returns the time in seconds remaining on the currently active cooldown. */
+	UFUNCTION(BlueprintCallable, Category = Ability)
+	float GetCooldownTimeRemaining() const;
 
-	bool HasAuthorityOrPredictionKey(const FGameplayAbilityActorInfo* ActorInfo, const FGameplayAbilityActivationInfo* ActivationInfo) const;
+	/** Returns the time in seconds remaining on the currently active cooldown. */
+	virtual float GetCooldownTimeRemaining(const FGameplayAbilityActorInfo* ActorInfo) const;
 
-	/** Called when the ability is given to an AbilitySystemComponent */
-	virtual void OnGiveAbility(const FGameplayAbilityActorInfo* ActorInfo, const FGameplayAbilitySpec& Spec);
+	/** Returns the time in seconds remaining on the currently active cooldown and the original duration for this cooldown. */
+	virtual void GetCooldownTimeRemainingAndDuration(FGameplayAbilitySpecHandle Handle, const FGameplayAbilityActorInfo* ActorInfo, float& TimeRemaining, float& CooldownDuration) const;
 
-	/** Called when the avatar actor is set/changes */
-	virtual void OnAvatarSet(const FGameplayAbilityActorInfo* ActorInfo, const FGameplayAbilitySpec& Spec);
+	/** Returns all tags that are currently on cooldown */
+	virtual const FGameplayTagContainer* GetCooldownTags() const;
 	
-	// --------------------------------------
-	//	IGameplayTaskOwnerInterface
-	// --------------------------------------	
-	virtual UGameplayTasksComponent* GetGameplayTasksComponent(const UGameplayTask& Task) const override;
-	virtual AActor* GetGameplayTaskOwner(const UGameplayTask* Task) const override;
-	virtual AActor* GetGameplayTaskAvatar(const UGameplayTask* Task) const override;
-	virtual void OnGameplayTaskInitialized(UGameplayTask& Task) override;
-	virtual void OnGameplayTaskActivated(UGameplayTask& Task) override;
-	virtual void OnGameplayTaskDeactivated(UGameplayTask& Task) override;
+	/** Returns true if none of the ability's tags are blocked and if it doesn't have a "Blocking" tag and has all "Required" tags. */
+	virtual bool DoesAbilitySatisfyTagRequirements(const UAbilitySystemComponent& AbilitySystemComponent, const FGameplayTagContainer* SourceTags = nullptr, const FGameplayTagContainer* TargetTags = nullptr, OUT FGameplayTagContainer* OptionalRelevantTags = nullptr) const;
+
+	/** Returns true if this ability is blocking other abilities */
+	virtual bool IsBlockingOtherAbilities() const;
+
+	/** Sets rather ability block flags are enabled or disabled. Only valid on instanced abilities */
+	UFUNCTION(BlueprintCallable, Category = Ability)
+	virtual void SetShouldBlockOtherAbilities(bool bShouldBlockAbilities);
 
 	// --------------------------------------
+	//	CancelAbility
+	// --------------------------------------
+
+	/** Destroys instanced-per-execution abilities. Instance-per-actor abilities should 'reset'. Any active ability state tasks receive the 'OnAbilityStateInterrupted' event. Non instance abilities - what can we do? */
+	virtual void CancelAbility(const FGameplayAbilitySpecHandle Handle, const FGameplayAbilityActorInfo* ActorInfo, const FGameplayAbilityActivationInfo ActivationInfo, bool bReplicateCancelAbility);
+
+	/** Call from Blueprint to cancel the ability naturally */
+	UFUNCTION(BlueprintCallable, Category = Ability, DisplayName = "CancelAbility", meta=(ScriptName = "CancelAbility"))
+	void K2_CancelAbility();
+
+	/** Returns true if this ability can be canceled */
+	virtual bool CanBeCanceled() const;
+
+	/** Sets whether the ability should ignore cancel requests. Only valid on instanced abilities */
+	UFUNCTION(BlueprintCallable, Category=Ability)
+	virtual void SetCanBeCanceled(bool bCanBeCanceled);
+
+	// --------------------------------------
+	//	CommitAbility
+	// --------------------------------------
+
+	/** Attempts to commit the ability (spend resources, etc). This our last chance to fail. Child classes that override ActivateAbility must call this themselves! */
+	UFUNCTION(BlueprintCallable, Category = Ability, DisplayName = "CommitAbility", meta=(ScriptName = "CommitAbility"))
+	virtual bool K2_CommitAbility();
+
+	/** Attempts to commit the ability's cooldown only. If BroadcastCommitEvent is true, it will broadcast the commit event that tasks like WaitAbilityCommit are listening for. */
+	UFUNCTION(BlueprintCallable, Category = Ability, DisplayName = "CommitAbilityCooldown", meta=(ScriptName = "CommitAbilityCooldown"))
+	virtual bool K2_CommitAbilityCooldown(bool BroadcastCommitEvent=false, bool ForceCooldown=false);
+
+	/** Attempts to commit the ability's cost only. If BroadcastCommitEvent is true, it will broadcast the commit event that tasks like WaitAbilityCommit are listening for. */
+	UFUNCTION(BlueprintCallable, Category = Ability, DisplayName = "CommitAbilityCost", meta=(ScriptName = "CommitAbilityCost"))
+	virtual bool K2_CommitAbilityCost(bool BroadcastCommitEvent=false);
+
+	/** Checks the ability's cooldown, but does not apply it.*/
+	UFUNCTION(BlueprintCallable, Category = Ability, DisplayName = "CheckAbilityCooldown", meta=(ScriptName = "CheckAbilityCooldown"))
+	virtual bool K2_CheckAbilityCooldown();
+
+	/** Checks the ability's cost, but does not apply it. */
+	UFUNCTION(BlueprintCallable, Category = Ability, DisplayName = "CheckAbilityCost", meta=(ScriptName = "CheckAbilityCost"))
+	virtual bool K2_CheckAbilityCost();
+
+	virtual bool CommitAbility(const FGameplayAbilitySpecHandle Handle, const FGameplayAbilityActorInfo* ActorInfo, const FGameplayAbilityActivationInfo ActivationInfo);
+	virtual bool CommitAbilityCooldown(const FGameplayAbilitySpecHandle Handle, const FGameplayAbilityActorInfo* ActorInfo, const FGameplayAbilityActivationInfo ActivationInfo, const bool ForceCooldown);
+	virtual bool CommitAbilityCost(const FGameplayAbilitySpecHandle Handle, const FGameplayAbilityActorInfo* ActorInfo, const FGameplayAbilityActivationInfo ActivationInfo);
+
+	/** The last chance to fail before committing, this will usually be the same as CanActivateAbility. Some abilities may need to do extra checks here if they are consuming extra stuff in CommitExecute */
+	virtual bool CommitCheck(const FGameplayAbilitySpecHandle Handle, const FGameplayAbilityActorInfo* ActorInfo, const FGameplayAbilityActivationInfo ActivationInfo);
+
+	/** BP event called from CommitAbility */
+	UFUNCTION(BlueprintImplementableEvent, Category = Ability, DisplayName = "CommitExecute", meta = (ScriptName = "CommitExecute"))
+	void K2_CommitExecute();
+
+	/** Does the commit atomically (consume resources, do cooldowns, etc) */
+	virtual void CommitExecute(const FGameplayAbilitySpecHandle Handle, const FGameplayAbilityActorInfo* ActorInfo, const FGameplayAbilityActivationInfo ActivationInfo);
+
+	/** Returns the gameplay effect used to determine cooldown */
+	virtual UGameplayEffect* GetCooldownGameplayEffect() const;
+
+	/** Returns the gameplay effect used to apply cost */
+	virtual UGameplayEffect* GetCostGameplayEffect() const;
+
+	/** Checks cooldown. returns true if we can be used again. False if not */
+	virtual bool CheckCooldown(const FGameplayAbilitySpecHandle Handle, const FGameplayAbilityActorInfo* ActorInfo, OUT FGameplayTagContainer* OptionalRelevantTags = nullptr) const;
+
+	/** Applies CooldownGameplayEffect to the target */
+	virtual void ApplyCooldown(const FGameplayAbilitySpecHandle Handle, const FGameplayAbilityActorInfo* ActorInfo, const FGameplayAbilityActivationInfo ActivationInfo) const;
+
+	/** Checks cost. returns true if we can pay for the ability. False if not */
+	virtual bool CheckCost(const FGameplayAbilitySpecHandle Handle, const FGameplayAbilityActorInfo* ActorInfo, OUT FGameplayTagContainer* OptionalRelevantTags = nullptr) const;
+
+	/** Applies the ability's cost to the target */
+	virtual void ApplyCost(const FGameplayAbilitySpecHandle Handle, const FGameplayAbilityActorInfo* ActorInfo, const FGameplayAbilityActivationInfo ActivationInfo) const;
+
+	// --------------------------------------
 	//	Input
@@ -300,26 +386,117 @@
 	/** Input binding stub. */
 	virtual void InputReleased(const FGameplayAbilitySpecHandle Handle, const FGameplayAbilityActorInfo* ActorInfo, const FGameplayAbilityActivationInfo ActivationInfo) {};
 
+	/** Called from AbilityTask_WaitConfirmCancel to handle input confirming */
 	virtual void OnWaitingForConfirmInputBegin() {}
 	virtual void OnWaitingForConfirmInputEnd() {}
 
+	// --------------------------------------
+	//	Animation
+	// --------------------------------------
+
+	/** Returns the currently playing montage for this ability, if any */
+	UFUNCTION(BlueprintCallable, Category = Animation)
+	UAnimMontage* GetCurrentMontage() const;
+
+	/** Call to set/get the current montage from a montage task. Set to allow hooking up montage events to ability events */
+	virtual void SetCurrentMontage(class UAnimMontage* InCurrentMontage);
+
+	/** Movement Sync */
+	virtual void SetMovementSyncPoint(FName SyncName);
+
+	// ----------------------------------------------------------------------------------------------------------------
+	//	Ability Levels and source objects 
+	// ----------------------------------------------------------------------------------------------------------------
+	
+	/** Returns current level of the Ability */
+	UFUNCTION(BlueprintCallable, Category = Ability)
+	int32 GetAbilityLevel() const;
+
+	/** Returns current ability level for non instanced abilities. You must call this version in these contexts! */
+	int32 GetAbilityLevel(FGameplayAbilitySpecHandle Handle, const FGameplayAbilityActorInfo* ActorInfo) const;
+
+	/** Retrieves the SourceObject associated with this ability. Can only be called on instanced abilities. */
+	UFUNCTION(BlueprintCallable, Category = Ability)
+	UObject* GetCurrentSourceObject() const;
+
+	/** Retrieves the SourceObject associated with this ability. Callable on non instanced */
+	UObject* GetSourceObject(FGameplayAbilitySpecHandle Handle, const FGameplayAbilityActorInfo* ActorInfo) const;
+
+	// --------------------------------------
+	//	Interaction with ability system component
+	// --------------------------------------
+
+	/** Called by ability system component to inform this ability instance the remote instance was ended */
+	virtual void SetRemoteInstanceHasEnded();
+
+	/** Called to inform the ability that the AvatarActor has been replaced. If the ability is dependent on avatar state, it may want to end itself. */
+	virtual void NotifyAvatarDestroyed();
+
+	/** Called to inform the ability that a task is waiting for the player to do something */
+	virtual void NotifyAbilityTaskWaitingOnPlayerData(class UAbilityTask* AbilityTask);
+
+	/** Called to inform the ability that a task is waiting for the player's avatar to do something in world */
+	virtual void NotifyAbilityTaskWaitingOnAvatar(class UAbilityTask* AbilityTask);
+
+	/** Called when the ability is given to an AbilitySystemComponent */
+	virtual void OnGiveAbility(const FGameplayAbilityActorInfo* ActorInfo, const FGameplayAbilitySpec& Spec);
+
+	/** Called when the avatar actor is set/changes */
+	virtual void OnAvatarSet(const FGameplayAbilityActorInfo* ActorInfo, const FGameplayAbilitySpec& Spec);
+
+	/** Takes in the ability spec and checks if we should allow replication on the ability spec, this will NOT stop replication of the ability UObject just the spec inside the UAbilitySystemComponenet ActivatableAbilities for this ability */
+	virtual bool ShouldReplicateAbilitySpec(const FGameplayAbilitySpec& AbilitySpec) const
+	{
+		return true;
+	}
+
+	/** 
+	 * Invalidates the current prediction key. This should be used in cases where there is a valid prediction window, but the server is doing logic that only he can do, and afterwards performs an action that the client could predict (had the client been able to run the server-only code prior).
+	 * This returns instantly and has no other side effects other than clearing the current prediction key.
+	 */ 
+	UFUNCTION(BlueprintCallable, Category = Ability)
+	void InvalidateClientPredictionKey() const;
+
+	/** Removes the GameplayEffect that granted this ability. Can only be called on instanced abilities. */
+	UFUNCTION(BlueprintCallable, Category = Ability)
+	virtual void RemoveGrantedByEffect();
+
+	/** Adds a debug message to display to the user */
+	void AddAbilityTaskDebugMessage(UGameplayTask* AbilityTask, FString DebugMessage);
+
+	// --------------------------------------
+	//	Public variables, exposed for backwards compatibility
+	// --------------------------------------
+
+	/** This ability has these tags */
+	UPROPERTY(EditDefaultsOnly, Category = Tags)
+	FGameplayTagContainer AbilityTags;
+
 	/** If true, this ability will always replicate input press/release events to the server. */
 	UPROPERTY(EditDefaultsOnly, Category = Input)
 	bool bReplicateInputDirectly;
 
+	/** Set when the remote instance of this ability has ended (but the local instance may still be running or finishing up */
+	UPROPERTY()
+	bool RemoteInstanceEnded;
+
 	// --------------------------------------
-	//	CancelAbility
+	//	UObject overrides
+	// --------------------------------------	
+	virtual UWorld* GetWorld() const override;
+	virtual int32 GetFunctionCallspace(UFunction* Function, void* Parameters, FFrame* Stack) override;
+	virtual bool CallRemoteFunction(UFunction* Function, void* Parameters, FOutParmRec* OutParms, FFrame* Stack) override;
+	virtual bool IsSupportedForNetworking() const override;
+
 	// --------------------------------------
-
-	/** Destroys instanced-per-execution abilities. Instance-per-actor abilities should 'reset'. Any active ability state tasks receive the 'OnAbilityStateInterrupted' event. Non instance abilities - what can we do? */
-	virtual void CancelAbility(const FGameplayAbilitySpecHandle Handle, const FGameplayAbilityActorInfo* ActorInfo, const FGameplayAbilityActivationInfo ActivationInfo, bool bReplicateCancelAbility);
-
-	/** Returns true if an an ability should be activated */
-	virtual bool ShouldActivateAbility(ENetRole Role) const;
-
-	/* Call from Blueprint to cancel the ability naturally */
-	UFUNCTION(BlueprintCallable, Category = Ability, DisplayName = "CancelAbility", meta=(ScriptName = "CancelAbility"))
-	void K2_CancelAbility();
+	//	IGameplayTaskOwnerInterface
+	// --------------------------------------	
+	virtual UGameplayTasksComponent* GetGameplayTasksComponent(const UGameplayTask& Task) const override;
+	virtual AActor* GetGameplayTaskOwner(const UGameplayTask* Task) const override;
+	virtual AActor* GetGameplayTaskAvatar(const UGameplayTask* Task) const override;
+	virtual void OnGameplayTaskInitialized(UGameplayTask& Task) override;
+	virtual void OnGameplayTaskActivated(UGameplayTask& Task) override;
+	virtual void OnGameplayTaskDeactivated(UGameplayTask& Task) override;
 
 protected:
 
@@ -332,7 +509,11 @@
 	bool K2_ShouldAbilityRespondToEvent(FGameplayAbilityActorInfo ActorInfo, FGameplayEventData Payload) const;
 
 	bool bHasBlueprintShouldAbilityRespondToEvent;
-		
+
+	/** Sends a gameplay event, also creates a prediction window */
+	UFUNCTION(BlueprintCallable, Category = "Abilities")
+	virtual void SendGameplayEvent(FGameplayTag EventTag, FGameplayEventData Payload);
+
 	// --------------------------------------
 	//	CanActivate
 	// --------------------------------------
@@ -364,92 +545,20 @@
 	UFUNCTION(BlueprintImplementableEvent, Category = Ability, DisplayName = "ActivateAbilityFromEvent", meta=(ScriptName = "ActivateAbilityFromEvent"))
 	void K2_ActivateAbilityFromEvent(const FGameplayEventData& EventData);
 
-	virtual void ActivateAbility(const FGameplayAbilitySpecHandle Handle, const FGameplayAbilityActorInfo* ActorInfo, const FGameplayAbilityActivationInfo ActivationInfo, const FGameplayEventData* TriggerEventData);
-
 	bool bHasBlueprintActivate;
 	bool bHasBlueprintActivateFromEvent;
 
+	/** Actually activate ability, do not call this directly */
+	virtual void ActivateAbility(const FGameplayAbilitySpecHandle Handle, const FGameplayAbilityActorInfo* ActorInfo, const FGameplayAbilityActivationInfo ActivationInfo, const FGameplayEventData* TriggerEventData);
+
+	/** Do boilerplate init stuff and then call ActivateAbility */
+	virtual void PreActivate(const FGameplayAbilitySpecHandle Handle, const FGameplayAbilityActorInfo* ActorInfo, const FGameplayAbilityActivationInfo ActivationInfo, FOnGameplayAbilityEnded::FDelegate* OnGameplayAbilityEndedDelegate);
+
+	/** Executes PreActivate and ActivateAbility */
 	void CallActivateAbility(const FGameplayAbilitySpecHandle Handle, const FGameplayAbilityActorInfo* ActorInfo, const FGameplayAbilityActivationInfo ActivationInfo, FOnGameplayAbilityEnded::FDelegate* OnGameplayAbilityEndedDelegate = nullptr, const FGameplayEventData* TriggerEventData = nullptr);
 
 	/** Called on a predictive ability when the server confirms its execution */
 	virtual void ConfirmActivateSucceed();
 
-	UFUNCTION(BlueprintCallable, Category = "Abilities")
-	virtual void SendGameplayEvent(FGameplayTag EventTag, FGameplayEventData Payload);
-
-public:
-
-	// --------------------------------------
-	//	CommitAbility
-	// --------------------------------------
-
-	/**
-	 * Attempts to commit the ability (spend resources, etc). This our last chance to fail.
-	 *	-Child classes that override ActivateAbility must call this themselves!
-	 */
-	UFUNCTION(BlueprintCallable, Category = Ability, DisplayName = "CommitAbility", meta=(ScriptName = "CommitAbility"))
-	virtual bool K2_CommitAbility();
-
-	/** Attempts to commit the ability's cooldown only. If BroadcastCommitEvent is true, it will broadcast the commit event that tasks like WaitAbilityCommit are listening for. */
-	UFUNCTION(BlueprintCallable, Category = Ability, DisplayName = "CommitAbilityCooldown", meta=(ScriptName = "CommitAbilityCooldown"))
-	virtual bool K2_CommitAbilityCooldown(bool BroadcastCommitEvent=false, bool ForceCooldown=false);
-
-	/** Attempts to commit the ability's cost only. If BroadcastCommitEvent is true, it will broadcast the commit event that tasks like WaitAbilityCommit are listening for. */
-	UFUNCTION(BlueprintCallable, Category = Ability, DisplayName = "CommitAbilityCost", meta=(ScriptName = "CommitAbilityCost"))
-	virtual bool K2_CommitAbilityCost(bool BroadcastCommitEvent=false);
-
-	/** Checks the ability's cooldown, but does not apply it.*/
-	UFUNCTION(BlueprintCallable, Category = Ability, DisplayName = "CheckAbilityCooldown", meta=(ScriptName = "CheckAbilityCooldown"))
-	virtual bool K2_CheckAbilityCooldown();
-
-	/** Checks the ability's cost, but does not apply it. */
-	UFUNCTION(BlueprintCallable, Category = Ability, DisplayName = "CheckAbilityCost", meta=(ScriptName = "CheckAbilityCost"))
-	virtual bool K2_CheckAbilityCost();
-
-	virtual bool CommitAbility(const FGameplayAbilitySpecHandle Handle, const FGameplayAbilityActorInfo* ActorInfo, const FGameplayAbilityActivationInfo ActivationInfo);
-	virtual bool CommitAbilityCooldown(const FGameplayAbilitySpecHandle Handle, const FGameplayAbilityActorInfo* ActorInfo, const FGameplayAbilityActivationInfo ActivationInfo, const bool ForceCooldown);
-	virtual bool CommitAbilityCost(const FGameplayAbilitySpecHandle Handle, const FGameplayAbilityActorInfo* ActorInfo, const FGameplayAbilityActivationInfo ActivationInfo);
-
-	/**
-	 * The last chance to fail before commiting
-	 *	-This will usually be the same as CanActivateAbility. Some abilities may need to do extra checks here if they are consuming extra stuff in CommitExecute
-	 */
-	virtual bool CommitCheck(const FGameplayAbilitySpecHandle Handle, const FGameplayAbilityActorInfo* ActorInfo, const FGameplayAbilityActivationInfo ActivationInfo);
-
-	// --------------------------------------
-	//	CommitExecute
-	// --------------------------------------
-
-	UFUNCTION(BlueprintImplementableEvent, Category = Ability, DisplayName="CommitExecute", meta=(ScriptName="CommitExecute"))
-	void K2_CommitExecute();
-
-	/** Does the commit atomically (consume resources, do cooldowns, etc) */
-	virtual void CommitExecute(const FGameplayAbilitySpecHandle Handle, const FGameplayAbilityActorInfo* ActorInfo, const FGameplayAbilityActivationInfo ActivationInfo);
-
-	/** Do boilerplate init stuff and then call ActivateAbility */
-	virtual void PreActivate(const FGameplayAbilitySpecHandle Handle, const FGameplayAbilityActorInfo* ActorInfo, const FGameplayAbilityActivationInfo ActivationInfo, FOnGameplayAbilityEnded::FDelegate* OnGameplayAbilityEndedDelegate);
-
-protected:
-	/** Destroys instanced-per-execution abilities. Instance-per-actor abilities should 'reset'. Non instance abilities - what can we do? */
-	UFUNCTION(BlueprintCallable, Category = Ability)
-	void ConfirmTaskByInstanceName(FName InstanceName, bool bEndTask);
-
-	/** Internal function, cancels all the tasks we asked to cancel last frame (by instance name). */
-	void EndOrCancelTasksByInstanceName();
-	TArray<FName> CancelTaskInstanceNames;
-
-	/** Add any task with this instance name to a list to be ended (not canceled) next frame.  See also CancelTaskByInstanceName. */
-	UFUNCTION(BlueprintCallable, Category = Ability)
-	void EndTaskByInstanceName(FName InstanceName);
-	TArray<FName> EndTaskInstanceNames;
-
-	/** Add any task with this instance name to a list to be canceled (not ended) next frame.  See also EndTaskByInstanceName. */
-	UFUNCTION(BlueprintCallable, Category = Ability)
-	void CancelTaskByInstanceName(FName InstanceName);
-
-	/** Ends any active ability state task with the given name. If name is 'None' all active states will be ended (in an arbitrary order). */
-	UFUNCTION(BlueprintCallable, Category = Ability)
-	void EndAbilityState(FName OptionalStateNameToEnd);
-
 	// -------------------------------------
 	//	EndAbility
@@ -470,51 +579,52 @@
 	virtual void EndAbility(const FGameplayAbilitySpecHandle Handle, const FGameplayAbilityActorInfo* ActorInfo, const FGameplayAbilityActivationInfo ActivationInfo, bool bReplicateEndAbility, bool bWasCancelled);
 
 	// -------------------------------------
-	//	GameplayEffects
-	//	
+	//  Apply Gameplay effects to Self
 	// -------------------------------------
 
-	// ---------
-	// Apply Self
-	// ---------
-
+	/** Apply a gameplay effect to the owner of this ability */
 	UFUNCTION(BlueprintCallable, Category = Ability, DisplayName="ApplyGameplayEffectToOwner", meta=(ScriptName="ApplyGameplayEffectToOwner"))
 	FActiveGameplayEffectHandle BP_ApplyGameplayEffectToOwner(TSubclassOf<UGameplayEffect> GameplayEffectClass, int32 GameplayEffectLevel = 1, int32 Stacks = 1);
 
 	/** Non blueprintcallable, safe to call on CDO/NonInstance abilities */
 	FActiveGameplayEffectHandle ApplyGameplayEffectToOwner(const FGameplayAbilitySpecHandle Handle, const FGameplayAbilityActorInfo* ActorInfo, const FGameplayAbilityActivationInfo ActivationInfo, const UGameplayEffect* GameplayEffect, float GameplayEffectLevel, int32 Stacks = 1) const;
 
+	/** Apply a previously created gameplay effect spec to the owner of this ability */
 	UFUNCTION(BlueprintCallable, Category = Ability, DisplayName = "ApplyGameplayEffectSpecToOwner", meta=(ScriptName = "ApplyGameplayEffectSpecToOwner"))
 	FActiveGameplayEffectHandle K2_ApplyGameplayEffectSpecToOwner(const FGameplayEffectSpecHandle EffectSpecHandle);
 
 	FActiveGameplayEffectHandle ApplyGameplayEffectSpecToOwner(const FGameplayAbilitySpecHandle AbilityHandle, const FGameplayAbilityActorInfo* ActorInfo, const FGameplayAbilityActivationInfo ActivationInfo, const FGameplayEffectSpecHandle SpecHandle) const;
 
-	// ---------
-	// Apply Target
-	// ---------
+	// -------------------------------------
+	//  Apply Gameplay effects to Target
+	// -------------------------------------
 
+	/** Apply a gameplay effect to a Target */
 	UFUNCTION(BlueprintCallable, Category = Ability, DisplayName = "ApplyGameplayEffectToTarget", meta=(ScriptName = "ApplyGameplayEffectToTarget"))
 	TArray<FActiveGameplayEffectHandle> BP_ApplyGameplayEffectToTarget(FGameplayAbilityTargetDataHandle TargetData, TSubclassOf<UGameplayEffect> GameplayEffectClass, int32 GameplayEffectLevel = 1, int32 Stacks = 1);
 
 	/** Non blueprintcallable, safe to call on CDO/NonInstance abilities */
 	TArray<FActiveGameplayEffectHandle> ApplyGameplayEffectToTarget(const FGameplayAbilitySpecHandle Handle, const FGameplayAbilityActorInfo* ActorInfo, const FGameplayAbilityActivationInfo ActivationInfo, const FGameplayAbilityTargetDataHandle& Target, TSubclassOf<UGameplayEffect> GameplayEffectClass, float GameplayEffectLevel, int32 Stacks = 1) const;
 
+	/** Apply a previously created gameplay effect spec to a target */
 	UFUNCTION(BlueprintCallable, Category = Ability, DisplayName = "ApplyGameplayEffectSpecToTarget", meta=(ScriptName = "ApplyGameplayEffectSpecToTarget"))
 	TArray<FActiveGameplayEffectHandle> K2_ApplyGameplayEffectSpecToTarget(const FGameplayEffectSpecHandle EffectSpecHandle, FGameplayAbilityTargetDataHandle TargetData);
 
 	TArray<FActiveGameplayEffectHandle> ApplyGameplayEffectSpecToTarget(const FGameplayAbilitySpecHandle AbilityHandle, const FGameplayAbilityActorInfo* ActorInfo, const FGameplayAbilityActivationInfo ActivationInfo, const FGameplayEffectSpecHandle SpecHandle, const FGameplayAbilityTargetDataHandle& TargetData) const;
 
-	// ---------
-	// Remove Self
-	// ---------
+	// -------------------------------------
+	//  Remove Gameplay effects from Self
+	// -------------------------------------
 	
-	/** Removes GameplayEffects from owner which match the given asset level tags. */
+	/** Removes GameplayEffects from owner which match the given asset level tags */
 	UFUNCTION(BlueprintCallable, Category = Ability, DisplayName="RemoveGameplayEffectFromOwnerWithAssetTags", meta=(ScriptName="RemoveGameplayEffectFromOwnerWithAssetTags"))
 	void BP_RemoveGameplayEffectFromOwnerWithAssetTags(FGameplayTagContainer WithAssetTags, int32 StacksToRemove = -1);
 
+	/** Removes GameplayEffects from owner which grant the given tags */
 	UFUNCTION(BlueprintCallable, Category = Ability, DisplayName="RemoveGameplayEffectFromOwnerWithGrantedTags", meta=(ScriptName="RemoveGameplayEffectFromOwnerWithGrantedTags"))
 	void BP_RemoveGameplayEffectFromOwnerWithGrantedTags(FGameplayTagContainer WithGrantedTags, int32 StacksToRemove = -1);
 
+	/** Removes GameplayEffect from owner that match the given handle */
 	UFUNCTION(BlueprintCallable, Category = Ability, DisplayName = "RemoveGameplayEffectFromOwnerWithHandle", meta=(ScriptName = "RemoveGameplayEffectFromOwnerWithHandle"))
 	void BP_RemoveGameplayEffectFromOwnerWithHandle(FActiveGameplayEffectHandle Handle, int32 StacksToRemove = -1);
 
@@ -522,157 +632,30 @@
 	//	GameplayCue
 	//	Abilities can invoke GameplayCues without having to create GameplayEffects
 	// -------------------------------------
-	
+
+	/** Invoke a gameplay cue on the ability owner */
 	UFUNCTION(BlueprintCallable, Category = Ability, meta=(GameplayTagFilter="GameplayCue"), DisplayName="ExecuteGameplayCue", meta=(ScriptName="ExecuteGameplayCue"))
 	virtual void K2_ExecuteGameplayCue(FGameplayTag GameplayCueTag, FGameplayEffectContextHandle Context);
 
+	/** Invoke a gameplay cue on the ability owner, with extra parameters */
 	UFUNCTION(BlueprintCallable, Category = Ability, meta = (GameplayTagFilter = "GameplayCue"), DisplayName = "ExecuteGameplayCueWithParams", meta=(ScriptName = "ExecuteGameplayCueWithParams"))
 	virtual void K2_ExecuteGameplayCueWithParams(FGameplayTag GameplayCueTag, const FGameplayCueParameters& GameplayCueParameters);
 
+	/** Adds a persistent gameplay cue to the ability owner. Optionally will remove if ability ends */
 	UFUNCTION(BlueprintCallable, Category = Ability, meta=(GameplayTagFilter="GameplayCue"), DisplayName="AddGameplayCue", meta=(ScriptName="AddGameplayCue"))
 	virtual void K2_AddGameplayCue(FGameplayTag GameplayCueTag, FGameplayEffectContextHandle Context, bool bRemoveOnAbilityEnd = true);
 
+	/** Adds a persistent gameplay cue to the ability owner. Optionally will remove if ability ends */
 	UFUNCTION(BlueprintCallable, Category = Ability, meta = (GameplayTagFilter = "GameplayCue"), DisplayName = "AddGameplayCueWithParams", meta=(ScriptName = "AddGameplayCueWithParams"))
 	virtual void K2_AddGameplayCueWithParams(FGameplayTag GameplayCueTag, const FGameplayCueParameters& GameplayCueParameter, bool bRemoveOnAbilityEnd = true);
 
+	/** Removes a persistent gameplay cue from the ability owner */
 	UFUNCTION(BlueprintCallable, Category = Ability, meta=(GameplayTagFilter="GameplayCue"), DisplayName="RemoveGameplayCue", meta=(ScriptName="RemoveGameplayCue"))
 	virtual void K2_RemoveGameplayCue(FGameplayTag GameplayCueTag);
 
-	/** Generates a GameplayEffectContextHandle from our owner and an optional TargetData.*/
-	UFUNCTION(BlueprintCallable, Category = Ability, meta = (GameplayTagFilter = "GameplayCue"))
-	virtual FGameplayEffectContextHandle GetContextFromOwner(FGameplayAbilityTargetDataHandle OptionalTargetData) const;
-
-
+	// -------------------------------------
+	//	Protected properties
 	// -------------------------------------
 
-
-public:
-	bool IsInstantiated() const
-	{
-		return !HasAllFlags(RF_ClassDefaultObject);
-	}
-
-protected:
-	virtual void SetCurrentActorInfo(const FGameplayAbilitySpecHandle Handle, const FGameplayAbilityActorInfo* ActorInfo) const
-	{
-		if (IsInstantiated())
-		{
-			CurrentActorInfo = ActorInfo;
-			CurrentSpecHandle = Handle;
-		}
-	}
-
-	virtual void SetCurrentActivationInfo(const FGameplayAbilityActivationInfo ActivationInfo)
-	{
-		if (IsInstantiated())
-		{
-			CurrentActivationInfo = ActivationInfo;
-		}
-	}
-
-	void SetCurrentInfo(const FGameplayAbilitySpecHandle Handle, const FGameplayAbilityActorInfo* ActorInfo, const FGameplayAbilityActivationInfo ActivationInfo)
-	{
-		SetCurrentActorInfo(Handle, ActorInfo);
-		SetCurrentActivationInfo(ActivationInfo);
-	}
-
-public:
-
-	/** Returns the actor info associated with this ability, has cached pointers to useful objects */
-	UFUNCTION(BlueprintCallable, Category=Ability)
-	FGameplayAbilityActorInfo GetActorInfo() const;
-
-	/** Returns the actor that owns this ability, which may not have a physical location */
-	UFUNCTION(BlueprintCallable, Category = Ability)
-	AActor* GetOwningActorFromActorInfo() const;
-
-	/** Returns the physical actor that is executing this ability. May be null */
-	UFUNCTION(BlueprintCallable, Category = Ability)
-	AActor* GetAvatarActorFromActorInfo() const;
-
-	/** Convenience method for abilities to get skeletal mesh component - useful for aiming abilities */
-	UFUNCTION(BlueprintCallable, DisplayName = "GetSkeletalMeshComponentFromActorInfo", Category = Ability)
-	USkeletalMeshComponent* GetOwningComponentFromActorInfo() const;
-
-	/** Convenience method for abilities to get outgoing gameplay effect specs (for example, to pass on to projectiles to apply to whoever they hit) */
-	UFUNCTION(BlueprintCallable, Category=Ability)
-	FGameplayEffectSpecHandle MakeOutgoingGameplayEffectSpec(TSubclassOf<UGameplayEffect> GameplayEffectClass, float Level=1.f) const;
-
-	FGameplayEffectSpecHandle MakeOutgoingGameplayEffectSpec(const FGameplayAbilitySpecHandle Handle, const FGameplayAbilityActorInfo* ActorInfo, const FGameplayAbilityActivationInfo ActivationInfo, TSubclassOf<UGameplayEffect> GameplayEffectClass, float Level = 1.f) const;
-
-	/** Add the Ability's tags to the given GameplayEffectSpec. This is likely to be overridden per project. */
-	virtual void ApplyAbilityTagsToGameplayEffectSpec(FGameplayEffectSpec& Spec, FGameplayAbilitySpec* AbilitySpec) const;
-
-	/** Returns the currently playing montage for this ability, if any */
-	UFUNCTION(BlueprintCallable, Category = Animation)
-	UAnimMontage* GetCurrentMontage() const;
-
-	/** Call to set/get the current montage from a montage task. Set to allow hooking up montage events to ability events */
-	virtual void SetCurrentMontage(class UAnimMontage* InCurrentMontage);
-
-	/** Returns true if this ability can be canceled */
-	virtual bool CanBeCanceled() const;
-
-	/** Sets whether the ability should ignore cancel requests. Only valid on instanced abilities */
-	UFUNCTION(BlueprintCallable, Category=Ability)
-	virtual void SetCanBeCanceled(bool bCanBeCanceled);
-
-	/** Returns true if this ability is blocking other abilities */
-	virtual bool IsBlockingOtherAbilities() const;
-
-	/** Sets rather ability block flags are enabled or disabled. Only valid on instanced abilities */
-	UFUNCTION(BlueprintCallable, Category = Ability)
-	virtual void SetShouldBlockOtherAbilities(bool bShouldBlockAbilities);
-
-	bool IsSupportedForNetworking() const override;
-
-	/** Returns the gameplay effect used to determine cooldown */
-	virtual class UGameplayEffect* GetCooldownGameplayEffect() const;
-
-	/** Returns the gameplay effect used to apply cost */
-	class UGameplayEffect* GetCostGameplayEffect() const;
-
-	/** Checks cooldown. returns true if we can be used again. False if not */
-	virtual bool CheckCooldown(const FGameplayAbilitySpecHandle Handle, const FGameplayAbilityActorInfo* ActorInfo, OUT FGameplayTagContainer* OptionalRelevantTags = nullptr) const;
-
-	/** Applies CooldownGameplayEffect to the target */
-	virtual void ApplyCooldown(const FGameplayAbilitySpecHandle Handle, const FGameplayAbilityActorInfo* ActorInfo, const FGameplayAbilityActivationInfo ActivationInfo) const;
-
-	/** Checks cost. returns true if we can pay for the ability. False if not */
-	virtual bool CheckCost(const FGameplayAbilitySpecHandle Handle, const FGameplayAbilityActorInfo* ActorInfo, OUT FGameplayTagContainer* OptionalRelevantTags = nullptr) const;
-
-	/** Applies the ability's cost to the target */
-	virtual void ApplyCost(const FGameplayAbilitySpecHandle Handle, const FGameplayAbilityActorInfo* ActorInfo, const FGameplayAbilityActivationInfo ActivationInfo) const;
-
-	/** Movement Sync */
-	virtual void SetMovementSyncPoint(FName SyncName);
-
-	/** Set when the remote instance of this ability has ended (but the local instance may still be running or finishing up */
-	bool RemoteInstanceEnded;
-
-	/** Called by ability system component to inform this ability instance the remote instance was ended */
-	void SetRemoteInstanceHasEnded();
-
-	/** Called to inform the ability that the AvatarActor has been replaced. If the ability is dependant on avatar state, it may want to end itself. */
-	void NotifyAvatarDestroyed();
-
-	void NotifyAbilityTaskWaitingOnPlayerData(class UAbilityTask* AbilityTask);
-
-	void NotifyAbilityTaskWaitingOnAvatar(class UAbilityTask* AbilityTask);
-
-	UFUNCTION(BlueprintCallable, Category = Ability)
-	float GetCooldownTimeRemaining() const;
-
-	/** 
-	 * Invalidates the current prediction key. This should be used in cases where there is a valid prediction window, but the server is doing logic that only he can do, and afterwards performs an action that the client could predict (had the client been able to run the server-only code prior).
-	 * This returns instantly and has no other side effects other than clearing the current prediction key.
-	 */ 
-	UFUNCTION(BlueprintCallable, Category = Ability)
-	void InvalidateClientPredictionKey() const;
-
-protected:
-
-	// -----------------------------------------------	
-
 	UPROPERTY(EditDefaultsOnly, Category = Advanced)
 	TEnumAsByte<EGameplayAbilityReplicationPolicy::Type> ReplicationPolicy;
@@ -705,29 +688,19 @@
 	/** Triggers to determine if this ability should execute in response to an event */
 	UPROPERTY(EditDefaultsOnly, Category = Triggers)
 	TArray<FAbilityTriggerData> AbilityTriggers;
-
-	// ----------------------------------------------------------------------------------------------------------------
-	//
-	//	Cooldowns
-	//
-	// ----------------------------------------------------------------------------------------------------------------
 			
-	/** Deprecated? This GameplayEffect represents the cooldown. It will be applied when the ability is committed and the ability cannot be used again until it is expired. */
+	/** This GameplayEffect represents the cooldown. It will be applied when the ability is committed and the ability cannot be used again until it is expired. */
 	UPROPERTY(EditDefaultsOnly, Category = Cooldowns)
 	TSubclassOf<class UGameplayEffect> CooldownGameplayEffectClass;
 
 	// ----------------------------------------------------------------------------------------------------------------
-	//
 	//	Ability exclusion / canceling
-	//
 	// ----------------------------------------------------------------------------------------------------------------
 
+	/** Abilities matching query are cancelled when this ability is executed */
 	UPROPERTY(EditDefaultsOnly, Category = TagQueries)
 	FGameplayTagQuery CancelAbilitiesMatchingTagQuery;
 
-	UPROPERTY(VisibleDefaultsOnly, BlueprintReadOnly, Category = TagQueries)
-	FGameplayTagQuery ConstTagQuery;
-
 	/** Abilities with these tags are cancelled when this ability is executed */
 	UPROPERTY(EditDefaultsOnly, Category = Tags)
 	FGameplayTagContainer CancelAbilitiesWithTag;
@@ -764,33 +737,47 @@
 	UPROPERTY(EditDefaultsOnly, Category = Tags)
 	FGameplayTagContainer TargetBlockedTags;
 
-
 	// ----------------------------------------------------------------------------------------------------------------
-	//
 	//	Ability Tasks
-	//
 	// ----------------------------------------------------------------------------------------------------------------
-	
+
+	/** Destroys instanced-per-execution abilities. Instance-per-actor abilities should 'reset'. Non instance abilities - what can we do? */
+	UFUNCTION(BlueprintCallable, Category = Ability)
+	void ConfirmTaskByInstanceName(FName InstanceName, bool bEndTask);
+
+	/** Internal function, cancels all the tasks we asked to cancel last frame (by instance name). */
+	void EndOrCancelTasksByInstanceName();
+	TArray<FName> CancelTaskInstanceNames;
+
+	/** Add any task with this instance name to a list to be ended (not canceled) next frame.  See also CancelTaskByInstanceName. */
+	UFUNCTION(BlueprintCallable, Category = Ability)
+	void EndTaskByInstanceName(FName InstanceName);
+	TArray<FName> EndTaskInstanceNames;
+
+	/** Add any task with this instance name to a list to be canceled (not ended) next frame.  See also EndTaskByInstanceName. */
+	UFUNCTION(BlueprintCallable, Category = Ability)
+	void CancelTaskByInstanceName(FName InstanceName);
+
+	/** Ends any active ability state task with the given name. If name is 'None' all active states will be ended (in an arbitrary order). */
+	UFUNCTION(BlueprintCallable, Category = Ability)
+	void EndAbilityState(FName OptionalStateNameToEnd);
+
+	/** List of currently active tasks, do not modify directly */
 	UPROPERTY()
 	TArray<UGameplayTask*>	ActiveTasks;
 
 	/** Tasks can emit debug messages throughout their life for debugging purposes. Saved on the ability so that they persist after task is finished */
 	TArray<FAbilityTaskDebugMessage> TaskDebugMessages;
 
-public:
-	void AddAbilityTaskDebugMessage(UGameplayTask* AbilityTask, FString DebugMessage);
-
-protected:
-
 	// ----------------------------------------------------------------------------------------------------------------
-	//
 	//	Animation
-	//
 	// ----------------------------------------------------------------------------------------------------------------
 
+	/** Immediately jumps the active montage to a section */
 	UFUNCTION(BlueprintCallable, Category="Ability|Animation")
 	void MontageJumpToSection(FName SectionName);
 
+	/** Sets pending section on active montage */
 	UFUNCTION(BlueprintCallable, Category = "Ability|Animation")
 	void MontageSetNextSectionName(FName FromSectionName, FName ToSectionName);
 
@@ -802,41 +789,38 @@
 	UFUNCTION(BlueprintCallable, Category="Ability|Animation", Meta = (AdvancedDisplay = "OverrideBlendOutTime"))
 	void MontageStop(float OverrideBlendOutTime = -1.0f);
 
+	/** Active montage being played by this ability */
+	UPROPERTY()
+	class UAnimMontage* CurrentMontage;
+
 	// ----------------------------------------------------------------------------------------------------------------
-	//
 	//	Target Data
-	//
 	// ----------------------------------------------------------------------------------------------------------------
 
+	/** Creates a target location from where the owner avatar is */
 	UFUNCTION(BlueprintPure, Category = Ability)
 	FGameplayAbilityTargetingLocationInfo MakeTargetLocationInfoFromOwnerActor();
 
+	/** Creates a target location from a socket on the owner avatar's skeletal mesh */
 	UFUNCTION(BlueprintPure, Category = Ability)
 	FGameplayAbilityTargetingLocationInfo MakeTargetLocationInfoFromOwnerSkeletalMeshComponent(FName SocketName);
 
 	// ----------------------------------------------------------------------------------------------------------------
-	//
-	//	Ability Levels
-	// 
+	//	Setters for temporary execution data
 	// ----------------------------------------------------------------------------------------------------------------
+
+	/** Called to initialize after being created due to replication */
+	virtual void PostNetInit();
+
+	/** Modifies actor info, only safe on instanced abilities */
+	virtual void SetCurrentActorInfo(const FGameplayAbilitySpecHandle Handle, const FGameplayAbilityActorInfo* ActorInfo) const;
 	
-public:
-	/** Returns current level of the Ability */
-	UFUNCTION(BlueprintCallable, Category = Ability)
-	int32 GetAbilityLevel() const;
-
-	/** Returns current ability level for non instanced abilities. You must call this version in these contexts! */
-	int32 GetAbilityLevel(FGameplayAbilitySpecHandle Handle, const FGameplayAbilityActorInfo* ActorInfo) const;
-
-	/** Retrieves the SourceObject associated with this ability. Can only be called on instanced abilities. */
-	UFUNCTION(BlueprintCallable, Category = Ability)
-	UObject* GetCurrentSourceObject() const;
-
-	/** Retrieves the SourceObject associated with this ability. Callable on non instanced */
-	UObject* GetSourceObject(FGameplayAbilitySpecHandle Handle, const FGameplayAbilityActorInfo* ActorInfo) const;
-
-protected:
-
+	/** Modifies activation info, only safe on instanced abilities */
+	virtual void SetCurrentActivationInfo(const FGameplayAbilityActivationInfo ActivationInfo);
+	
+	/** Sets both actor and activation info */
+	void SetCurrentInfo(const FGameplayAbilitySpecHandle Handle, const FGameplayAbilityActorInfo* ActorInfo, const FGameplayAbilityActivationInfo ActivationInfo);
+	
 	/** 
 	 *  This is shared, cached information about the thing using us
 	 *	 E.g, Actor*, MovementComponent*, AnimInstance, etc.
@@ -846,15 +830,12 @@
 	 */
 	mutable const FGameplayAbilityActorInfo* CurrentActorInfo;
 
+	/** For instanced abilities */
 	mutable FGameplayAbilitySpecHandle CurrentSpecHandle;
 
 	/** GameplayCues that were added during this ability that will get automatically removed when it ends */
 	TSet<FGameplayTag> TrackedGameplayCues;
 
-	/** Active montage being played by this ability */
-	UPROPERTY()
-	class UAnimMontage* CurrentMontage;
-
 	/** True if the ability is currently active. For instance per owner abilities */
 	UPROPERTY()
 	bool bIsActive;
@@ -863,6 +844,10 @@
 	UPROPERTY()
 	bool bIsCancelable;
 
+	/** True if the ability block flags are currently enabled */
+	UPROPERTY()
+	bool bIsBlockingOtherAbilities;
+
 	/** A count of all the current scope locks. */
 	mutable int8 ScopeLockCount;
 
@@ -875,7 +860,4 @@
 	/** Decreases the scope lock count. Runs the waiting to execute delegates if the count drops to zero. */
 	void DecrementListLock() const;
 
-	/** True if the ability block flags are currently enabled */
-	UPROPERTY()
-	bool bIsBlockingOtherAbilities;
 };
diff --git a/Engine/Plugins/Runtime/GameplayAbilities/Source/GameplayAbilities/Public/Abilities/GameplayAbilityTargetActor.h b/Engine/Plugins/Runtime/GameplayAbilities/Source/GameplayAbilities/Public/Abilities/GameplayAbilityTargetActor.h
--- a/Engine/Plugins/Runtime/GameplayAbilities/Source/GameplayAbilities/Public/Abilities/GameplayAbilityTargetActor.h
+++ b/Engine/Plugins/Runtime/GameplayAbilities/Source/GameplayAbilities/Public/Abilities/GameplayAbilityTargetActor.h
@@ -14,18 +14,13 @@ class APlayerController;
 class UAbilitySystemComponent;
 class UGameplayAbility;
 
-UENUM(BlueprintType)
-namespace ETargetAbilitySelfSelection
-{
-	enum Type
-	{
-		TASS_Permit			UMETA(DisplayName = "Allow self-selection"),
-		TASS_Forbid			UMETA(DisplayName = "Forbid self-selection"),
-		TASS_Require 		UMETA(DisplayName = "Force self-selection (add to final data)")
-	};
-}
-
-/** TargetActors are spawned to assist with ability targeting. They are spawned by ability tasks and create/determine the outgoing targeting data passed from one task to another. */
+/**
+ * TargetActors are spawned to assist with ability targeting. They are spawned by ability tasks and create/determine the outgoing targeting data passed from one task to another
+ *
+ * WARNING: These actors are spawned once per ability activation and in their default form are not very efficient
+ * For most games you will need to subclass and heavily modify this actor, or you will want to implement similar functions in a game-specific actor or blueprint to avoid actor spawn costs
+ * This class is not well tested by internal games, but it is a useful class to look at to learn how target replication occurs
+ */
 UCLASS(Blueprintable, abstract, notplaceable)
 class GAMEPLAYABILITIES_API AGameplayAbilityTargetActor : public AActor
 {
@@ -67,10 +62,6 @@ public:
 	/** Replicated target data was received from a client. Possibly sanitize/verify. return true if data is good and we should broadcast it as valid data. */
 	virtual bool OnReplicatedTargetDataReceived(FGameplayAbilityTargetDataHandle& Data) const;
 
-	/** Accessor for checking, before instantiating, if this TargetActor will replicate. */
-	DEPRECATED(4.12, "Call AActor::GetIsReplicated instead")
-	bool GetReplicates() const { return GetIsReplicated(); }
-
 	// ------------------------------
 	
 	FAbilityTargetData	TargetDataReadyDelegate;
diff --git a/Engine/Plugins/Runtime/GameplayAbilities/Source/GameplayAbilities/Public/Abilities/GameplayAbilityTargetDataFilter.h b/Engine/Plugins/Runtime/GameplayAbilities/Source/GameplayAbilities/Public/Abilities/GameplayAbilityTargetDataFilter.h
--- a/Engine/Plugins/Runtime/GameplayAbilities/Source/GameplayAbilities/Public/Abilities/GameplayAbilityTargetDataFilter.h
+++ b/Engine/Plugins/Runtime/GameplayAbilities/Source/GameplayAbilities/Public/Abilities/GameplayAbilityTargetDataFilter.h
@@ -8,6 +8,7 @@
 #include "GameFramework/Actor.h"
 #include "GameplayAbilityTargetDataFilter.generated.h"
 
+/** Set rather it is possible to target self */
 UENUM(BlueprintType)
 namespace ETargetDataFilterSelf
 {
@@ -19,6 +20,7 @@ namespace ETargetDataFilterSelf
 	};
 }
 
+/** Simple actor target filter, games can subclass this */
 USTRUCT(BlueprintType)
 struct GAMEPLAYABILITIES_API FGameplayTargetDataFilter
 {
@@ -28,6 +30,7 @@ struct GAMEPLAYABILITIES_API FGameplayTargetDataFilter
 	{
 	}
 
+	/** Returns true if the actor passes the filter and will be targeted */
 	virtual bool FilterPassesForActor(const AActor* ActorToBeFiltered) const
 	{
 		switch (SelfFilter.GetValue())
@@ -57,6 +60,7 @@ struct GAMEPLAYABILITIES_API FGameplayTargetDataFilter
 		return (bReverseFilter ^ true);
 	}
 
+	/** Initializes SelfActor */
 	void InitializeFilterContext(AActor* FilterActor);
 
 	/** Actor we're comparing against. */
@@ -76,7 +80,7 @@ struct GAMEPLAYABILITIES_API FGameplayTargetDataFilter
 	bool bReverseFilter;
 };
 
-
+/** Polymorphic handle to filter structure that handles checking if actors should be targeted */
 USTRUCT(BlueprintType)
 struct GAMEPLAYABILITIES_API FGameplayTargetDataFilterHandle
 {
@@ -84,6 +88,7 @@ struct GAMEPLAYABILITIES_API FGameplayTargetDataFilterHandle
 
 	TSharedPtr<FGameplayTargetDataFilter> Filter;
 
+	/** Returns true if the actor passes the filter and will be targeted */
 	bool FilterPassesForActor(const AActor* ActorToBeFiltered) const
 	{
 		if (!ActorToBeFiltered)
diff --git a/Engine/Plugins/Runtime/GameplayAbilities/Source/GameplayAbilities/Public/Abilities/GameplayAbilityTargetTypes.h b/Engine/Plugins/Runtime/GameplayAbilities/Source/GameplayAbilities/Public/Abilities/GameplayAbilityTargetTypes.h
--- a/Engine/Plugins/Runtime/GameplayAbilities/Source/GameplayAbilities/Public/Abilities/GameplayAbilityTargetTypes.h
+++ b/Engine/Plugins/Runtime/GameplayAbilities/Source/GameplayAbilities/Public/Abilities/GameplayAbilityTargetTypes.h
@@ -20,16 +20,23 @@ struct FGameplayEffectSpec;
 UENUM(BlueprintType)
 namespace EGameplayTargetingConfirmation
 {
+	/** Describes how the targeting information is confirmed */
 	enum Type
 	{
-		Instant,			// The targeting happens instantly without special logic or user input deciding when to 'fire'.
-		UserConfirmed,		// The targeting happens when the user confirms the targeting.
-		Custom,				// The GameplayTargeting Ability is responsible for deciding when the targeting data is ready. Not supported by all TargetingActors.
-		CustomMulti,		// The GameplayTargeting Ability is responsible for deciding when the targeting data is ready. Not supported by all TargetingActors. Should not destroy upon data production.
+		/** The targeting happens instantly without special logic or user input deciding when to 'fire' */
+		Instant,
+
+		/** The targeting happens when the user confirms the targeting */
+		UserConfirmed,
+
+		/** The GameplayTargeting Ability is responsible for deciding when the targeting data is ready. Not supported by all TargetingActors */
+		Custom,
+
+		/** The GameplayTargeting Ability is responsible for deciding when the targeting data is ready. Not supported by all TargetingActors. Should not destroy upon data production */
+		CustomMulti,
 	};
 }
 
-
 /**
  *	A generic structure for targeting data. We want generic functions to produce this data and other generic
  *	functions to consume this data.
@@ -52,14 +59,9 @@ namespace EGameplayTargetingConfirmation
  * 		-Filter or merge TargetDatas
  *		-Spawn a new actor at a TargetData location
  *
- *
- *
  *	Maybe it is better to distinguish between actor list targeting vs positional targeting data?
  *		-AOE/aura type of targeting data blurs the line
- *
- *
  */
-
 USTRUCT()
 struct GAMEPLAYABILITIES_API FGameplayAbilityTargetData
 {
@@ -67,77 +69,83 @@ struct GAMEPLAYABILITIES_API FGameplayAbilityTargetData
 
 	virtual ~FGameplayAbilityTargetData() { }
 
+	/** Applies a gameplay effect to each target represented */
 	TArray<FActiveGameplayEffectHandle> ApplyGameplayEffect(const UGameplayEffect* GameplayEffect, const FGameplayEffectContextHandle& InEffectContext, float Level, FPredictionKey PredictionKey = FPredictionKey());
 
+	/** Applies a previously created gameplay effect spec to each target represented */
 	virtual TArray<FActiveGameplayEffectHandle> ApplyGameplayEffectSpec(FGameplayEffectSpec& Spec, FPredictionKey PredictionKey = FPredictionKey());
 
+	/** Modifies the context and adds this target data to the target data handle stored within */
 	virtual void AddTargetDataToContext(FGameplayEffectContextHandle& Context, bool bIncludeActorArray) const;
 
+	/** Modifies the cue parameters and adds this target data to the target data handle stored within */
 	virtual void AddTargetDataToGameplayCueParameters(FGameplayCueParameters& Parameters) const;
 
+	/** Returns all actors targeted, almost always overridden */
 	virtual TArray<TWeakObjectPtr<AActor>> GetActors() const
 	{
 		return TArray<TWeakObjectPtr<AActor>>();
 	}
 
+	/** Modify the actor list */
 	virtual bool SetActors(TArray<TWeakObjectPtr<AActor>> NewActorArray)
 	{
 		//By default, we don't keep this data, and therefore can't set it.
 		return false;
 	}
 
-	// -------------------------------------
-
+	/** Return true in subclasses if GetHitResult will work */
 	virtual bool HasHitResult() const
 	{
 		return false;
 	}
 
+	/** Override to return a hit result */
 	virtual const FHitResult* GetHitResult() const
 	{
 		return nullptr;
 	}
 
-	// -------------------------------------
-
+	/** Override to true if GetOrigin will work */
 	virtual bool HasOrigin() const
 	{
 		return false;
 	}
 
+	/** Override to return an origin point, which may be derived from other data */
 	virtual FTransform GetOrigin() const
 	{
 		return FTransform::Identity;
 	}
 
-	// -------------------------------------
-
+	/** Override to true if GetEndPoint/Transform will work */
 	virtual bool HasEndPoint() const
 	{
 		return false;
 	}
 
+	/** Override to return a target/end point */
 	virtual FVector GetEndPoint() const
 	{
 		return FVector::ZeroVector;
 	}
 
+	/** Override to return a transform, default will create one from just the location */
 	virtual FTransform GetEndPointTransform() const
 	{
 		return FTransform(GetEndPoint());
 	}
 
-	// -------------------------------------
-
+	/** Returns the serialization data, must always be overridden */
 	virtual UScriptStruct* GetScriptStruct() const
 	{
 		return FGameplayAbilityTargetData::StaticStruct();
 	}
 
+	/** Returns a debug string representation */
 	virtual FString ToString() const;
 
-	// -------------------------------------
-	//..See notes on delegate definition FOnTargetActorSwapped.
+	/** See notes on delegate definition FOnTargetActorSwapped */
 	virtual bool ShouldCheckForTargetActorSwap() const
 	{
 		return false;
@@ -145,91 +153,93 @@ struct GAMEPLAYABILITIES_API FGameplayAbilityTargetData
 
 };
 
-
 UENUM(BlueprintType)
 namespace EGameplayAbilityTargetingLocationType
 {
-	/**
-	*	What type of location calculation to use when an ability asks for our transform.
-	*/
-
+	/** What type of location calculation to use when an ability asks for our transform */
 	enum Type
 	{
-		LiteralTransform		UMETA(DisplayName = "Literal Transform"),		// We report an actual raw transform. This is also the final fallback if other methods fail.
-		ActorTransform			UMETA(DisplayName = "Actor Transform"),			// We pull the transform from an associated actor directly.
-		SocketTransform			UMETA(DisplayName = "Socket Transform"),		// We aim from a named socket on the player's skeletal mesh component.
+		/** We report an actual raw transform. This is also the final fallback if other methods fail */
+		LiteralTransform		UMETA(DisplayName = "Literal Transform"),
+
+		/** We pull the transform from an associated actor directly */
+		ActorTransform			UMETA(DisplayName = "Actor Transform"),
+
+		/** We aim from a named socket on the player's skeletal mesh component */
+		SocketTransform			UMETA(DisplayName = "Socket Transform"),		
 	};
 }
 
 /**
-*	Handle for Targeting Data. This servers two main purposes:
-*		-Avoid us having to copy around the full targeting data structure in Blueprints
-*		-Allows us to leverage polymorphism in the target data structure
-*		-Allows us to implement NetSerialize and replicate by value between clients/server
-*
-*		-Avoid using UObjects could be used to give us polymorphism and by reference passing in blueprints.
-*		-However we would still be screwed when it came to replication
-*
-*		-Replication by value
-*		-Pass by reference in blueprints
-*		-Polymophism in TargetData structure
-*
-*/
-
+ *	Handle for Targeting Data. This servers two main purposes:
+ *		-Avoid us having to copy around the full targeting data structure in Blueprints
+ *		-Allows us to leverage polymorphism in the target data structure
+ *		-Allows us to implement NetSerialize and replicate by value between clients/server
+ *
+ *		-Avoid using UObjects could be used to give us polymorphism and by reference passing in blueprints.
+ *		-However we would still be screwed when it came to replication
+ *
+ *		-Replication by value
+ *		-Pass by reference in blueprints
+ *		-Polymophism in TargetData structure
+ */
 USTRUCT(BlueprintType)
 struct GAMEPLAYABILITIES_API FGameplayAbilityTargetDataHandle
 {
 	GENERATED_USTRUCT_BODY()
 
 	FGameplayAbilityTargetDataHandle() { }
-	FGameplayAbilityTargetDataHandle(struct FGameplayAbilityTargetData* DataPtr)
+	FGameplayAbilityTargetDataHandle(FGameplayAbilityTargetData* DataPtr)
 	{
 		Data.Add(TSharedPtr<FGameplayAbilityTargetData>(DataPtr));
 	}
 
+	/** Raw storage of target data, do not modify this directly */
 	TArray<TSharedPtr<FGameplayAbilityTargetData>, TInlineAllocator<1> >	Data;
 
+	/** Resets handle to have no targets */
 	void Clear()
 	{
 		Data.Reset();
 	}
 
+	/** Returns number of target data, not number of actors/targets as target data may contain multiple actors */
 	int32 Num() const
 	{
 		return Data.Num();
 	}
 
+	/** Returns true if there are any valid targets */
 	bool IsValid(int32 Index) const
 	{
 		return (Index < Data.Num() && Data[Index].IsValid());
 	}
 
+	/** Returns data at index, or nullptr if invalid */
 	const FGameplayAbilityTargetData* Get(int32 Index) const
 	{
 		return IsValid(Index) ? Data[Index].Get() : nullptr;
 	}
 
+	/** Returns data at index, or nullptr if invalid */
 	FGameplayAbilityTargetData* Get(int32 Index)
 	{
 		return IsValid(Index) ? Data[Index].Get() : nullptr;
 	}
 
-	void Add(struct FGameplayAbilityTargetData* DataPtr)
+	/** Adds a new target data to handle, it must have been created with new */
+	void Add(FGameplayAbilityTargetData* DataPtr)
 	{
 		Data.Add(TSharedPtr<FGameplayAbilityTargetData>(DataPtr));
 	}
 
-	DEPRECATED(4.11, "Pass Handle by reference, not pointer")
-	void Append(FGameplayAbilityTargetDataHandle* OtherHandle)
-	{
-		Append(*OtherHandle);
-	}
-
+	/** Does a shallow copy of target data from one handle to another */
 	void Append(const FGameplayAbilityTargetDataHandle& OtherHandle)
 	{
 		Data.Append(OtherHandle.Data);
 	}
 
+	/** Serialize for networking, handles polymorphism */
 	bool NetSerialize(FArchive& Ar, class UPackageMap* Map, bool& bOutSuccess);
 
 	/** Comparison operator */
@@ -272,20 +282,20 @@ struct TStructOpsTypeTraits<FGameplayAbilityTargetDataHandle> : public TStructOp
 	};
 };
 
+/** Structure that stores a location in one of several different formats */
 USTRUCT(BlueprintType)
 struct GAMEPLAYABILITIES_API FGameplayAbilityTargetingLocationInfo
 {
 	GENERATED_USTRUCT_BODY()
 
 	FGameplayAbilityTargetingLocationInfo()
-	: LocationType(EGameplayAbilityTargetingLocationType::LiteralTransform)
-	, SourceActor(nullptr)
-	, SourceComponent(nullptr)
-	, SourceAbility(nullptr)
-	{
-	};
+		: LocationType(EGameplayAbilityTargetingLocationType::LiteralTransform)
+		, SourceActor(nullptr)
+		, SourceComponent(nullptr)
+		, SourceAbility(nullptr)
+	{ }
 
-	virtual ~FGameplayAbilityTargetingLocationInfo() {};
+	virtual ~FGameplayAbilityTargetingLocationInfo() {}
 
 	void operator=(const FGameplayAbilityTargetingLocationInfo& Other)
 	{
@@ -297,10 +307,10 @@ struct GAMEPLAYABILITIES_API FGameplayAbilityTargetingLocationInfo
 		SourceSocketName = Other.SourceSocketName;
 	}
 
-public:
+	/** Converts internal format into a literal world space transform */
 	FTransform GetTargetingTransform() const
 	{
-		//Return or calculate based on LocationType.
+		// Return or calculate based on LocationType.
 		switch (LocationType)
 		{
 		case EGameplayAbilityTargetingLocationType::ActorTransform:
@@ -312,21 +322,26 @@ public:
 		case EGameplayAbilityTargetingLocationType::SocketTransform:
 			if (SourceComponent)
 			{
-				return SourceComponent->GetSocketTransform(SourceSocketName);		//Bad socket name will just return component transform anyway, so we're safe
+				// Bad socket name will just return component transform anyway, so we're safe
+				return SourceComponent->GetSocketTransform(SourceSocketName);
 			}
 			break;
 		case EGameplayAbilityTargetingLocationType::LiteralTransform:
 			return LiteralTransform;
 		default:
-			check(false);		//This case should not happen
+			check(false);
 			break;
 		}
-		//Error
+
+		// It cannot get here
 		return FTransform::Identity;
 	}
 
+	/** Initializes new target data and fills in with hit results */
 	FGameplayAbilityTargetDataHandle MakeTargetDataHandleFromHitResult(TWeakObjectPtr<UGameplayAbility> Ability, const FHitResult& HitResult) const;
 	FGameplayAbilityTargetDataHandle MakeTargetDataHandleFromHitResults(TWeakObjectPtr<UGameplayAbility> Ability, const TArray<FHitResult>& HitResults) const;
+
+	/** Initializes new actor list target data, and sets this as the origin */
 	FGameplayAbilityTargetDataHandle MakeTargetDataHandleFromActors(const TArray<TWeakObjectPtr<AActor>>& TargetActors, bool OneActorPerHandle = false) const;
 
 	/** Type of location used - will determine what data is transmitted over the network and what fields are used when calculating position. */
@@ -353,19 +368,8 @@ public:
 	UPROPERTY(BlueprintReadWrite, meta = (ExposeOnSpawn = true), Category = Targeting)
 	FName SourceSocketName;
 
-	// -------------------------------------
-
-	virtual FString ToString() const
-	{
-		return TEXT("FGameplayAbilityTargetingLocationInfo");
-	}
-
+	/** Optimized serialize function */
 	bool NetSerialize(FArchive& Ar, class UPackageMap* Map, bool& bOutSuccess);
-
-	virtual UScriptStruct* GetScriptStruct() const
-	{
-		return FGameplayAbilityTargetingLocationInfo::StaticStruct();
-	}
 };
 
 template<>
@@ -377,6 +381,7 @@ struct TStructOpsTypeTraits<FGameplayAbilityTargetingLocationInfo> : public TStr
 	};
 };
 
+/** Target data with just a source and target location in space */
 USTRUCT(BlueprintType)
 struct GAMEPLAYABILITIES_API FGameplayAbilityTargetData_LocationInfo : public FGameplayAbilityTargetData
 {
@@ -438,6 +443,7 @@ struct TStructOpsTypeTraits<FGameplayAbilityTargetData_LocationInfo> : public TS
 	};
 };
 
+/** Target data with a source location and a list of targeted actors, makes sense for AOE attacks */
 USTRUCT(BlueprintType)
 struct GAMEPLAYABILITIES_API FGameplayAbilityTargetData_ActorArray : public FGameplayAbilityTargetData
 {
@@ -540,6 +546,7 @@ struct TStructOpsTypeTraits<FGameplayAbilityTargetData_ActorArray> : public TStr
 	};
 };
 
+/** Target data with a single hit result, data is packed into the hit result */
 USTRUCT(BlueprintType)
 struct GAMEPLAYABILITIES_API FGameplayAbilityTargetData_SingleTargetHit : public FGameplayAbilityTargetData
 {
@@ -600,6 +607,7 @@ struct GAMEPLAYABILITIES_API FGameplayAbilityTargetData_SingleTargetHit : public
 
 	// -------------------------------------
 
+	/** Hit result that stores data */
 	UPROPERTY()
 	FHitResult	HitResult;
 
@@ -623,9 +631,6 @@ struct TStructOpsTypeTraits<FGameplayAbilityTargetData_SingleTargetHit> : public
 /** Generic callback for returning when target data is available */
 DECLARE_MULTICAST_DELEGATE_OneParam(FAbilityTargetData, const FGameplayAbilityTargetDataHandle&);
 
-
-// ----------------------------------------------------
-
 /** Generic callback for returning when target data is available */
 DECLARE_MULTICAST_DELEGATE_TwoParams(FAbilityTargetDataSetDelegate, const FGameplayAbilityTargetDataHandle&, FGameplayTag);
 
@@ -658,19 +663,22 @@ namespace EAbilityGenericReplicatedEvent
 	};
 }
 
+/** Payload for generic replicated events */
 struct FAbilityReplicatedData
 {
 	FAbilityReplicatedData() : bTriggered(false), VectorPayload(ForceInitToZero) {}
+
 	/** Event has triggered */
 	bool bTriggered;
 
 	/** Optional Vector payload for event */
 	FVector_NetQuantize100 VectorPayload;
 
+	/** Delegate that will be called on replication */
 	FSimpleMulticastDelegate Delegate;
 };
 
-/** Struct defining the cached data for a specific gameplay ability. This data is generally syncronized client->server in a network game. */
+/** Struct defining the cached data for a specific gameplay ability. This data is generally synchronized client->server in a network game. */
 struct FAbilityReplicatedDataCache
 {
 	/** What elements this activation is targeting */
diff --git a/Engine/Plugins/Runtime/GameplayAbilities/Source/GameplayAbilities/Public/Abilities/GameplayAbilityTypes.h b/Engine/Plugins/Runtime/GameplayAbilities/Source/GameplayAbilities/Public/Abilities/GameplayAbilityTypes.h
--- a/Engine/Plugins/Runtime/GameplayAbilities/Source/GameplayAbilities/Public/Abilities/GameplayAbilityTypes.h
+++ b/Engine/Plugins/Runtime/GameplayAbilities/Source/GameplayAbilities/Public/Abilities/GameplayAbilityTypes.h
@@ -22,11 +22,14 @@ class UAnimInstance;
 class UAnimMontage;
 class UGameplayAbility;
 class UMovementComponent;
+class UDataTable;
+class UGameplayTask;
 
 GAMEPLAYABILITIES_API DECLARE_LOG_CATEGORY_EXTERN(LogAbilitySystemComponent, Log, All);
 
 #define ENABLE_ABILITYTASK_DEBUGMSG !(UE_BUILD_SHIPPING | UE_BUILD_TEST)
 
+
 UENUM(BlueprintType)
 namespace EGameplayAbilityInstancingPolicy
 {
@@ -34,7 +37,6 @@ namespace EGameplayAbilityInstancingPolicy
 	 *	How the ability is instanced when executed. This limits what an ability can do in its implementation. For example, a NonInstanced
 	 *	Ability cannot have state. It is probably unsafe for an InstancedPerActor ability to have latent actions, etc.
 	 */
-
 	enum Type
 	{
 		// This ability is never instanced. Anything that executes the ability is operating on the CDO.
@@ -51,10 +53,7 @@ namespace EGameplayAbilityInstancingPolicy
 UENUM(BlueprintType)
 namespace EGameplayAbilityNetExecutionPolicy
 {
-	/**
-	 *	How does an ability execute on the network. Does a client "ask and predict", "ask and wait", "don't ask (just do it)"
-	 */
-
+	/** How does an ability execute on the network. Does a client "ask and predict", "ask and wait", "don't ask (just do it)" */
 	enum Type
 	{
 		// Part of this ability runs predictively on the local client if there is one
@@ -74,10 +73,7 @@ namespace EGameplayAbilityNetExecutionPolicy
 UENUM(BlueprintType)
 namespace EGameplayAbilityReplicationPolicy
 {
-	/**
-	 *	How an ability replicates state/events to everyone on the network.
-	 */
-
+	/** How an ability replicates state/events to everyone on the network */
 	enum Type
 	{
 		// We don't replicate the instance of the ability to anyone.
@@ -88,12 +84,10 @@ namespace EGameplayAbilityReplicationPolicy
 	};
 }
 
-
 UENUM(BlueprintType)
 namespace EGameplayAbilityTriggerSource
 {
 	/**	Defines what type of trigger will activate the ability, paired to a tag */
-
 	enum Type
 	{
 		// Triggered from a gameplay event, will come with payload
@@ -108,8 +102,6 @@ namespace EGameplayAbilityTriggerSource
 }
 
 
-// ----------------------------------------------------
-
 /**
  *	FGameplayAbilityActorInfo
  *
@@ -176,15 +168,15 @@ struct GAMEPLAYABILITIES_API FGameplayAbilityActorInfo
 	virtual void ClearActorInfo();
 };
 
-// ----------------------------------------------------
 
+/** Used as a key for storing internal ability data */
 USTRUCT()
 struct FGameplayAbilitySpecHandleAndPredictionKey
 {
 	GENERATED_USTRUCT_BODY()
 
 	FGameplayAbilitySpecHandleAndPredictionKey()
-	: PredictionKeyAtCreation(0)
+		: PredictionKeyAtCreation(0)
 	{}
 
 	FGameplayAbilitySpecHandleAndPredictionKey(const FGameplayAbilitySpecHandle& HandleRef, const FPredictionKey& PredictionKeyAtCreationRef)
@@ -216,9 +208,6 @@ private:
 };
 
 
-// ----------------------------------------------------
-
-
 /** Data about montages that is replicated to simulated clients */
 USTRUCT()
 struct GAMEPLAYABILITIES_API FGameplayAbilityRepAnimMontage
@@ -253,6 +242,7 @@ struct GAMEPLAYABILITIES_API FGameplayAbilityRepAnimMontage
 	UPROPERTY()
 	uint8 ForcePlayBit : 1;
 
+	/** Stops montage position from replicating at all to save bandwidth */
 	UPROPERTY()
 	uint8 SkipPositionCorrection : 1;
 
@@ -272,6 +262,7 @@ struct GAMEPLAYABILITIES_API FGameplayAbilityRepAnimMontage
 	FPredictionKey PredictionKey;
 };
 
+
 /** Data about montages that were played locally (all montages in case of server. predictive montages in case of client). Never replicated directly. */
 USTRUCT()
 struct GAMEPLAYABILITIES_API FGameplayAbilityLocalAnimMontage
@@ -279,16 +270,19 @@ struct GAMEPLAYABILITIES_API FGameplayAbilityLocalAnimMontage
 	GENERATED_USTRUCT_BODY()
 
 	FGameplayAbilityLocalAnimMontage()
-	: AnimMontage(nullptr), PlayBit(false), AnimatingAbility(nullptr)
+		: AnimMontage(nullptr), PlayBit(false), AnimatingAbility(nullptr)
 	{
 	}
 
+	/** What montage is being played */
 	UPROPERTY()
 	UAnimMontage* AnimMontage;
 
+	/** Rather the montage is actively playing */
 	UPROPERTY()
 	bool PlayBit;
 
+	/** Prediction key that started the montage play */
 	UPROPERTY()
 	FPredictionKey PredictionKey;
 
@@ -297,71 +291,76 @@ struct GAMEPLAYABILITIES_API FGameplayAbilityLocalAnimMontage
 	UGameplayAbility* AnimatingAbility;
 };
 
-// ----------------------------------------------------
 
+/** Metadata for a tag-based Gameplay Event, that can activate other abilities or run ability-specific logic */
 USTRUCT(BlueprintType)
 struct GAMEPLAYABILITIES_API FGameplayEventData
 {
 	GENERATED_USTRUCT_BODY()
 
 	FGameplayEventData()
-	: Instigator(nullptr)
-	, Target(nullptr)
-	, OptionalObject(nullptr)
-	, OptionalObject2(nullptr)
-	, EventMagnitude(0.f)
+		: Instigator(nullptr)
+		, Target(nullptr)
+		, OptionalObject(nullptr)
+		, OptionalObject2(nullptr)
+		, EventMagnitude(0.f)
 	{
 	}
 	
-	// Tag of the event that triggered this
+	/** Tag of the event that triggered this */
 	UPROPERTY(EditAnywhere, BlueprintReadWrite, Category = GameplayAbilityTriggerPayload)
 	FGameplayTag EventTag;
 
-	// The instigator of the event
+	/** The instigator of the event */
 	UPROPERTY(EditAnywhere, BlueprintReadWrite, Category = GameplayAbilityTriggerPayload)
 	const AActor* Instigator;
 
-	// The target of the event
+	/** The target of the event */
 	UPROPERTY(EditAnywhere, BlueprintReadWrite, Category = GameplayAbilityTriggerPayload)
 	const AActor* Target;
 
-	// An optional ability-specific object to be passed though the event
+	/** An optional ability-specific object to be passed though the event */
 	UPROPERTY(EditAnywhere, BlueprintReadWrite, Category = GameplayAbilityTriggerPayload)
 	const UObject* OptionalObject;
 
-	// A second optional ability-specific object to be passed though the event
+	/** A second optional ability-specific object to be passed though the event */
 	UPROPERTY(EditAnywhere, BlueprintReadWrite, Category = GameplayAbilityTriggerPayload)
 	const UObject* OptionalObject2;
 
-	// Polymorphic context information
+	/** Polymorphic context information */
 	UPROPERTY(EditAnywhere, BlueprintReadWrite, Category = GameplayAbilityTriggerPayload)
 	FGameplayEffectContextHandle ContextHandle;
 
-	// Tags that the instigator has
+	/** Tags that the instigator has */
 	UPROPERTY(EditAnywhere, BlueprintReadWrite, Category = GameplayAbilityTriggerPayload)
 	FGameplayTagContainer InstigatorTags;
 
-	// Tags that the target has
+	/** Tags that the target has */
 	UPROPERTY(EditAnywhere, BlueprintReadWrite, Category = GameplayAbilityTriggerPayload)
 	FGameplayTagContainer TargetTags;
 
-	// The magnitude of the triggering event
+	/** The magnitude of the triggering event */
 	UPROPERTY(EditAnywhere, BlueprintReadWrite, Category = GameplayAbilityTriggerPayload)
 	float EventMagnitude;
 
-	// The magnitude of the triggering event
+	/** The polymorphic target information for the event */
 	UPROPERTY(EditAnywhere, BlueprintReadWrite, Category = GameplayAbilityTriggerPayload)
-	FGameplayAbilityTargetDataHandle	TargetData;
+	FGameplayAbilityTargetDataHandle TargetData;
 };
 
-/** Used for cleaning up predicted data on network clients */
+/** Delegate for handling gameplay event data */
 DECLARE_MULTICAST_DELEGATE_OneParam(FGameplayEventMulticastDelegate, const FGameplayEventData*);
 
+/** Delegate for handling gameplay event data, includes tag as the Event Data does not always have it filled out */
+DECLARE_MULTICAST_DELEGATE_TwoParams(FGameplayEventTagMulticastDelegate, FGameplayTag, const FGameplayEventData*);
+
+
 /** Ability Ended Data */
 USTRUCT(BlueprintType)
 struct FAbilityEndedData
 {
 	GENERATED_USTRUCT_BODY()
+
 	FAbilityEndedData()
 		: AbilityThatEnded(nullptr)
 		, bReplicateEndAbility(false)
@@ -377,15 +376,19 @@ struct FAbilityEndedData
 	{
 	}
 
+	/** Ability that ended, normally instance but could be CDO */
 	UPROPERTY()
 	UGameplayAbility* AbilityThatEnded;
 
+	/** Specific ability spec that ended */
 	UPROPERTY()
 	FGameplayAbilitySpecHandle AbilitySpecHandle;
 
+	/** Rather to replicate the ability to ending */
 	UPROPERTY()
 	bool bReplicateEndAbility;
 
+	/** True if this was cancelled deliberately, false if it ended normally */
 	UPROPERTY()
 	bool bWasCancelled;
 };
@@ -393,18 +396,16 @@ struct FAbilityEndedData
 /** Notification delegate definition for when the gameplay ability ends */
 DECLARE_MULTICAST_DELEGATE_OneParam(FGameplayAbilityEndedDelegate, const FAbilityEndedData&);
 
-/** 
- *	Structure that tells AbilitySystemComponent what to bind to an InputComponent (see BindAbilityActivationToInputComponent) 
- *	
- */
+
+/** Structure that tells AbilitySystemComponent what to bind to an InputComponent (see BindAbilityActivationToInputComponent) */
 struct FGameplayAbilityInputBinds
 {
 	FGameplayAbilityInputBinds(FString InConfirmTargetCommand, FString InCancelTargetCommand, FString InEnumName, int32 InConfirmTargetInputID = INDEX_NONE, int32 InCancelTargetInputID = INDEX_NONE)
-	: ConfirmTargetCommand(InConfirmTargetCommand)
-	, CancelTargetCommand(InCancelTargetCommand)
-	, EnumName(InEnumName)
-	, ConfirmTargetInputID(InConfirmTargetInputID)
-	, CancelTargetInputID(InCancelTargetInputID)
+		: ConfirmTargetCommand(InConfirmTargetCommand)
+		, CancelTargetCommand(InCancelTargetCommand)
+		, EnumName(InEnumName)
+		, ConfirmTargetInputID(InConfirmTargetInputID)
+		, CancelTargetInputID(InCancelTargetInputID)
 	{ }
 
 	/** Defines command string that will be bound to Confirm Targeting */
@@ -419,32 +420,45 @@ struct FGameplayAbilityInputBinds
 	/** If >=0, Confirm is bound to an entry in the enum */
 	int32 ConfirmTargetInputID;
 
-	/** If >=0, Confirm is bound to an entry in the enum */
+	/** If >=0, Cancel is bound to an entry in the enum */
 	int32 CancelTargetInputID;
 
 	UEnum* GetBindEnum() { return FindObject<UEnum>(ANY_PACKAGE, *EnumName); }
 };
 
+
+/** Used to initialize default values for attributes */
 USTRUCT()
 struct GAMEPLAYABILITIES_API FAttributeDefaults
 {
 	GENERATED_USTRUCT_BODY()
 
+	FAttributeDefaults()
+		: DefaultStartingTable(nullptr)
+	{ }
+
 	UPROPERTY(EditAnywhere, Category = "AttributeTest")
 	TSubclassOf<UAttributeSet> Attributes;
 
 	UPROPERTY(EditAnywhere, Category = "AttributeTest")
-	class UDataTable*	DefaultStartingTable;
+	UDataTable* DefaultStartingTable;
 };
 
+
+/** Debug message emitted by ability tasks */
 USTRUCT()
 struct GAMEPLAYABILITIES_API FAbilityTaskDebugMessage
 {
 	GENERATED_USTRUCT_BODY()
 
+	FAbilityTaskDebugMessage()
+		: FromTask(nullptr)
+	{ }
+
 	UPROPERTY()
-	class UGameplayTask*	FromTask;
+	UGameplayTask*	FromTask;
 
+	UPROPERTY()
 	FString Message;
 };
 
@@ -454,7 +468,6 @@ DECLARE_MULTICAST_DELEGATE(FAbilitySystemComponentPredictionKeyClear);
 /** Generic delegate for ability 'events'/notifies */
 DECLARE_MULTICAST_DELEGATE_OneParam(FGenericAbilityDelegate, UGameplayAbility*);
 
-// ----------------------------------------------------
 
 /** This struct holds state to batch server RPC calls: ServerTryActivateAbility, ServerSetReplicatedTargetData, ServerEndAbility.  */
 USTRUCT()
@@ -493,6 +506,7 @@ struct FServerAbilityRPCBatch
 	}
 };
 
+
 /** Helper struct for defining ServerRPC batch windows. If null ASC is passed in, this becomes a noop. */
 struct GAMEPLAYABILITIES_API FScopedServerAbilityRPCBatcher
 {
diff --git a/Engine/Plugins/Runtime/GameplayAbilities/Source/GameplayAbilities/Public/Abilities/Tasks/AbilityTask_ApplyRootMotion_Base.h b/Engine/Plugins/Runtime/GameplayAbilities/Source/GameplayAbilities/Public/Abilities/Tasks/AbilityTask_ApplyRootMotion_Base.h
--- a/Engine/Plugins/Runtime/GameplayAbilities/Source/GameplayAbilities/Public/Abilities/Tasks/AbilityTask_ApplyRootMotion_Base.h
+++ b/Engine/Plugins/Runtime/GameplayAbilities/Source/GameplayAbilities/Public/Abilities/Tasks/AbilityTask_ApplyRootMotion_Base.h
@@ -8,9 +8,10 @@
 class UCharacterMovementComponent;
 enum class ERootMotionFinishVelocityMode : uint8;
 
-//This delegate can be used to support target swapping on abilities.  e.g. If a decoy is created and you want root motion to switch the destination to the decoy.
+/** This delegate can be used to support target swapping on abilities.  e.g. If a decoy is created and you want root motion to switch the destination to the decoy */
 DECLARE_MULTICAST_DELEGATE_TwoParams(FOnTargetActorSwapped, AActor*, AActor*);
 
+/** Base class for ability tasks that apply root motion */
 UCLASS(MinimalAPI)
 class UAbilityTask_ApplyRootMotion_Base : public UAbilityTask
 {
diff --git a/Engine/Plugins/Runtime/GameplayAbilities/Source/GameplayAbilities/Public/Abilities/Tasks/AbilityTask_PlayMontageAndWait.h b/Engine/Plugins/Runtime/GameplayAbilities/Source/GameplayAbilities/Public/Abilities/Tasks/AbilityTask_PlayMontageAndWait.h
--- a/Engine/Plugins/Runtime/GameplayAbilities/Source/GameplayAbilities/Public/Abilities/Tasks/AbilityTask_PlayMontageAndWait.h
+++ b/Engine/Plugins/Runtime/GameplayAbilities/Source/GameplayAbilities/Public/Abilities/Tasks/AbilityTask_PlayMontageAndWait.h
@@ -9,6 +9,7 @@
 
 DECLARE_DYNAMIC_MULTICAST_DELEGATE(FMontageWaitSimpleDelegate);
 
+/** Ability task to simply play a montage. Many games will want to make a modified version of this task that looks for game-specific events */
 UCLASS()
 class GAMEPLAYABILITIES_API UAbilityTask_PlayMontageAndWait : public UAbilityTask
 {
@@ -58,9 +59,18 @@ private:
 	FOnMontageEnded MontageEndedDelegate;
 	FDelegateHandle InterruptedHandle;
 
+	UPROPERTY()
 	UAnimMontage* MontageToPlay;
-	float Rate;	
+
+	UPROPERTY()
+	float Rate;
+
+	UPROPERTY()
 	FName StartSection;
+
+	UPROPERTY()
 	float AnimRootMotionTranslationScale;
+
+	UPROPERTY()
 	bool bStopWhenAbilityEnds;
 };
diff --git a/Engine/Plugins/Runtime/GameplayAbilities/Source/GameplayAbilities/Public/Abilities/Tasks/AbilityTask_WaitGameplayEvent.h b/Engine/Plugins/Runtime/GameplayAbilities/Source/GameplayAbilities/Public/Abilities/Tasks/AbilityTask_WaitGameplayEvent.h
--- a/Engine/Plugins/Runtime/GameplayAbilities/Source/GameplayAbilities/Public/Abilities/Tasks/AbilityTask_WaitGameplayEvent.h
+++ b/Engine/Plugins/Runtime/GameplayAbilities/Source/GameplayAbilities/Public/Abilities/Tasks/AbilityTask_WaitGameplayEvent.h
@@ -21,11 +21,12 @@ class GAMEPLAYABILITIES_API UAbilityTask_WaitGameplayEvent : public UAbilityTask
 	FWaitGameplayEventDelegate	EventReceived;
 
 	/**
-	 * 	Wait until the specified gameplay tag event is triggered. By default this will look at the owner of this ability. OptionalExternalTarget can be set to make this look at another actor's tags for changes. 
-	 *  It will keep listening as long as OnlyTriggerOnce = false.
+	 * Wait until the specified gameplay tag event is triggered. By default this will look at the owner of this ability. OptionalExternalTarget can be set to make this look at another actor's tags for changes
+	 * It will keep listening as long as OnlyTriggerOnce = false
+	 * If OnlyMatchExact = false it will trigger for nested tags
 	 */
 	UFUNCTION(BlueprintCallable, Category = "Ability|Tasks", meta = (HidePin = "OwningAbility", DefaultToSelf = "OwningAbility", BlueprintInternalUseOnly = "TRUE"))
-	static UAbilityTask_WaitGameplayEvent* WaitGameplayEvent(UGameplayAbility* OwningAbility, FGameplayTag EventTag, AActor* OptionalExternalTarget=nullptr, bool OnlyTriggerOnce=false);
+	static UAbilityTask_WaitGameplayEvent* WaitGameplayEvent(UGameplayAbility* OwningAbility, FGameplayTag EventTag, AActor* OptionalExternalTarget=nullptr, bool OnlyTriggerOnce=false, bool OnlyMatchExact = true);
 
 	void SetExternalTarget(AActor* Actor);
 
@@ -34,6 +35,7 @@ class GAMEPLAYABILITIES_API UAbilityTask_WaitGameplayEvent : public UAbilityTask
 	virtual void Activate() override;
 
 	virtual void GameplayEventCallback(const FGameplayEventData* Payload);
+	virtual void GameplayEventContainerCallback(FGameplayTag MatchingTag, const FGameplayEventData* Payload);
 
 	void OnDestroy(bool AbilityEnding) override;
 
@@ -44,6 +46,7 @@ class GAMEPLAYABILITIES_API UAbilityTask_WaitGameplayEvent : public UAbilityTask
 
 	bool UseExternalTarget;	
 	bool OnlyTriggerOnce;
+	bool OnlyMatchExact;
 
 	FDelegateHandle MyHandle;
 };
diff --git a/Engine/Plugins/Runtime/GameplayAbilities/Source/GameplayAbilities/Public/Abilities/Tasks/AbilityTask_WaitTargetData.h b/Engine/Plugins/Runtime/GameplayAbilities/Source/GameplayAbilities/Public/Abilities/Tasks/AbilityTask_WaitTargetData.h
--- a/Engine/Plugins/Runtime/GameplayAbilities/Source/GameplayAbilities/Public/Abilities/Tasks/AbilityTask_WaitTargetData.h
+++ b/Engine/Plugins/Runtime/GameplayAbilities/Source/GameplayAbilities/Public/Abilities/Tasks/AbilityTask_WaitTargetData.h
@@ -11,7 +11,13 @@
 
 DECLARE_DYNAMIC_MULTICAST_DELEGATE_OneParam(FWaitTargetDataDelegate, const FGameplayAbilityTargetDataHandle&, Data);
 
-/** Wait for targeting actor (spawned from parameter) to provide data. Can be set not to end upon outputting data. Can be ended by task name. */
+/**
+ * Wait for targeting actor (spawned from parameter) to provide data. Can be set not to end upon outputting data. Can be ended by task name.
+ *
+ * WARNING: These actors are spawned once per ability activation and in their default form are not very efficient
+ * For most games you will need to subclass and heavily modify this actor, or you will want to implement similar functions in a game-specific actor or blueprint to avoid actor spawn costs
+ * This task is not well tested by internal games, but it is a useful class to look at to learn how target replication occurs
+ */
 UCLASS(notplaceable)
 class GAMEPLAYABILITIES_API UAbilityTask_WaitTargetData: public UAbilityTask
 {
@@ -71,6 +77,7 @@ protected:
 
 protected:
 
+	UPROPERTY()
 	TSubclassOf<AGameplayAbilityTargetActor> TargetClass;
 
 	/** The TargetActor that we spawned */
