Copying //UE4/Dev-Framework to //UE4/Dev-Main (Source: //UE4/Dev-Framework @ 4058146)

diff --git a/Engine/Plugins/Runtime/GameplayAbilities/Source/GameplayAbilities/Public/AbilitySystemBlueprintLibrary.h b/Engine/Plugins/Runtime/GameplayAbilities/Source/GameplayAbilities/Public/AbilitySystemBlueprintLibrary.h
--- a/Engine/Plugins/Runtime/GameplayAbilities/Source/GameplayAbilities/Public/AbilitySystemBlueprintLibrary.h
+++ b/Engine/Plugins/Runtime/GameplayAbilities/Source/GameplayAbilities/Public/AbilitySystemBlueprintLibrary.h
@@ -19,18 +19,22 @@
 class UAbilitySystemComponent;
 class UGameplayEffect;
 
-// meta =(RestrictedToClasses="GameplayAbility")
+/** Blueprint library for ability system. Many of these functions are useful to call from native as well */
 UCLASS(meta=(ScriptName="AbilitySystemLibrary"))
 class GAMEPLAYABILITIES_API UAbilitySystemBlueprintLibrary : public UBlueprintFunctionLibrary
 {
 	GENERATED_UCLASS_BODY()
 
+	/** Tries to find an ability system component on the actor, will use AbilitySystemInterface */
 	UFUNCTION(BlueprintPure, Category = Ability)
 	static UAbilitySystemComponent* GetAbilitySystemComponent(AActor *Actor);
 
-	// NOTE: The Actor passed in must implement IAbilitySystemInterface! or else this function will silently fail to
-	// send the event.  The actor needs the interface to find the UAbilitySystemComponent, and if the component isn't
-	// found, the event will not be sent.
+	/**
+	 * This function can be used to trigger an ability on the actor in question with useful payload data.
+	 * NOTE: The Actor passed in must implement IAbilitySystemInterface! or else this function will silently fail to
+	 * send the event.  The actor needs the interface to find the UAbilitySystemComponent, and if the component isn't
+	 * found, the event will not be sent.
+	 */
 	UFUNCTION(BlueprintCallable, Category = Ability, Meta = (Tooltip = "This function can be used to trigger an ability on the actor in question with useful payload data."))
 	static void SendGameplayEventToActor(AActor* Actor, FGameplayTag EventTag, FGameplayEventData Payload);
 	
@@ -38,32 +42,33 @@ class GAMEPLAYABILITIES_API UAbilitySystemBlueprintLibrary : public UBlueprintFu
 	//		Attribute
 	// -------------------------------------------------------------------------------
 
+	/** Returns true if the attribute actually exists */
 	UFUNCTION(BlueprintPure, Category = "Ability|Attribute")
 	static bool IsValid(FGameplayAttribute Attribute);
 
 	/** Returns the value of Attribute from the ability system component belonging to Actor. */
 	UFUNCTION(BlueprintPure, Category = "Ability|Attribute")
-	static float GetFloatAttribute(const class AActor* Actor, FGameplayAttribute Attribute, bool& bSuccessfullyFoundAttribute);
+	static float GetFloatAttribute(const AActor* Actor, FGameplayAttribute Attribute, bool& bSuccessfullyFoundAttribute);
 
 	/** Returns the value of Attribute from the ability system component AbilitySystem. */
 	UFUNCTION(BlueprintPure, Category = "Ability|Attribute")
-	static float GetFloatAttributeFromAbilitySystemComponent(const class UAbilitySystemComponent* AbilitySystem, FGameplayAttribute Attribute, bool& bSuccessfullyFoundAttribute);
+	static float GetFloatAttributeFromAbilitySystemComponent(const UAbilitySystemComponent* AbilitySystem, FGameplayAttribute Attribute, bool& bSuccessfullyFoundAttribute);
 
 	/** Returns the base value of Attribute from the ability system component belonging to Actor. */
 	UFUNCTION(BlueprintPure, Category = "Ability|Attribute")
-	static float GetFloatAttributeBase(const class AActor* Actor, FGameplayAttribute Attribute, bool& bSuccessfullyFoundAttribute);
+	static float GetFloatAttributeBase(const AActor* Actor, FGameplayAttribute Attribute, bool& bSuccessfullyFoundAttribute);
 
 	/** Returns the base value of Attribute from the ability system component AbilitySystemComponent. */
 	UFUNCTION(BlueprintPure, Category = "Ability|Attribute")
-	static float GetFloatAttributeBaseFromAbilitySystemComponent(const class UAbilitySystemComponent* AbilitySystemComponent, FGameplayAttribute Attribute, bool& bSuccessfullyFoundAttribute);
+	static float GetFloatAttributeBaseFromAbilitySystemComponent(const UAbilitySystemComponent* AbilitySystemComponent, FGameplayAttribute Attribute, bool& bSuccessfullyFoundAttribute);
 
 	/** Returns the value of Attribute from the ability system component AbilitySystem after evaluating it with source and target tags. bSuccess indicates the success or failure of this operation. */
 	UFUNCTION(BlueprintPure, Category = "Ability|Attribute")
-	static float EvaluateAttributeValueWithTags(class UAbilitySystemComponent* AbilitySystem, FGameplayAttribute Attribute, const FGameplayTagContainer& SourceTags, const FGameplayTagContainer& TargetTags, bool& bSuccess);
+	static float EvaluateAttributeValueWithTags(UAbilitySystemComponent* AbilitySystem, FGameplayAttribute Attribute, const FGameplayTagContainer& SourceTags, const FGameplayTagContainer& TargetTags, bool& bSuccess);
 
 	/** Returns the value of Attribute from the ability system component AbilitySystem after evaluating it with source and target tags using the passed in base value instead of the real base value. bSuccess indicates the success or failure of this operation. */
 	UFUNCTION(BlueprintPure, Category = "Ability|Attribute")
-	static float EvaluateAttributeValueWithTagsAndBase(class UAbilitySystemComponent* AbilitySystem, FGameplayAttribute Attribute, const FGameplayTagContainer& SourceTags, const FGameplayTagContainer& TargetTags, float BaseValue, bool& bSuccess);
+	static float EvaluateAttributeValueWithTagsAndBase(UAbilitySystemComponent* AbilitySystem, FGameplayAttribute Attribute, const FGameplayTagContainer& SourceTags, const FGameplayTagContainer& TargetTags, float BaseValue, bool& bSuccess);
 
 	/** Simple equality operator for gameplay attributes */
 	UFUNCTION(BlueprintPure, meta=(DisplayName = "Equal (Gameplay Attribute)", CompactNodeTitle = "==", Keywords = "== equal"), Category="Ability|Attribute")
@@ -77,21 +82,27 @@ class GAMEPLAYABILITIES_API UAbilitySystemBlueprintLibrary : public UBlueprintFu
 	//		TargetData
 	// -------------------------------------------------------------------------------
 
+	/** Copies targets from HandleToAdd to TargetHandle */
 	UFUNCTION(BlueprintCallable, Category = "Ability|TargetData")
 	static FGameplayAbilityTargetDataHandle AppendTargetDataHandle(FGameplayAbilityTargetDataHandle TargetHandle, const FGameplayAbilityTargetDataHandle& HandleToAdd);
 
+	/** Creates a target data with a source and destination location */
 	UFUNCTION(BlueprintPure, Category = "Ability|TargetData")
 	static FGameplayAbilityTargetDataHandle	AbilityTargetDataFromLocations(const FGameplayAbilityTargetingLocationInfo& SourceLocation, const FGameplayAbilityTargetingLocationInfo& TargetLocation);
 
+	/** Creates a target data with a single hit result */
 	UFUNCTION(BlueprintPure, Category = "Ability|TargetData")
 	static FGameplayAbilityTargetDataHandle	AbilityTargetDataFromHitResult(const FHitResult& HitResult);
 
+	/** Returns number of target data objects, not necessarily number of distinct targets */
 	UFUNCTION(BlueprintPure, Category = "Ability|TargetData")
 	static int32 GetDataCountFromTargetData(const FGameplayAbilityTargetDataHandle& TargetData);
 
+	/** Creates single actor target data */
 	UFUNCTION(BlueprintPure, Category = "Ability|TargetData")
 	static FGameplayAbilityTargetDataHandle	AbilityTargetDataFromActor(AActor* Actor);
 
+	/** Creates actor array target data */
 	UFUNCTION(BlueprintPure, Category = "Ability|TargetData")
 	static FGameplayAbilityTargetDataHandle	AbilityTargetDataFromActorArray(const TArray<AActor*>& ActorArray, bool OneTargetPerHandle);
 
@@ -111,6 +122,7 @@ class GAMEPLAYABILITIES_API UAbilitySystemBlueprintLibrary : public UBlueprintFu
 	UFUNCTION(BlueprintPure, Category = "Spec")
 	static FGameplayEffectSpecHandle CloneSpecHandle(AActor* InNewInstigator, AActor* InEffectCauser, FGameplayEffectSpecHandle GameplayEffectSpecHandle_Clone);
 
+	/** Returns all actors targeted */
 	UFUNCTION(BlueprintPure, Category = "Ability|TargetData")
 	static TArray<AActor*> GetActorsFromTargetData(const FGameplayAbilityTargetDataHandle& TargetData, int32 Index);
 
@@ -122,24 +134,31 @@ class GAMEPLAYABILITIES_API UAbilitySystemBlueprintLibrary : public UBlueprintFu
 	UFUNCTION(BlueprintPure, Category = "Ability|TargetData")
 	static bool TargetDataHasActor(const FGameplayAbilityTargetDataHandle& TargetData, int32 Index);
 
+	/** Returns true if the target data has a hit result */
 	UFUNCTION(BlueprintPure, Category = "Ability|TargetData")
 	static bool TargetDataHasHitResult(const FGameplayAbilityTargetDataHandle& HitResult, int32 Index);
 
+	/** Returns the hit result for a given index if it exists */
 	UFUNCTION(BlueprintPure, Category = "Ability|TargetData")
 	static FHitResult GetHitResultFromTargetData(const FGameplayAbilityTargetDataHandle& HitResult, int32 Index);
 
+	/** Returns true if the target data has an origin */
 	UFUNCTION(BlueprintPure, Category = "Ability|TargetData")
 	static bool TargetDataHasOrigin(const FGameplayAbilityTargetDataHandle& TargetData, int32 Index);
 
+	/** Returns the origin for a given index if it exists */
 	UFUNCTION(BlueprintPure, Category = "Ability|TargetData")
 	static FTransform GetTargetDataOrigin(const FGameplayAbilityTargetDataHandle& TargetData, int32 Index);
 
+	/** Returns true if the target data has an end point */
 	UFUNCTION(BlueprintPure, Category = "Ability|TargetData")
 	static bool TargetDataHasEndPoint(const FGameplayAbilityTargetDataHandle& TargetData, int32 Index);
 
+	/** Returns the end point for a given index if it exists */
 	UFUNCTION(BlueprintPure, Category = "Ability|TargetData")
 	static FVector GetTargetDataEndPoint(const FGameplayAbilityTargetDataHandle& TargetData, int32 Index);
 
+	/** Returns the end point transform for a given index if it exists */
 	UFUNCTION(BlueprintPure, Category = "Ability|TargetData")
 	static FTransform GetTargetDataEndPointTransform(const FGameplayAbilityTargetDataHandle& TargetData, int32 Index);
 
@@ -147,18 +166,23 @@ class GAMEPLAYABILITIES_API UAbilitySystemBlueprintLibrary : public UBlueprintFu
 	//		GameplayEffectContext
 	// -------------------------------------------------------------------------------
 
+	/** Returns true if this context has ever been initialized */
 	UFUNCTION(BlueprintPure, Category = "Ability|EffectContext", Meta = (DisplayName = "IsValid"))
 	static bool EffectContextIsValid(FGameplayEffectContextHandle EffectContext);
 
+	/** Returns true if the ability system component that instigated this is locally controlled */
 	UFUNCTION(BlueprintPure, Category = "Ability|EffectContext", Meta = (DisplayName = "IsInstigatorLocallyControlled"))
 	static bool EffectContextIsInstigatorLocallyControlled(FGameplayEffectContextHandle EffectContext);
 
+	/** Extracts a hit result from the effect context if it is set */
 	UFUNCTION(BlueprintPure, Category = "Ability|EffectContext", Meta = (DisplayName = "GetHitResult"))
 	static FHitResult EffectContextGetHitResult(FGameplayEffectContextHandle EffectContext);
 
+	/** Returns true if there is a valid hit result inside the effect context */
 	UFUNCTION(BlueprintPure, Category = "Ability|EffectContext", Meta = (DisplayName = "HasHitResult"))
 	static bool EffectContextHasHitResult(FGameplayEffectContextHandle EffectContext);
 
+	/** Adds a hit result to the effect context */
 	UFUNCTION(BlueprintCallable, Category = "Ability|EffectContext", Meta = (DisplayName = "AddHitResult"))
 	static void EffectContextAddHitResult(FGameplayEffectContextHandle EffectContext, FHitResult HitResult, bool bReset);
 
@@ -189,22 +213,28 @@ class GAMEPLAYABILITIES_API UAbilitySystemBlueprintLibrary : public UBlueprintFu
 	// -------------------------------------------------------------------------------
 	//		GameplayCue
 	// -------------------------------------------------------------------------------
-	
+
+	/** Returns true if the ability system component that spawned this cue is locally controlled */
 	UFUNCTION(BlueprintPure, Category="Ability|GameplayCue")
 	static bool IsInstigatorLocallyControlled(FGameplayCueParameters Parameters);
 
+	/** Returns true if the ability system component that spawned this cue is locally controlled and a player */
 	UFUNCTION(BlueprintPure, Category="Ability|GameplayCue")
 	static bool IsInstigatorLocallyControlledPlayer(FGameplayCueParameters Parameters);
 
+	/** Returns number of actors stored in the Effect Context used by this cue */
 	UFUNCTION(BlueprintPure, Category = "Ability|GameplayCue")
 	static int32 GetActorCount(FGameplayCueParameters Parameters);
 
+	/** Returns actor stored in the Effect Context used by this cue */
 	UFUNCTION(BlueprintPure, Category = "Ability|GameplayCue")
 	static AActor* GetActorByIndex(FGameplayCueParameters Parameters, int32 Index);
 
+	/** Returns a hit result stored in the effect context if valid */
 	UFUNCTION(BlueprintPure, Category = "Ability|GameplayCue")
 	static FHitResult GetHitResult(FGameplayCueParameters Parameters);
 
+	/** Checks if the effect context has a hit reslt stored inside */
 	UFUNCTION(BlueprintPure, Category = "Ability|GameplayCue")
 	static bool HasHitResult(FGameplayCueParameters Parameters);
 
@@ -239,29 +269,32 @@ class GAMEPLAYABILITIES_API UAbilitySystemBlueprintLibrary : public UBlueprintFu
 	// -------------------------------------------------------------------------------
 	//		GameplayEffectSpec
 	// -------------------------------------------------------------------------------
-	
+
+	/** Sets a raw name Set By Caller magnitude value, the tag version should normally be used */
 	UFUNCTION(BlueprintCallable, Category = "Ability|GameplayEffect")
 	static FGameplayEffectSpecHandle AssignSetByCallerMagnitude(FGameplayEffectSpecHandle SpecHandle, FName DataName, float Magnitude);
 
+	/** Sets a gameplay tag Set By Caller magnitude value */
 	UFUNCTION(BlueprintCallable, Category = "Ability|GameplayEffect")
 	static FGameplayEffectSpecHandle AssignTagSetByCallerMagnitude(FGameplayEffectSpecHandle SpecHandle, FGameplayTag DataTag, float Magnitude);
 
+	/** Manually sets the duration on a specific effect */
 	UFUNCTION(BlueprintCallable, Category = "Ability|GameplayEffect")
 	static FGameplayEffectSpecHandle SetDuration(FGameplayEffectSpecHandle SpecHandle, float Duration);
 
-	// This instance of the effect will now grant NewGameplayTag to the object that this effect is applied to.
+	/** This instance of the effect will now grant NewGameplayTag to the object that this effect is applied to */
 	UFUNCTION(BlueprintCallable, Category = "Ability|GameplayEffect")
 	static FGameplayEffectSpecHandle AddGrantedTag(FGameplayEffectSpecHandle SpecHandle, FGameplayTag NewGameplayTag);
 
-	// This instance of the effect will now grant NewGameplayTags to the object that this effect is applied to.
+	/** This instance of the effect will now grant NewGameplayTags to the object that this effect is applied to */
 	UFUNCTION(BlueprintCallable, Category = "Ability|GameplayEffect")
 	static FGameplayEffectSpecHandle AddGrantedTags(FGameplayEffectSpecHandle SpecHandle, FGameplayTagContainer NewGameplayTags);
 
-	// Adds NewGameplayTag to this instance of the effect.
+	/** Adds NewGameplayTag to this instance of the effect */
 	UFUNCTION(BlueprintCallable, Category = "Ability|GameplayEffect")
 	static FGameplayEffectSpecHandle AddAssetTag(FGameplayEffectSpecHandle SpecHandle, FGameplayTag NewGameplayTag);
 
-	// Adds NewGameplayTags to this instance of the effect.
+	/** Adds NewGameplayTags to this instance of the effect */
 	UFUNCTION(BlueprintCallable, Category = "Ability|GameplayEffect")
 	static FGameplayEffectSpecHandle AddAssetTags(FGameplayEffectSpecHandle SpecHandle, FGameplayTagContainer NewGameplayTags);
 
@@ -320,12 +353,15 @@ class GAMEPLAYABILITIES_API UAbilitySystemBlueprintLibrary : public UBlueprintFu
 	UFUNCTION(BlueprintCallable, Category = "Ability|GameplayEffect")
 	static float GetActiveGameplayEffectExpectedEndTime(FActiveGameplayEffectHandle ActiveHandle);
 
+	/** Returns the total duration for a given GameplayEffect */
 	UFUNCTION(BlueprintCallable, Category = "Ability|GameplayEffect")
 	static float GetActiveGameplayEffectTotalDuration(FActiveGameplayEffectHandle ActiveHandle);
 
+	/** Returns the total duration for a given GameplayEffect, basically ExpectedEndTime - Current Time */
 	UFUNCTION(BlueprintCallable, Category = "Ability|GameplayEffect", meta = (WorldContext = "WorldContextObject"))
 	static float GetActiveGameplayEffectRemainingDuration(UObject* WorldContextObject, FActiveGameplayEffectHandle ActiveHandle);
 
+	/** Returns a debug string for display */
 	UFUNCTION(BlueprintPure, Category = "Ability|GameplayEffect", Meta = (DisplayName = "Get Active GameplayEffect Debug String "))
 	static FString GetActiveGameplayEffectDebugString(FActiveGameplayEffectHandle ActiveHandle);
 
diff --git a/Engine/Plugins/Runtime/GameplayAbilities/Source/GameplayAbilities/Public/AbilitySystemComponent.h b/Engine/Plugins/Runtime/GameplayAbilities/Source/GameplayAbilities/Public/AbilitySystemComponent.h
--- a/Engine/Plugins/Runtime/GameplayAbilities/Source/GameplayAbilities/Public/AbilitySystemComponent.h
+++ b/Engine/Plugins/Runtime/GameplayAbilities/Source/GameplayAbilities/Public/AbilitySystemComponent.h
@@ -75,8 +75,9 @@
 /** Notifies when GameplayEffectSpec is blocked by an ActiveGameplayEffect due to immunity  */
 DECLARE_MULTICAST_DELEGATE_TwoParams(FImmunityBlockGE, const FGameplayEffectSpec& /*BlockedSpec*/, const FActiveGameplayEffect* /*ImmunityGameplayEffect*/);
 
+/** How gameplay effects will be replicated to clients */
 UENUM()
-enum class EReplicationMode : uint8
+enum class EGameplayEffectReplicationMode : uint8
 {
 	/** Only replicate minimal gameplay effect info*/
 	Minimal,
@@ -86,9 +87,7 @@
 	Full,
 };
 
-/**
- *	The core ActorComponent for interfacing with the GameplayAbilities System
- */
+/** The core ActorComponent for interfacing with the GameplayAbilities System */
 UCLASS(ClassGroup=AbilitySystem, hidecategories=(Object,LOD,Lighting,Transform,Sockets,TextureStreaming), editinlinenew, meta=(BlueprintSpawnableComponent))
 class GAMEPLAYABILITIES_API UAbilitySystemComponent : public UGameplayTasksComponent, public IGameplayTagAssetInterface, public IAbilitySystemReplicationProxyInterface
 {
@@ -100,18 +99,12 @@
 	/** Used to register callbacks to confirm/cancel input */
 	DECLARE_DYNAMIC_MULTICAST_DELEGATE(FAbilityConfirmOrCancel);
 
-	friend struct FActiveGameplayEffectAction_Add;
-	friend FGameplayEffectSpec;
-	friend class AAbilitySystemDebugHUD;
+	/** Delegate for when an effect is applied */
+	DECLARE_MULTICAST_DELEGATE_ThreeParams(FOnGameplayEffectAppliedDelegate, UAbilitySystemComponent*, const FGameplayEffectSpec&, FActiveGameplayEffectHandle);
 
-	virtual void InitializeComponent() override;
-	virtual void UninitializeComponent() override;
-	virtual void OnComponentDestroyed(bool bDestroyingHierarchy) override;
-	virtual void TickComponent(float DeltaTime, enum ELevelTick TickType, FActorComponentTickFunction *ThisTickFunction) override;
+	// ----------------------------------------------------------------------------------------------------------------
+	//	Attributes
+	// ----------------------------------------------------------------------------------------------------------------
 
-	/** retrieves information whether this component should be ticking taken current
-	 *	activity into consideration*/
-	virtual bool GetShouldTick() const override;
-
 	/** Finds existing AttributeSet */
 	template <class T >
@@ -143,14 +136,15 @@
 	}
 
 	/**
-	* Does this ability system component have this attribute?
-	*
-	* @param Attribute	Handle of the gameplay effect to retrieve target tags from
-	*
-	* @return true if Attribute is valid and this ability system component contains an attribute set that contains Attribute. Returns false otherwise.
-	*/
+	 * Does this ability system component have this attribute?
+	 *
+	 * @param Attribute	Handle of the gameplay effect to retrieve target tags from
+	 *
+	 * @return true if Attribute is valid and this ability system component contains an attribute set that contains Attribute. Returns false otherwise.
+	 */
 	bool HasAttributeSetForAttribute(FGameplayAttribute Attribute) const;
 
+	/** Initializes starting attributes from a data table. Not well supported, a gameplay effect with curve table references may be a better solution */
 	const UAttributeSet* InitStats(TSubclassOf<class UAttributeSet> Attributes, const UDataTable* DataTable);
 
 	UFUNCTION(BlueprintCallable, Category="Skills", meta=(DisplayName="InitStats", ScriptName="InitStats"))
@@ -162,6 +156,7 @@
 	UPROPERTY(EditAnywhere, Category="AttributeTest")
 	TArray<FAttributeDefaults>	DefaultStartingData;
 
+	/** List of attribute sets */
 	UPROPERTY(Replicated)
 	TArray<UAttributeSet*>	SpawnedAttributes;
 
@@ -192,35 +187,29 @@
 	float GetNumericAttribute(const FGameplayAttribute &Attribute) const;
 	float GetNumericAttributeChecked(const FGameplayAttribute &Attribute) const;
 
-	// -- Replication -------------------------------------------------------------------------------------------------
+	/** Returns an attribute value, after applying tag filters */
+	float GetFilteredAttributeValue(const FGameplayAttribute& Attribute, const FGameplayTagRequirements& SourceTags, const FGameplayTagContainer& TargetTags, const TArray<FActiveGameplayEffectHandle>& HandlesToIgnore = TArray<FActiveGameplayEffectHandle>());
 
-	virtual bool ReplicateSubobjects(class UActorChannel *Channel, class FOutBunch *Bunch, FReplicationFlags *RepFlags) override;
-	
-	/** Force owning actor to update it's replication, to make sure that gameplay cues get sent down quickly. Override to change how aggressive this is */
-	virtual void ForceReplication() override;
+	// ----------------------------------------------------------------------------------------------------------------
+	//	Replication
+	// ----------------------------------------------------------------------------------------------------------------
 
 	/** Forces avatar actor to update it's replication. Useful for things like needing to replication for movement / locations reasons. */
 	virtual void ForceAvatarReplication();
 
-	virtual void GetSubobjectsWithStableNamesForNetworking(TArray<UObject*>& Objs) override;
+	/** When true, we will not replicate active gameplay effects for this ability system component, so attributes and tags */
+	virtual void SetReplicationMode(EGameplayEffectReplicationMode NewReplicationMode);
 
-	virtual void PreNetReceive() override;
-	
-	virtual void PostNetReceive() override;
+	/** How gameplay effects are replicated */
+	EGameplayEffectReplicationMode ReplicationMode;
 
-	/** When true, we will not replicate active gameplay effects for this abiltiy system component, so attributes and tags */
-	void SetReplicationMode(EReplicationMode NewReplicationMode);
-
-	EReplicationMode ReplicationMode;
-
-	/** When enabled, we will not replicate this ASC to simulated proxies. We will route multicast RPCs through   */
-	bool ReplicationProxyEnabled;
-
 	/** Who to route replication through if ReplicationProxyEnabled (if this returns null, when ReplicationProxyEnabled, we wont replicate)  */
 	virtual IAbilitySystemReplicationProxyInterface* GetReplicationInterface();
 
+	/** Current prediction key, set with FScopedPredictionWindow */
 	FPredictionKey	ScopedPredictionKey;
 
+	/** Returns the prediction key that should be used for any actions */
 	FPredictionKey GetPredictionKeyForNewAction() const
 	{
 		return ScopedPredictionKey.IsValidForMorePrediction() ? ScopedPredictionKey : FPredictionKey();
@@ -232,9 +221,19 @@
 		return ScopedPredictionKey.IsValidForMorePrediction();
 	}
 
+	/** Returns true if this is running on the server or has a valid prediciton key */
 	bool HasAuthorityOrPredictionKey(const FGameplayAbilityActivationInfo* ActivationInfo) const;
 
-	// A pending activation that cannot be activated yet, will be rechecked at a later point
+	/** Returns true if this component's actor has authority */
+	virtual bool IsOwnerActorAuthoritative() const;
+
+	/** Replicate that an ability has ended/canceled, to the client or server as appropriate */
+	virtual void ReplicateEndOrCancelAbility(FGameplayAbilitySpecHandle Handle, FGameplayAbilityActivationInfo ActivationInfo, UGameplayAbility* Ability, bool bWasCanceled);
+
+	/** Force cancels the ability and does not replicate this to the other side. This should be called when the ability is cancelled by the other side */
+	virtual void ForceCancelAbilityDueToReplication(UGameplayAbility* Instance);
+
+	/** A pending activation that cannot be activated yet, will be rechecked at a later point */
 	struct FPendingAbilityInfo
 	{
 		bool operator==(const FPendingAbilityInfo& Other) const
@@ -256,13 +255,16 @@
 		{}
 	};
 
-	// This is a list of GameplayAbilities that are predicted by the client and were triggered by abilities that were also predicted by the client
-	// When the server version of the predicted ability executes it should trigger copies of these and the copies will be associated with the correct prediction keys
+	/**
+	 * This is a list of GameplayAbilities that are predicted by the client and were triggered by abilities that were also predicted by the client
+	 * When the server version of the predicted ability executes it should trigger copies of these and the copies will be associated with the correct prediction keys
+	 */
 	TArray<FPendingAbilityInfo> PendingClientActivatedAbilities;
 
-	// This is a list of GameplayAbilities that were activated on the server and can't yet execute on the client. It will try to execute these at a later point
+	/** This is a list of GameplayAbilities that were activated on the server and can't yet execute on the client. It will try to execute these at a later point */
 	TArray<FPendingAbilityInfo> PendingServerActivatedAbilities;
 
+	/** State of execution for an ability, used to track on the server */
 	enum class EAbilityExecutionState : uint8
 	{
 		Executing,
@@ -284,32 +286,31 @@
 		FGameplayAbilitySpecHandle Handle;
 	};
 
+	/** List of all executing abilities the server knows about */
 	TArray<FExecutingAbilityInfo> ExecutingServerAbilities;
 
 	// ----------------------------------------------------------------------------------------------------------------
-	//
-	//	GameplayEffects	
-	//	
+	//	GameplayEffects: Primary outward facing API for other systems
 	// ----------------------------------------------------------------------------------------------------------------
 
-	// --------------------------------------------
-	// Primary outward facing API for other systems:
-	// --------------------------------------------
-	FActiveGameplayEffectHandle ApplyGameplayEffectSpecToTarget(const FGameplayEffectSpec& GameplayEffect, UAbilitySystemComponent *Target, FPredictionKey PredictionKey=FPredictionKey());
-	FActiveGameplayEffectHandle ApplyGameplayEffectSpecToSelf(const FGameplayEffectSpec& GameplayEffect, FPredictionKey PredictionKey = FPredictionKey());
-
-	UFUNCTION(BlueprintCallable, Category = GameplayEffects, meta=(DisplayName = "ApplyGameplayEffectSpecToTarget", ScriptName = "ApplyGameplayEffectSpecToTarget"))
+	/** Applies a previously created gameplay effect spec to a target */
+	UFUNCTION(BlueprintCallable, Category = GameplayEffects, meta = (DisplayName = "ApplyGameplayEffectSpecToTarget", ScriptName = "ApplyGameplayEffectSpecToTarget"))
 	FActiveGameplayEffectHandle BP_ApplyGameplayEffectSpecToTarget(const FGameplayEffectSpecHandle& SpecHandle, UAbilitySystemComponent* Target);
 
-	UFUNCTION(BlueprintCallable, Category = GameplayEffects, meta=(DisplayName = "ApplyGameplayEffectSpecToSelf", ScriptName = "ApplyGameplayEffectSpecToSelf"))
+	virtual FActiveGameplayEffectHandle ApplyGameplayEffectSpecToTarget(const FGameplayEffectSpec& GameplayEffect, UAbilitySystemComponent *Target, FPredictionKey PredictionKey=FPredictionKey());
+
+	/** Applies a previously created gameplay effect spec to this component */
+	UFUNCTION(BlueprintCallable, Category = GameplayEffects, meta = (DisplayName = "ApplyGameplayEffectSpecToSelf", ScriptName = "ApplyGameplayEffectSpecToSelf"))
 	FActiveGameplayEffectHandle BP_ApplyGameplayEffectSpecToSelf(const FGameplayEffectSpecHandle& SpecHandle);
-	
+
+	virtual FActiveGameplayEffectHandle ApplyGameplayEffectSpecToSelf(const FGameplayEffectSpec& GameplayEffect, FPredictionKey PredictionKey = FPredictionKey());
+
 	/** Gets the FActiveGameplayEffect based on the passed in Handle */
 	const UGameplayEffect* GetGameplayEffectDefForHandle(FActiveGameplayEffectHandle Handle);
 
 	/** Removes GameplayEffect by Handle. StacksToRemove=-1 will remove all stacks. */
 	UFUNCTION(BlueprintCallable, BlueprintAuthorityOnly, Category = GameplayEffects)
-	bool RemoveActiveGameplayEffect(FActiveGameplayEffectHandle Handle, int32 StacksToRemove=-1);
+	virtual bool RemoveActiveGameplayEffect(FActiveGameplayEffectHandle Handle, int32 StacksToRemove=-1);
 
 	/** 
 	 * Remove active gameplay effects whose backing definition are the specified gameplay effect class
@@ -319,15 +320,15 @@
 	 * @param StacksToRemove					Number of stacks to remove, -1 means remove all
 	 */
 	UFUNCTION(BlueprintCallable, BlueprintAuthorityOnly, Category = GameplayEffects)
-	void RemoveActiveGameplayEffectBySourceEffect(TSubclassOf<UGameplayEffect> GameplayEffect, UAbilitySystemComponent* InstigatorAbilitySystemComponent, int32 StacksToRemove = -1);
+	virtual void RemoveActiveGameplayEffectBySourceEffect(TSubclassOf<UGameplayEffect> GameplayEffect, UAbilitySystemComponent* InstigatorAbilitySystemComponent, int32 StacksToRemove = -1);
 
 	/** Get an outgoing GameplayEffectSpec that is ready to be applied to other things. */
 	UFUNCTION(BlueprintCallable, Category = GameplayEffects)
-	FGameplayEffectSpecHandle MakeOutgoingSpec(TSubclassOf<UGameplayEffect> GameplayEffectClass, float Level, FGameplayEffectContextHandle Context) const;
+	virtual FGameplayEffectSpecHandle MakeOutgoingSpec(TSubclassOf<UGameplayEffect> GameplayEffectClass, float Level, FGameplayEffectContextHandle Context) const;
 
 	/** Create an EffectContext for the owner of this AbilitySystemComponent */
 	UFUNCTION(BlueprintCallable, Category = GameplayEffects)
-	FGameplayEffectContextHandle MakeEffectContext() const;
+	virtual FGameplayEffectContextHandle MakeEffectContext() const;
 
 	/**
 	 * Get the count of the specified source effect on the ability system component. For non-stacking effects, this is the sum of all active instances.
@@ -350,32 +351,30 @@
 		RemoveActiveGameplayEffect(Handle, StacksToRemove);
 	}
 
-	/** Needed for delegate callback for tag prediction */
-	void RemoveOneTagCount_NoReturn(FGameplayTag Tag)
-	{
-		UpdateTagMap(Tag, -1);
-	}
-
 	/** Called for predictively added gameplay cue. Needs to remove tag count and possible invoke OnRemove event if misprediction */
-	void OnPredictiveGameplayCueCatchup(FGameplayTag Tag);
+	virtual void OnPredictiveGameplayCueCatchup(FGameplayTag Tag);
 
+	/** Returns the total duration of a gameplay effect */
 	float GetGameplayEffectDuration(FActiveGameplayEffectHandle Handle) const;
 
 	/** Called whenever the server time replicates via the game state to keep our cooldown timers in sync with the server */
-	void RecomputeGameplayEffectStartTimes(const float WorldTime, const float ServerWorldTime);
+	virtual void RecomputeGameplayEffectStartTimes(const float WorldTime, const float ServerWorldTime);
 
+	/** Return start time and total duration of a gameplay effect */
 	void GetGameplayEffectStartTimeAndDuration(FActiveGameplayEffectHandle Handle, float& StartEffectTime, float& Duration) const;
 
 	/** Updates the level of an already applied gameplay effect. The intention is that this is 'seemless' and doesnt behave like removing/reapplying */
 	UFUNCTION(BlueprintCallable, BlueprintAuthorityOnly, Category = GameplayEffects)
-	void SetActiveGameplayEffectLevel(FActiveGameplayEffectHandle ActiveHandle, int32 NewLevel);
+	virtual void SetActiveGameplayEffectLevel(FActiveGameplayEffectHandle ActiveHandle, int32 NewLevel);
 
 	/** Updates the level of an already applied gameplay effect. The intention is that this is 'seemless' and doesnt behave like removing/reapplying */
 	UFUNCTION(BlueprintCallable, BlueprintAuthorityOnly, Category = GameplayEffects)
-	void SetActiveGameplayEffectLevelUsingQuery(FGameplayEffectQuery Query, int32 NewLevel);
+	virtual void SetActiveGameplayEffectLevelUsingQuery(FGameplayEffectQuery Query, int32 NewLevel);
 
-	// Not happy with this interface but don't see a better way yet. How should outside code (UI, etc) ask things like 'how much is this gameplay effect modifying my damage by'
-	// (most likely we want to catch this on the backend - when damage is applied we can get a full dump/history of how the number got to where it is. But still we may need polling methods like below (how much would my damage be)
+	/**
+	 * Raw accessor to ask the magnitude of a gameplay effect, not necessarily always correct. How should outside code (UI, etc) ask things like 'how much is this gameplay effect modifying my damage by'
+	 * (most likely we want to catch this on the backend - when damage is applied we can get a full dump/history of how the number got to where it is. But still we may need polling methods like below (how much would my damage be)
+	 */
 	UFUNCTION(BlueprintCallable, Category = GameplayEffects)
 	float GetGameplayEffectMagnitude(FActiveGameplayEffectHandle Handle, FGameplayAttribute Attribute) const;
 
@@ -391,6 +390,9 @@
 	/** Gets the GE Handle of the GE that granted the passed in Ability */
 	FActiveGameplayEffectHandle FindActiveGameplayEffectHandle(FGameplayAbilitySpecHandle Handle) const;
 
+	/** Returns const pointer to the actual active gamepay effect structure */
+	const FActiveGameplayEffect* GetActiveGameplayEffect(const FActiveGameplayEffectHandle Handle) const;
+
 	/**
 	 * Get the source tags from the gameplay spec represented by the specified handle, if possible
 	 * 
@@ -398,10 +400,7 @@
 	 * 
 	 * @return Source tags from the gameplay spec represented by the handle, if possible
 	 */
-	const FGameplayTagContainer* GetGameplayEffectSourceTagsFromHandle(FActiveGameplayEffectHandle Handle) const
-	{
-		return ActiveGameplayEffects.GetGameplayEffectSourceTagsFromHandle(Handle);
-	}
+	const FGameplayTagContainer* GetGameplayEffectSourceTagsFromHandle(FActiveGameplayEffectHandle Handle) const;
 
 	/**
 	 * Get the target tags from the gameplay spec represented by the specified handle, if possible
@@ -410,50 +409,35 @@
 	 * 
 	 * @return Target tags from the gameplay spec represented by the handle, if possible
 	 */
-	const FGameplayTagContainer* GetGameplayEffectTargetTagsFromHandle(FActiveGameplayEffectHandle Handle) const
-	{
-		return ActiveGameplayEffects.GetGameplayEffectTargetTagsFromHandle(Handle);
-	}
+	const FGameplayTagContainer* GetGameplayEffectTargetTagsFromHandle(FActiveGameplayEffectHandle Handle) const;
 
 	/**
 	 * Populate the specified capture spec with the data necessary to capture an attribute from the component
 	 * 
 	 * @param OutCaptureSpec	[OUT] Capture spec to populate with captured data
 	 */
-	void CaptureAttributeForGameplayEffect(OUT FGameplayEffectAttributeCaptureSpec& OutCaptureSpec)
-	{
-		// Verify the capture is happening on an attribute the component actually has a set for; if not, can't capture the value
-		const FGameplayAttribute& AttributeToCapture = OutCaptureSpec.BackingDefinition.AttributeToCapture;
-		if (AttributeToCapture.IsValid() && (AttributeToCapture.IsSystemAttribute() || GetAttributeSubobject(AttributeToCapture.GetAttributeSetClass())))
-		{
-			ActiveGameplayEffects.CaptureAttributeForGameplayEffect(OutCaptureSpec);
-		}
-	}
+	void CaptureAttributeForGameplayEffect(OUT FGameplayEffectAttributeCaptureSpec& OutCaptureSpec);
 	
-	// --------------------------------------------
-	// Callbacks / Notifies
-	// (these need to be at the UObject level so we can safetly bind, rather than binding to raw at the ActiveGameplayEffect/Container level which is unsafe if the AbilitySystemComponent were killed).
-	// --------------------------------------------
+	// ----------------------------------------------------------------------------------------------------------------
+	//  Callbacks / Notifies
+	//  (these need to be at the UObject level so we can safely bind, rather than binding to raw at the ActiveGameplayEffect/Container level which is unsafe if the AbilitySystemComponent were killed).
+	// ----------------------------------------------------------------------------------------------------------------
 
+	/** Called when a specific attribute aggregator value changes, gameplay effects refresh their values when this happens */
 	void OnAttributeAggregatorDirty(FAggregator* Aggregator, FGameplayAttribute Attribute, bool FromRecursiveCall=false);
 
+	/** Called when attribute magnitudes change, to forward information to dependent gameplay effects */
 	void OnMagnitudeDependencyChange(FActiveGameplayEffectHandle Handle, const FAggregator* ChangedAggregator);
 
-	void ForceAttributeAggregatorPropertyUpdate(FAggregator* Aggregator);
-
 	/** This ASC has successfully applied a GE to something (potentially itself) */
 	void OnGameplayEffectAppliedToTarget(UAbilitySystemComponent* Target, const FGameplayEffectSpec& SpecApplied, FActiveGameplayEffectHandle ActiveHandle);
-
 	void OnGameplayEffectAppliedToSelf(UAbilitySystemComponent* Source, const FGameplayEffectSpec& SpecApplied, FActiveGameplayEffectHandle ActiveHandle);
-
 	void OnPeriodicGameplayEffectExecuteOnTarget(UAbilitySystemComponent* Target, const FGameplayEffectSpec& SpecExecuted, FActiveGameplayEffectHandle ActiveHandle);
-
 	void OnPeriodicGameplayEffectExecuteOnSelf(UAbilitySystemComponent* Source, const FGameplayEffectSpec& SpecExecuted, FActiveGameplayEffectHandle ActiveHandle);
 
+	/** Called when the duration of a gamepaly effect has changed */
 	virtual void OnGameplayEffectDurationChange(struct FActiveGameplayEffect& ActiveEffect);
 
-	DECLARE_MULTICAST_DELEGATE_ThreeParams(FOnGameplayEffectAppliedDelegate, UAbilitySystemComponent*, const FGameplayEffectSpec&, FActiveGameplayEffectHandle);
-
 	/** Called on server whenever a GE is applied to self. This includes instant and duration based GEs. */
 	FOnGameplayEffectAppliedDelegate OnGameplayEffectAppliedDelegateToSelf;
 
@@ -469,9 +453,54 @@
 	/** Called on server whenever a periodic GE executes on target */
 	FOnGameplayEffectAppliedDelegate OnPeriodicGameplayEffectExecuteDelegateOnTarget;
 
-	// --------------------------------------------
-	// Tags
-	// --------------------------------------------
+	/** Immunity notification support */
+	FImmunityBlockGE OnImmunityBlockGameplayEffectDelegate;
+
+	/** Register for when an attribute value changes, should be replaced by GetGameplayAttributeValueChangeDelegate */
+	FOnGameplayAttributeChange& RegisterGameplayAttributeEvent(FGameplayAttribute Attribute);
+
+	/** Register for when an attribute value changes */
+	FOnGameplayAttributeValueChange& GetGameplayAttributeValueChangeDelegate(FGameplayAttribute Attribute);
+
+	/** A generic callback anytime an ability is activated (started) */
+	FGenericAbilityDelegate AbilityActivatedCallbacks;
+
+	/** Callback anytime an ability is ended */
+	FAbilityEnded AbilityEndedCallbacks;
+
+	/** Callback anytime an ability is ended, with extra information */
+	FGameplayAbilityEndedDelegate OnAbilityEnded;
+
+	/** A generic callback anytime an ability is committed (cost/cooldown applied) */
+	FGenericAbilityDelegate AbilityCommittedCallbacks;
+
+	/** Called with a failure reason when an ability fails to execute */
+	FAbilityFailedDelegate AbilityFailedCallbacks;
+
+	/** Called when an ability spec's internals have changed */
+	FAbilitySpecDirtied AbilitySpecDirtiedCallbacks;
+
+	/** Call notify callbacks above */
+	virtual void NotifyAbilityCommit(UGameplayAbility* Ability);
+	virtual void NotifyAbilityActivated(const FGameplayAbilitySpecHandle Handle, UGameplayAbility* Ability);
+	virtual void NotifyAbilityFailed(const FGameplayAbilitySpecHandle Handle, UGameplayAbility* Ability, const FGameplayTagContainer& FailureReason);
+
+	/** Called when any gameplay effects are removed */
+	FOnGivenActiveGameplayEffectRemoved& OnAnyGameplayEffectRemovedDelegate();
+
+	DEPRECATED(4.17, "Use OnGameplayEffectRemoved_InfoDelegate (the delegate signature has changed)")
+	FOnActiveGameplayEffectRemoved* OnGameplayEffectRemovedDelegate(FActiveGameplayEffectHandle Handle);
+
+	/** Returns delegate structure that allows binding to several gameplay effect changes */
+	FActiveGameplayEffectEvents* GetActiveEffectEventSet(FActiveGameplayEffectHandle Handle);
+	FOnActiveGameplayEffectRemoved_Info* OnGameplayEffectRemoved_InfoDelegate(FActiveGameplayEffectHandle Handle);
+	FOnActiveGameplayEffectStackChange* OnGameplayEffectStackChangeDelegate(FActiveGameplayEffectHandle Handle);
+	FOnActiveGameplayEffectTimeChange* OnGameplayEffectTimeChangeDelegate(FActiveGameplayEffectHandle Handle);
+
+	// ----------------------------------------------------------------------------------------------------------------
+	//  Gameplay tag operations
+	//  Implements IGameplayTagAssetInterface using the TagCountContainer
+	// ----------------------------------------------------------------------------------------------------------------
 	FORCEINLINE bool HasMatchingGameplayTag(FGameplayTag TagToCheck) const override
 	{
 		return GameplayTagCountContainer.HasMatchingGameplayTag(TagToCheck);
@@ -492,11 +521,37 @@
 		TagContainer.AppendTags(GameplayTagCountContainer.GetExplicitGameplayTags());
 	}
 
+	/** Returns the number of instances of a given tag */
 	FORCEINLINE int32 GetTagCount(FGameplayTag TagToCheck) const
 	{
 		return GameplayTagCountContainer.GetTagCount(TagToCheck);
 	}
 
+	/** Forcibly sets the number of instances of a given tag */
+	FORCEINLINE void SetTagMapCount(const FGameplayTag& Tag, int32 NewCount)
+	{
+		GameplayTagCountContainer.SetTagCount(Tag, NewCount);
+	}
+
+	/** Update the number of instancse of a given tag and calls callback */
+	FORCEINLINE void UpdateTagMap(const FGameplayTag& BaseTag, int32 CountDelta)
+	{
+		if (GameplayTagCountContainer.UpdateTagCount(BaseTag, CountDelta))
+		{
+			OnTagUpdated(BaseTag, CountDelta > 0);
+		}
+	}
+
+	/** Update the number of instancse of a given tag and calls callback */
+	FORCEINLINE void UpdateTagMap(const FGameplayTagContainer& Container, int32 CountDelta)
+	{
+		for (auto TagIt = Container.CreateConstIterator(); TagIt; ++TagIt)
+		{
+			const FGameplayTag& Tag = *TagIt;
+			UpdateTagMap(Tag, CountDelta);
+		}
+	}
+
 	/** 	 
 	 *  Allows GameCode to add loose gameplaytags which are not backed by a GameplayEffect. 
 	 *
@@ -504,7 +559,6 @@
 	 *	
 	 *	It is up to the calling GameCode to make sure these tags are added on clients/server where necessary
 	 */
-
 	FORCEINLINE void AddLooseGameplayTag(const FGameplayTag& GameplayTag, int32 Count=1)
 	{
 		UpdateTagMap(GameplayTag, Count);
@@ -534,7 +588,6 @@
 	 * Minimally replicated tags are replicated tags that come from GEs when in bMinimalReplication mode. 
 	 * (The GEs do not replicate, but the tags they grant do replicate via these functions)
 	 */
-
 	FORCEINLINE void AddMinimalReplicationGameplayTag(const FGameplayTag& GameplayTag)
 	{
 		MinimalReplicationTags.AddTag(GameplayTag);
@@ -558,27 +611,37 @@
 		MinimalReplicationTags.RemoveTags(GameplayTags);
 		bIsNetDirty = true;
 	}
-	
+
 	/** Allow events to be registered for specific gameplay tags being added or removed */
 	FOnGameplayEffectTagCountChanged& RegisterGameplayTagEvent(FGameplayTag Tag, EGameplayTagEventType::Type EventType=EGameplayTagEventType::NewOrRemoved);
 
+	/** Register a tag event and immediately call it */
 	void RegisterAndCallGameplayTagEvent(FGameplayTag Tag, FOnGameplayEffectTagCountChanged::FDelegate Delegate, EGameplayTagEventType::Type EventType=EGameplayTagEventType::NewOrRemoved);
 
 	/** Returns multicast delegate that is invoked whenever a tag is added or removed (but not if just count is increased. Only for 'new' and 'removed' events) */
 	FOnGameplayEffectTagCountChanged& RegisterGenericGameplayTagEvent();
 
+	/** Executes a gameplay event. Returns the number of successful ability activations triggered by the event */
+	virtual int32 HandleGameplayEvent(FGameplayTag EventTag, const FGameplayEventData* Payload);
 
-	FOnGameplayAttributeChange& RegisterGameplayAttributeEvent(FGameplayAttribute Attribute);
+	/** Adds a new delegate to call when gameplay events happen. It will only be called if it matches any tags in passed filter container */
+	FDelegateHandle AddGameplayEventTagContainerDelegate(const FGameplayTagContainer& TagFilter, const FGameplayEventTagMulticastDelegate::FDelegate& Delegate);
 
-	FOnGameplayAttributeValueChange& GetGameplayAttributeValueChangeDelegate(FGameplayAttribute Attribute);
+	/** Remotes previously registered delegate */
+	void RemoveGameplayEventTagContainerDelegate(const FGameplayTagContainer& TagFilter, FDelegateHandle DelegateHandle);
 
-	// --------------------------------------------
-	// System Attributes
-	// --------------------------------------------
-	
+	/** Callbacks bound to Gameplay tags, these only activate if the exact tag is used. To handle tag hierarchies use AddGameplayEventContainerDelegate */
+	TMap<FGameplayTag, FGameplayEventMulticastDelegate> GenericGameplayEventCallbacks;
+
+	// ----------------------------------------------------------------------------------------------------------------
+	//  System Attributes
+	// ----------------------------------------------------------------------------------------------------------------
+
+	/** Internal Attribute that modifies the duration of gameplay effects created by this component */
 	UPROPERTY(meta=(SystemGameplayAttribute="true"))
 	float OutgoingDuration;
 
+	/** Internal Attribute that modifies the duration of gameplay effects applied to this component */
 	UPROPERTY(meta = (SystemGameplayAttribute = "true"))
 	float IncomingDuration;
 
@@ -588,46 +651,39 @@
 	static const FGameplayEffectAttributeCaptureDefinition& GetOutgoingDurationCapture();
 	static const FGameplayEffectAttributeCaptureDefinition& GetIncomingDurationCapture();
 
+	// ----------------------------------------------------------------------------------------------------------------
+	//  Additional Helper Functions
+	// ----------------------------------------------------------------------------------------------------------------
 
-	// --------------------------------------------
-	// Additional Helper Functions
-	// --------------------------------------------
-	
-	FOnGivenActiveGameplayEffectRemoved& OnAnyGameplayEffectRemovedDelegate();
-	DEPRECATED(4.17, "Use OnGameplayEffectRemoved_InfoDelegate (the delegate signature has changed)")
-	FOnActiveGameplayEffectRemoved* OnGameplayEffectRemovedDelegate(FActiveGameplayEffectHandle Handle);
-
-	FActiveGameplayEffectEvents* GetActiveEffectEventSet(FActiveGameplayEffectHandle Handle);
-	FOnActiveGameplayEffectRemoved_Info* OnGameplayEffectRemoved_InfoDelegate(FActiveGameplayEffectHandle Handle);
-	FOnActiveGameplayEffectStackChange* OnGameplayEffectStackChangeDelegate(FActiveGameplayEffectHandle Handle);
-	FOnActiveGameplayEffectTimeChange* OnGameplayEffectTimeChangeDelegate(FActiveGameplayEffectHandle Handle);
-
+	/** Apply a gameplay effect to passed in target */
 	UFUNCTION(BlueprintCallable, Category = GameplayEffects, meta=(DisplayName = "ApplyGameplayEffectToTarget", ScriptName = "ApplyGameplayEffectToTarget"))
 	FActiveGameplayEffectHandle BP_ApplyGameplayEffectToTarget(TSubclassOf<UGameplayEffect> GameplayEffectClass, UAbilitySystemComponent *Target, float Level, FGameplayEffectContextHandle Context);
 	FActiveGameplayEffectHandle ApplyGameplayEffectToTarget(UGameplayEffect *GameplayEffect, UAbilitySystemComponent *Target, float Level = UGameplayEffect::INVALID_LEVEL, FGameplayEffectContextHandle Context = FGameplayEffectContextHandle(), FPredictionKey PredictionKey = FPredictionKey());
 
+	/** Apply a gameplay effect to self */
 	UFUNCTION(BlueprintCallable, Category = GameplayEffects, meta=(DisplayName = "ApplyGameplayEffectToSelf", ScriptName = "ApplyGameplayEffectToSelf"))
 	FActiveGameplayEffectHandle BP_ApplyGameplayEffectToSelf(TSubclassOf<UGameplayEffect> GameplayEffectClass, float Level, FGameplayEffectContextHandle EffectContext);
-	
 	FActiveGameplayEffectHandle ApplyGameplayEffectToSelf(const UGameplayEffect *GameplayEffect, float Level, const FGameplayEffectContextHandle& EffectContext, FPredictionKey PredictionKey = FPredictionKey());
 
-	// Returns the number of gameplay effects that are currently active on this ability system component
+	/** Returns the number of gameplay effects that are currently active on this ability system component */
 	int32 GetNumActiveGameplayEffects() const
 	{
 		return ActiveGameplayEffects.GetNumGameplayEffects();
 	}
 
-	// Makes a copy of all the active effects on this ability component
+	/** Makes a copy of all the active effects on this ability component */
 	void GetAllActiveGameplayEffectSpecs(TArray<FGameplayEffectSpec>& OutSpecCopies) const
 	{
 		ActiveGameplayEffects.GetAllActiveGameplayEffectSpecs(OutSpecCopies);
 	}
 
+	/** Call from OnRep functions to set the attribute base value on the client */
 	void SetBaseAttributeValueFromReplication(float NewValue, FGameplayAttribute Attribute)
 	{
 		ActiveGameplayEffects.SetBaseAttributeValueFromReplication(Attribute, NewValue);
 	}
 
+	/** Call from OnRep functions to set the attribute base value on the client */
 	void SetBaseAttributeValueFromReplication(FGameplayAttributeData NewValue, FGameplayAttribute Attribute)
 	{
 		ActiveGameplayEffects.SetBaseAttributeValueFromReplication(Attribute, NewValue.GetBaseValue());
@@ -639,43 +695,28 @@
 		return ActiveGameplayEffects.CanApplyAttributeModifiers(GameplayEffect, Level, EffectContext);
 	}
 
-	// Generic 'Get expected magnitude (list) if I was to apply this outgoing or incoming'
-
-	// Get duration or magnitude (list) of active effects
-	//		-Get duration of CD
-	//		-Get magnitude + duration of a movespeed buff
-
+	/** Gets time remaining for all effects that match query */
 	TArray<float> GetActiveEffectsTimeRemaining(const FGameplayEffectQuery& Query) const;
 
+	/** Gets total duration for all effects that match query */
 	TArray<float> GetActiveEffectsDuration(const FGameplayEffectQuery& Query) const;
 
+	/** Gets both time remaining and total duration  for all effects that match query */
 	TArray<TPair<float,float>> GetActiveEffectsTimeRemainingAndDuration(const FGameplayEffectQuery& Query) const;
 
+	/** Returns list of active effects, for a query */
 	UFUNCTION(BlueprintCallable, BlueprintPure=false, Category = "GameplayEffects", meta=(DisplayName = "Get Activate Gameplay Effects for Query"))
 	TArray<FActiveGameplayEffectHandle> GetActiveEffects(const FGameplayEffectQuery& Query) const;
 
-	/** This will give the world time that all effects matching this query will be finished. If multiple effects match, it returns the one that returns last.*/
-	float GetActiveEffectsEndTime(const FGameplayEffectQuery& Query) const
-	{
-		TArray<AActor*> DummyInstigators;
-		return ActiveGameplayEffects.GetActiveEffectsEndTime(Query, DummyInstigators);
-	}
-	
-	float GetActiveEffectsEndTimeWithInstigators(const FGameplayEffectQuery& Query, TArray<AActor*>& Instigators) const
-	{
-		return ActiveGameplayEffects.GetActiveEffectsEndTime(Query, Instigators);
-	}
+	/** This will give the world time that all effects matching this query will be finished. If multiple effects match, it returns the one that returns last */
+	float GetActiveEffectsEndTime(const FGameplayEffectQuery& Query) const;
+	float GetActiveEffectsEndTimeWithInstigators(const FGameplayEffectQuery& Query, TArray<AActor*>& Instigators) const;
 
-	bool GetActiveEffectsEndTimeAndDuration(const FGameplayEffectQuery& Query, float& EndTime, float& Duration) const
-	{
-		TArray<AActor*> DummyInstigators;
-		return ActiveGameplayEffects.GetActiveEffectsEndTimeAndDuration(Query, EndTime, Duration, DummyInstigators);
-	}
+	/** Returns end time and total duration */
+	bool GetActiveEffectsEndTimeAndDuration(const FGameplayEffectQuery& Query, float& EndTime, float& Duration) const;
 
-	void ModifyActiveEffectStartTime(FActiveGameplayEffectHandle Handle, float StartTimeDiff)
-	{
-		ActiveGameplayEffects.ModifyActiveEffectStartTime(Handle, StartTimeDiff);
-	}
+	/** Modify the start time of a gameplay effect, to deal with timers being out of sync originally */
+	virtual void ModifyActiveEffectStartTime(FActiveGameplayEffectHandle Handle, float StartTimeDiff);
 
 	/** Removes all active effects that contain any of the tags in Tags */
 	UFUNCTION(BlueprintCallable, Category = GameplayEffects)
@@ -694,29 +735,10 @@
 	int32 RemoveActiveEffectsWithGrantedTags(FGameplayTagContainer Tags);
 
 	/** Removes all active effects that match given query. StacksToRemove=-1 will remove all stacks. */
-	int32 RemoveActiveEffects(const FGameplayEffectQuery& Query, int32 StacksToRemove = -1);
-
-	/** Invokes the WhileActive event for all GCs on active, non inhibited, GEs. This would typically be used on "respawn" or something where the mesh/avatar has changed */
-	void ReinvokeActiveGameplayCues();
-
-	// -----------------------------------------------------------
-
-	void OnRestackGameplayEffects();	
-	
-	void PrintAllGameplayEffects() const;
-
-	bool CachedIsNetSimulated;
-
-	/** Returns true of this component has authority */
-	bool IsOwnerActorAuthoritative() const
-	{
-		return !CachedIsNetSimulated;
-	}
+	virtual int32 RemoveActiveEffects(const FGameplayEffectQuery& Query, int32 StacksToRemove = -1);
 
 	// ----------------------------------------------------------------------------------------------------------------
-	//
 	//	GameplayCues
-	// 
 	// ----------------------------------------------------------------------------------------------------------------
 	
 	// Do not call these functions directly, call the wrappers on GameplayCueManager instead
@@ -750,59 +772,29 @@
 	UFUNCTION(NetMulticast, unreliable)
 	void NetMulticast_InvokeGameplayCuesAddedAndWhileActive_WithParams(const FGameplayTagContainer GameplayCueTags, FPredictionKey PredictionKey, FGameplayCueParameters GameplayCueParameters) override;
 
-	// GameplayCues can also come on their own. These take an optional effect context to pass through hit result, etc
+	/** GameplayCues can also come on their own. These take an optional effect context to pass through hit result, etc */
 	void ExecuteGameplayCue(const FGameplayTag GameplayCueTag, FGameplayEffectContextHandle EffectContext = FGameplayEffectContextHandle());
-
-	// This version allows the caller to set an explicit FGameplayCueParmeters.
 	void ExecuteGameplayCue(const FGameplayTag GameplayCueTag, const FGameplayCueParameters& GameplayCueParameters);
 
-	// -------------------------
+	/** Add a persistent gameplay cue */
+	void AddGameplayCue(const FGameplayTag GameplayCueTag, FGameplayEffectContextHandle EffectContext = FGameplayEffectContextHandle());
+	void AddGameplayCue(const FGameplayTag GameplayCueTag, const FGameplayCueParameters& GameplayCueParameters);
 
-	void AddGameplayCue(const FGameplayTag GameplayCueTag, FGameplayEffectContextHandle EffectContext = FGameplayEffectContextHandle())
-	{
-		AddGameplayCue_Internal(GameplayCueTag, EffectContext, ActiveGameplayCues);
-	}
-
-	void AddGameplayCue(const FGameplayTag GameplayCueTag, const FGameplayCueParameters& GameplayCueParameters)
-	{
-		AddGameplayCue_Internal(GameplayCueTag, GameplayCueParameters, ActiveGameplayCues);
-	}
-
-
 	/** Add gameplaycue for minimal replication mode. Should only be called in paths that would replicate gameplaycues in other ways (through GE for example) if not in minimal replication mode */
-	void AddGameplayCue_MinimalReplication(const FGameplayTag GameplayCueTag, FGameplayEffectContextHandle EffectContext = FGameplayEffectContextHandle())
-	{
-		AddGameplayCue_Internal(GameplayCueTag, EffectContext, MinimalReplicationGameplayCues);
-	}
+	void AddGameplayCue_MinimalReplication(const FGameplayTag GameplayCueTag, FGameplayEffectContextHandle EffectContext = FGameplayEffectContextHandle());
 
-	void AddGameplayCue_Internal(const FGameplayTag GameplayCueTag, FGameplayEffectContextHandle& EffectContext, FActiveGameplayCueContainer& GameplayCueContainer );
-
-	void AddGameplayCue_Internal(const FGameplayTag GameplayCueTag, const FGameplayCueParameters& GameplayCueParameters, FActiveGameplayCueContainer& GameplayCueContainer);
-
-	// -------------------------
+	/** Remove a persistent gameplay cue */
+	void RemoveGameplayCue(const FGameplayTag GameplayCueTag);
 	
-	void RemoveGameplayCue(const FGameplayTag GameplayCueTag)
-	{
-		RemoveGameplayCue_Internal(GameplayCueTag, ActiveGameplayCues);
-	}
-
 	/** Remove gameplaycue for minimal replication mode. Should only be called in paths that would replicate gameplaycues in other ways (through GE for example) if not in minimal replication mode */
-	void RemoveGameplayCue_MinimalReplication(const FGameplayTag GameplayCueTag)
-	{
-		RemoveGameplayCue_Internal(GameplayCueTag, MinimalReplicationGameplayCues);
-	}
-
-	void RemoveGameplayCue_Internal(const FGameplayTag GameplayCueTag, FActiveGameplayCueContainer& GameplayCueContainer);
-
-	// -------------------------
-
+	void RemoveGameplayCue_MinimalReplication(const FGameplayTag GameplayCueTag);
+	
 	/** Removes any GameplayCue added on its own, i.e. not as part of a GameplayEffect. */
 	void RemoveAllGameplayCues();
-	
+
+	/** Handles gameplay cue events from external sources */
 	void InvokeGameplayCueEvent(const FGameplayEffectSpecForRPC& Spec, EGameplayCueEvent::Type EventType);
-
 	void InvokeGameplayCueEvent(const FGameplayTag GameplayCueTag, EGameplayCueEvent::Type EventType, FGameplayEffectContextHandle EffectContext = FGameplayEffectContextHandle());
-
 	void InvokeGameplayCueEvent(const FGameplayTag GameplayCueTag, EGameplayCueEvent::Type EventType, const FGameplayCueParameters& GameplayCueParameters);
 
 	/** Allows polling to see if a GameplayCue is active. We expect most GameplayCue handling to be event based, but some cases we may need to check if a GamepalyCue is active (Animation Blueprint for example) */
@@ -818,7 +810,11 @@
 	/** Are we ready to invoke gameplaycues yet? */
 	virtual bool IsReadyForGameplayCues();
 
-	// ----------------------------------------------------------------------------------------------------------------
+	/** Handle GameplayCues that may have been deferred while doing the NetDeltaSerialize and waiting for the avatar actor to get loaded */
+	virtual void HandleDeferredGameplayCues(const FActiveGameplayEffectsContainer* GameplayEffectsContainer);
 
+	/** Invokes the WhileActive event for all GCs on active, non inhibited, GEs. This would typically be used on "respawn" or something where the mesh/avatar has changed */
+	virtual void ReinvokeActiveGameplayCues();
+
 	/**
 	 *	GameplayAbilities
@@ -887,9 +883,9 @@
 	/** Triggers an ability from a gameplay event, will only trigger on local/server depending on execution flags */
 	bool TriggerAbilityFromGameplayEvent(FGameplayAbilitySpecHandle AbilityToTrigger, FGameplayAbilityActorInfo* ActorInfo, FGameplayTag Tag, const FGameplayEventData* Payload, UAbilitySystemComponent& Component);
 
-	// --------------------------------------------
+	// ----------------------------------------------------------------------------------------------------------------
 	// Ability Cancelling/Interrupts
-	// --------------------------------------------
+	// ----------------------------------------------------------------------------------------------------------------
 
 	/** Cancels the specified ability CDO. */
 	void CancelAbility(UGameplayAbility* Ability);	
@@ -906,7 +902,6 @@
 	/** Cancels all abilities and kills any remaining instanced abilities */
 	virtual void DestroyActiveState();
 
-	// ----------------------------------------------------------------------------------------------------------------
 	/** 
 	 * Called from ability activation or native code, will apply the correct ability blocking tags and cancel existing abilities. Subclasses can override the behavior 
 	 * @param AbilityTags The tags of the ability that has block and cancel flags
@@ -924,16 +919,20 @@
 	/** Returns true if any passed in tags are blocked */
 	virtual bool AreAbilityTagsBlocked(const FGameplayTagContainer& Tags) const;
 
+	/** Block or cancel blocking for specific ability tags */
 	void BlockAbilitiesWithTags(const FGameplayTagContainer& Tags);
 	void UnBlockAbilitiesWithTags(const FGameplayTagContainer& Tags);
 
 	/** Checks if the ability system is currently blocking InputID. Returns true if InputID is blocked, false otherwise.  */
 	bool IsAbilityInputBlocked(int32 InputID) const;
 
+	/** Block or cancel blocking for specific input IDs */
 	void BlockAbilityByInputID(int32 InputID);
 	void UnBlockAbilityByInputID(int32 InputID);
 
+	// ----------------------------------------------------------------------------------------------------------------
 	// Functions meant to be called from GameplayAbility and subclasses, but not meant for general use
+	// ----------------------------------------------------------------------------------------------------------------
 
 	/** Returns the list of all activatable abilities */
 	const TArray<FGameplayAbilitySpec>& GetActivatableAbilities() const
@@ -976,15 +975,13 @@
 	/** Called from the ability to let the component know it is ended */
 	virtual void NotifyAbilityEnded(FGameplayAbilitySpecHandle Handle, UGameplayAbility* Ability, bool bWasCancelled);
 
-	/** Replicate that an ability has ended, to the client or server as appropriate */
-	void ReplicateEndAbility(FGameplayAbilitySpecHandle Handle, FGameplayAbilityActivationInfo ActivationInfo, UGameplayAbility* Ability);
-
+	/** Called from FScopedAbilityListLock */
 	void IncrementAbilityListLock();
 	void DecrementAbilityListLock();
 
-	// --------------------------------------------
+	// ----------------------------------------------------------------------------------------------------------------
 	// Debugging
-	// --------------------------------------------
+	// ----------------------------------------------------------------------------------------------------------------
 
 	struct FAbilitySystemComponentDebugInfo
 	{
@@ -1025,6 +1022,9 @@
 	void DebugLine(struct FAbilitySystemComponentDebugInfo& Info, FString Str, float XOffset, float YOffset);
 	FString CleanupName(FString Str);
 
+	/** Print a debug list of all gameplay effects */
+	void PrintAllGameplayEffects() const;
+
 	/** Ask the server to send ability system debug information back to the client, via ClientPrintDebug_Response  */
 	UFUNCTION(Server, reliable, WithValidation)
 	void ServerPrintDebug_Request();
@@ -1043,8 +1043,9 @@
 	void ClientPrintDebug_Response(const TArray<FString>& Strings, int32 GameFlags);
 	virtual void OnClientPrintDebug_Response(const TArray<FString>& Strings, int32 GameFlags);
 
-	/** Called when the ability is forced cancelled due to replication */
-	void ForceCancelAbilityDueToReplication(UGameplayAbility* Instance);
+#if ENABLE_VISUAL_LOG
+	void ClearDebugInstantEffects();
+#endif // ENABLE_VISUAL_LOG
 
 	UPROPERTY(ReplicatedUsing=OnRep_ClientDebugString)
 	TArray<FString>	ClientDebugStrings;
@@ -1058,96 +1059,11 @@
 	UFUNCTION()
 	virtual void OnRep_ServerDebugString();
 
-	float GetFilteredAttributeValue(const FGameplayAttribute& Attribute, const FGameplayTagRequirements& SourceTags, const FGameplayTagContainer& TargetTags, const TArray<FActiveGameplayEffectHandle>& HandlesToIgnore = TArray<FActiveGameplayEffectHandle>());
+	// ----------------------------------------------------------------------------------------------------------------
+	// Batching client->server RPCs
+	// This is a WIP feature to batch up client->server communication. It is opt in and not complete. It only batches the below functions. Other Server RPCs are not safe to call during a batch window. Only opt in if you know what you are doing!
+	// ----------------------------------------------------------------------------------------------------------------	
 
-protected:
-
-	/**
-	 *	The abilities we can activate. 
-	 *		-This will include CDOs for non instanced abilities and per-execution instanced abilities. 
-	 *		-Actor-instanced abilities will be the actual instance (not CDO)
-	 *		
-	 *	This array is not vital for things to work. It is a convenience thing for 'giving abilities to the actor'. But abilities could also work on things
-	 *	without an AbilitySystemComponent. For example an ability could be written to execute on a StaticMeshActor. As long as the ability doesn't require 
-	 *	instancing or anything else that the AbilitySystemComponent would provide, then it doesn't need the component to function.
-	 */
-
-	UPROPERTY(ReplicatedUsing=OnRep_ActivateAbilities, BlueprintReadOnly, Category = "Abilities")
-	FGameplayAbilitySpecContainer	ActivatableAbilities;
-
-	/** Maps from an ability spec to the target data. Used to track replicated data and callbacks */
-	TMap<FGameplayAbilitySpecHandleAndPredictionKey, FAbilityReplicatedDataCache> AbilityTargetDataMap;
-
-	/** Full list of all instance-per-execution gameplay abilities associated with this component */
-	UPROPERTY()
-	TArray<UGameplayAbility*>	AllReplicatedInstancedAbilities;
-
-	/** Will be called from GiveAbility or from OnRep. Initializes events (triggers and inputs) with the given ability */
-	virtual void OnGiveAbility(FGameplayAbilitySpec& AbilitySpec);
-
-	/** Will be called from RemoveAbility or from OnRep. Unbinds inputs with the given ability */
-	virtual void OnRemoveAbility(FGameplayAbilitySpec& AbilitySpec);
-
-	/** Called from ClearAbility, ClearAllAbilities or OnRep. Clears any triggers that should no longer exist. */
-	void CheckForClearedAbilities();
-
-	/** Cancel a specific ability spec */
-	void CancelAbilitySpec(FGameplayAbilitySpec& Spec, UGameplayAbility* Ignore);
-
-	/** Creates a new instance of an ability, storing it in the spec */
-	UGameplayAbility* CreateNewInstanceOfAbility(FGameplayAbilitySpec& Spec, const UGameplayAbility* Ability);
-
-	int32 AbilityScopeLockCount;
-	TArray<FGameplayAbilitySpecHandle, TInlineAllocator<2> > AbilityPendingRemoves;
-	TArray<FGameplayAbilitySpec, TInlineAllocator<2> > AbilityPendingAdds;
-
-	/** Local World time of the last ability activation. This is used for AFK/idle detection */
-	float AbilityLastActivatedTime;
-
-	UFUNCTION()
-	void	OnRep_ActivateAbilities();
-
-	UFUNCTION(Server, reliable, WithValidation)
-	void	ServerTryActivateAbility(FGameplayAbilitySpecHandle AbilityToActivate, bool InputPressed, FPredictionKey PredictionKey);
-
-	UFUNCTION(Server, reliable, WithValidation)
-	void	ServerTryActivateAbilityWithEventData(FGameplayAbilitySpecHandle AbilityToActivate, bool InputPressed, FPredictionKey PredictionKey, FGameplayEventData TriggerEventData);
-
-	UFUNCTION(Client, reliable)
-	void	ClientTryActivateAbility(FGameplayAbilitySpecHandle AbilityToActivate);
-
-	/** Called by ServerEndAbility and ClientEndAbility; avoids code duplication. */
-	void	RemoteEndOrCancelAbility(FGameplayAbilitySpecHandle AbilityToEnd, FGameplayAbilityActivationInfo ActivationInfo, bool bWasCanceled);
-
-	UFUNCTION(Server, reliable, WithValidation)
-	void	ServerEndAbility(FGameplayAbilitySpecHandle AbilityToEnd, FGameplayAbilityActivationInfo ActivationInfo, FPredictionKey PredictionKey);
-
-	UFUNCTION(Client, reliable)
-	void	ClientEndAbility(FGameplayAbilitySpecHandle AbilityToEnd, FGameplayAbilityActivationInfo ActivationInfo);
-
-	UFUNCTION(Server, reliable, WithValidation)
-	void    ServerCancelAbility(FGameplayAbilitySpecHandle AbilityToCancel, FGameplayAbilityActivationInfo ActivationInfo);
-
-	UFUNCTION(Client, reliable)
-	void    ClientCancelAbility(FGameplayAbilitySpecHandle AbilityToCancel, FGameplayAbilityActivationInfo ActivationInfo);
-
-	UFUNCTION(Client, Reliable)
-	void	ClientActivateAbilityFailed(FGameplayAbilitySpecHandle AbilityToActivate, int16 PredictionKey);
-	int32	ClientActivateAbilityFailedCountRecent;
-	float	ClientActivateAbilityFailedStartTime;
-
-	
-	void	OnClientActivateAbilityCaughtUp(FGameplayAbilitySpecHandle AbilityToActivate, FPredictionKey::KeyType PredictionKey);
-
-	UFUNCTION(Client, Reliable)
-	void	ClientActivateAbilitySucceed(FGameplayAbilitySpecHandle AbilityToActivate, FPredictionKey PredictionKey);
-
-	UFUNCTION(Client, Reliable)
-	void	ClientActivateAbilitySucceedWithEventData(FGameplayAbilitySpecHandle AbilityToActivate, FPredictionKey PredictionKey, FGameplayEventData TriggerEventData);
-
-public:
-	
-	// Batching client->server RPCs. This is a WIP feature to batch up client->server communication. It is opt in and not complete. It only batches the below functions. Other Server RPCs are not safe to call during a batch window. Only opt in if you know what you are doing!
 	void CallServerTryActivateAbility(FGameplayAbilitySpecHandle AbilityToActivate, bool InputPressed, FPredictionKey PredictionKey);
 	void CallServerSetReplicatedTargetData(FGameplayAbilitySpecHandle AbilityHandle, FPredictionKey AbilityOriginalPredictionKey, const FGameplayAbilityTargetDataHandle& ReplicatedTargetDataHandle, FGameplayTag ApplicationTag, FPredictionKey CurrentPredictionKey);
 	void CallServerEndAbility(FGameplayAbilitySpecHandle AbilityToEnd, FGameplayAbilityActivationInfo ActivationInfo, FPredictionKey PredictionKey);
@@ -1166,36 +1082,58 @@
 	virtual void ServerAbilityRPCBatch_Internal(FServerAbilityRPCBatch& BatchInfo);
 
 	// ----------------------------------------------------------------------------------------------------------------
+	//	Input handling/targeting
+	// ----------------------------------------------------------------------------------------------------------------	
 
-	// This is meant to be used to inhibit activating an ability from an input perspective. (E.g., the menu is pulled up, another game mechanism is consuming all input, etc)
-	// This should only be called on locally owned players.
-	// This should not be used to game mechanics like silences or disables. Those should be done through gameplay effects.
-
+	/**
+	 * This is meant to be used to inhibit activating an ability from an input perspective. (E.g., the menu is pulled up, another game mechanism is consuming all input, etc)
+	 * This should only be called on locally owned players.
+	 * This should not be used to game mechanics like silences or disables. Those should be done through gameplay effects.
+	 */
 	UFUNCTION(BlueprintCallable, Category="Abilities")
-	bool	GetUserAbilityActivationInhibited() const;
+	bool GetUserAbilityActivationInhibited() const;
 	
 	/** Disable or Enable a local user from being able to activate abilities. This should only be used for input/UI etc related inhibition. Do not use for game mechanics. */
 	UFUNCTION(BlueprintCallable, Category="Abilities")
-	void	SetUserAbilityActivationInhibited(bool NewInhibit);
+	virtual void SetUserAbilityActivationInhibited(bool NewInhibit);
 
-	bool	UserAbilityActivationInhibited;
+	/** Rather activation is currently inhibited */
+	UPROPERTY()
+	bool UserAbilityActivationInhibited;
 
-	// ----------------------------------------------------------------------------------------------------------------
+	/** When enabled, we will not replicate this ASC to simulated proxies. We will route multicast RPCs through   */
+	UPROPERTY()
+	bool ReplicationProxyEnabled;
 
+	/** Suppress all ability granting through GEs on this component */
+	UPROPERTY()
+	bool bSuppressGrantAbility;
+
+	/** Suppress all GameplayCues on this component */
+	UPROPERTY()
+	bool bSuppressGameplayCues;
+
+	/** List of currently active targeting actors */
+	UPROPERTY()
+	TArray<AGameplayAbilityTargetActor*> SpawnedTargetActors;
+
+	/** Bind to an input component with some default action names */
 	virtual void BindToInputComponent(UInputComponent* InputComponent);
 
-	void SetBlockAbilityBindingsArray(FGameplayAbilityInputBinds BindInfo);
+	/** Bind to an input component with customized bindings */
 	virtual void BindAbilityActivationToInputComponent(UInputComponent* InputComponent, FGameplayAbilityInputBinds BindInfo);
 
+	/** Initializes BlockedAbilityBindings variable */
+	virtual void SetBlockAbilityBindingsArray(FGameplayAbilityInputBinds BindInfo);
+
+	/** Called to handle ability bind input */
 	virtual void AbilityLocalInputPressed(int32 InputID);
 	virtual void AbilityLocalInputReleased(int32 InputID);
-	
+
+	/** Handle confirm/cancel for target actors */
 	virtual void LocalInputConfirm();
 	virtual void LocalInputCancel();
 	
-	/** Replicate that an ability has ended/canceled, to the client or server as appropriate */
-	void	ReplicateEndOrCancelAbility(FGameplayAbilitySpecHandle Handle, FGameplayAbilityActivationInfo ActivationInfo, UGameplayAbility* Ability, bool bWasCanceled);
-
 	/** InputID for binding GenericConfirm/Cancel events */
 	int32 GenericConfirmInputID;
 	int32 GenericCancelInputID;
@@ -1206,43 +1144,16 @@
 	/** Generic local callback for generic ConfirmEvent that any ability can listen to */
 	FAbilityConfirmOrCancel	GenericLocalConfirmCallbacks;
 
-	FAbilityEnded AbilityEndedCallbacks;
-
-	FAbilitySpecDirtied AbilitySpecDirtiedCallbacks;
-
 	/** Generic local callback for generic CancelEvent that any ability can listen to */
 	FAbilityConfirmOrCancel	GenericLocalCancelCallbacks;
 
-	/** A generic callback anytime an ability is activated (started) */	
-	FGenericAbilityDelegate AbilityActivatedCallbacks;
-
-	/** callback anytime an ability is ended */
-	FGameplayAbilityEndedDelegate OnAbilityEnded;
-
-	/** A generic callback anytime an ability is committed (cost/cooldown applied) */
-	FGenericAbilityDelegate AbilityCommittedCallbacks;
-	FAbilityFailedDelegate AbilityFailedCallbacks;
-
-	/** Executes a gameplay event. Returns the number of successful ability activations triggered by the event */
-	virtual int32 HandleGameplayEvent(FGameplayTag EventTag, const FGameplayEventData* Payload);
-
-	/** Generic callbacks for gameplay events. See UAbilityTask_WaitGameplayEvent */
-	TMap<FGameplayTag, FGameplayEventMulticastDelegate> GenericGameplayEventCallbacks;
-
-	virtual void NotifyAbilityCommit(UGameplayAbility* Ability);
-	virtual void NotifyAbilityActivated(const FGameplayAbilitySpecHandle Handle, UGameplayAbility* Ability);
-	virtual void NotifyAbilityFailed(const FGameplayAbilitySpecHandle Handle, UGameplayAbility* Ability, const FGameplayTagContainer& FailureReason);
-
-	UPROPERTY()
-	TArray<AGameplayAbilityTargetActor*>	SpawnedTargetActors;
-
 	/** Any active targeting actors will be told to stop and return current targeting data */
 	UFUNCTION(BlueprintCallable, Category = "Abilities")
-	void TargetConfirm();
+	virtual void TargetConfirm();
 
 	/** Any active targeting actors will be stopped and canceled, not returning any targeting data */
 	UFUNCTION(BlueprintCallable, Category = "Abilities")
-	void TargetCancel();
+	virtual void TargetCancel();
 
 	// ----------------------------------------------------------------------------------------------------------------
 	//	AnimMontage Support
@@ -1253,25 +1164,25 @@
 	// ----------------------------------------------------------------------------------------------------------------	
 
 	/** Plays a montage and handles replication and prediction based on passed in ability/activation info */
-	float PlayMontage(UGameplayAbility* AnimatingAbility, FGameplayAbilityActivationInfo ActivationInfo, UAnimMontage* Montage, float InPlayRate, FName StartSectionName = NAME_None);
+	virtual float PlayMontage(UGameplayAbility* AnimatingAbility, FGameplayAbilityActivationInfo ActivationInfo, UAnimMontage* Montage, float InPlayRate, FName StartSectionName = NAME_None);
 
 	/** Plays a montage without updating replication/prediction structures. Used by simulated proxies when replication tells them to play a montage. */
-	float PlayMontageSimulated(UAnimMontage* Montage, float InPlayRate, FName StartSectionName = NAME_None);
+	virtual float PlayMontageSimulated(UAnimMontage* Montage, float InPlayRate, FName StartSectionName = NAME_None);
 
 	/** Stops whatever montage is currently playing. Expectation is caller should only be stopping it if they are the current animating ability (or have good reason not to check) */
-	void CurrentMontageStop(float OverrideBlendOutTime = -1.0f);
+	virtual void CurrentMontageStop(float OverrideBlendOutTime = -1.0f);
 
 	/** Clear the animating ability that is passed in, if it's still currently animating */
-	void ClearAnimatingAbility(UGameplayAbility* Ability);
+	virtual void ClearAnimatingAbility(UGameplayAbility* Ability);
 
 	/** Jumps current montage to given section. Expectation is caller should only be stopping it if they are the current animating ability (or have good reason not to check) */
-	void CurrentMontageJumpToSection(FName SectionName);
+	virtual void CurrentMontageJumpToSection(FName SectionName);
 
 	/** Sets current montages next section name. Expectation is caller should only be stopping it if they are the current animating ability (or have good reason not to check) */
-	void CurrentMontageSetNextSectionName(FName FromSectionName, FName ToSectionName);
+	virtual void CurrentMontageSetNextSectionName(FName FromSectionName, FName ToSectionName);
 
 	/** Sets current montage's play rate */
-	void CurrentMontageSetPlayRate(float InPlayRate);
+	virtual void CurrentMontageSetPlayRate(float InPlayRate);
 
 	/** Returns true if the passed in ability is the current animating ability */
 	bool IsAnimatingAbility(UGameplayAbility* Ability) const;
@@ -1294,65 +1205,9 @@
 	/** Returns amount of time left in current section */
 	float GetCurrentMontageSectionTimeLeft() const;
 
-protected:
+	// ----------------------------------------------------------------------------------------------------------------
+	//	Actor interaction
+	// ----------------------------------------------------------------------------------------------------------------	
 
-	/** Implementation of ServerTryActivateAbility */
-	virtual void InternalServerTryActiveAbility(FGameplayAbilitySpecHandle AbilityToActivate, bool InputPressed, const FPredictionKey& PredictionKey, const FGameplayEventData* TriggerEventData);
-
-	/** Called when a prediction key that played a montage is rejected */
-	void OnPredictiveMontageRejected(UAnimMontage* PredictiveMontage);
-
-	/** Copy LocalAnimMontageInfo into RepAnimMontageInfo */
-	void AnimMontage_UpdateReplicatedData();
-
-	/** Data structure for replicating montage info to simulated clients */
-	UPROPERTY(ReplicatedUsing=OnRep_ReplicatedAnimMontage)
-	FGameplayAbilityRepAnimMontage RepAnimMontageInfo;
-
-	/** Set if montage rep happens while we don't have the animinstance associated with us yet */
-	bool bPendingMontagerep;
-
-	/** Data structure for montages that were instigated locally (everything if server, predictive if client. replicated if simulated proxy) */
-	UPROPERTY()
-	FGameplayAbilityLocalAnimMontage LocalAnimMontageInfo;
-
-	UFUNCTION()
-	void OnRep_ReplicatedAnimMontage();
-
-	/** Returns true if we are ready to handle replicated montage information */
-	virtual bool IsReadyForReplicatedMontage();
-
-	/** RPC function called from CurrentMontageSetNextSectopnName, replicates to other clients */
-	UFUNCTION(reliable, server, WithValidation)
-	void ServerCurrentMontageSetNextSectionName(UAnimMontage* ClientAnimMontage, float ClientPosition, FName SectionName, FName NextSectionName);
-
-	/** RPC function called from CurrentMontageJumpToSection, replicates to other clients */
-	UFUNCTION(reliable, server, WithValidation)
-	void ServerCurrentMontageJumpToSectionName(UAnimMontage* ClientAnimMontage, FName SectionName);
-
-	/** RPC function called from CurrentMontageSetPlayRate, replicates to other clients */
-	UFUNCTION(reliable, server, WithValidation)
-	void ServerCurrentMontageSetPlayRate(UAnimMontage* ClientAnimMontage, float InPlayRate);
-
-	/** Abilities that are triggered from a gameplay event */
-	TMap<FGameplayTag, TArray<FGameplayAbilitySpecHandle > > GameplayEventTriggeredAbilities;
-
-	/** Abilities that are triggered from a tag being added to the owner */
-	TMap<FGameplayTag, TArray<FGameplayAbilitySpecHandle > > OwnedTagTriggeredAbilities;
-
-	/** Callback that is called when an owned tag bound to an ability changes */
-	virtual void MonitoredTagChanged(const FGameplayTag Tag, int32 NewCount);
-
-	/** Returns true if the specified ability should be activated from an event in this network mode */
-	bool HasNetworkAuthorityToActivateTriggeredAbility(const FGameplayAbilitySpec &Spec) const;
-
-	virtual void OnImmunityBlockGameplayEffect(const FGameplayEffectSpec& Spec, const FActiveGameplayEffect* ImmunityGE);
-
-	// -----------------------------------------------------------------------------
-public:
-
-	/** Immunity notification support */
-	FImmunityBlockGE OnImmunityBlockGameplayEffectDelegate;
-
 	/** The actor that owns this component logically */
 	UPROPERTY(ReplicatedUsing = OnRep_OwningActor)
@@ -1375,6 +1230,12 @@
 	 */
 	virtual void InitAbilityActorInfo(AActor* InOwnerActor, AActor* InAvatarActor);
 
+	/** Returns avatar actor to be used for a specific task, normally GetAvatarActor */
+	virtual AActor* GetGameplayTaskAvatar(const UGameplayTask* Task) const override;
+
+	/** Returns the avatar actor for this component */
+	AActor* GetAvatarActor() const;
+
 	/** Changes the avatar actor, leaves the owner actor the same */
 	void SetAvatarActor(AActor* InAvatarActor);
 
@@ -1384,7 +1245,7 @@
 	/**
 	* This is called when the actor that is initialized to this system dies, this will clear that actor from this system and FGameplayAbilityActorInfo
 	*/
-	void ClearActorInfo();
+	virtual void ClearActorInfo();
 
 	/**
 	 *	This will refresh the Ability's ActorInfo structure based on the current ActorInfo. That is, AvatarActor will be the same but we will look for new
@@ -1392,11 +1253,10 @@
 	 */	
 	void RefreshAbilityActorInfo();
 
-	// -----------------------------------------------------------------------------
-
-	/**
-	 *	While these appear to be state, these are actually synchronization events w/ some payload data
-	 */
+	// ----------------------------------------------------------------------------------------------------------------
+	//	Synchronization RPCs
+	//  While these appear to be state, these are actually synchronization events w/ some payload data
+	// ----------------------------------------------------------------------------------------------------------------	
 	
 	/** Replicates the Generic Replicated Event to the server. */
 	UFUNCTION(Server, reliable, WithValidation)
@@ -1416,7 +1276,7 @@
 	/** Calls local callbacks that are registered with the given Generic Replicated Event */
 	bool InvokeReplicatedEventWithPayload(EAbilityGenericReplicatedEvent::Type EventType, FGameplayAbilitySpecHandle AbilityHandle, FPredictionKey AbilityOriginalPredictionKey, FPredictionKey CurrentPredictionKey, FVector_NetQuantize100 VectorPayload);
 	
-	/**  */
+	/** Replicates targeting data to the server */
 	UFUNCTION(Server, reliable, WithValidation)
 	void ServerSetReplicatedTargetData(FGameplayAbilitySpecHandle AbilityHandle, FPredictionKey AbilityOriginalPredictionKey, const FGameplayAbilityTargetDataHandle& ReplicatedTargetDataHandle, FGameplayTag ApplicationTag, FPredictionKey CurrentPredictionKey);
 
@@ -1462,7 +1322,7 @@
 	/** Returns Generic Replicated Event for a given Ability/PredictionKey pair */
 	FSimpleMulticastDelegate& AbilityReplicatedEventDelegate(EAbilityGenericReplicatedEvent::Type EventType, FGameplayAbilitySpecHandle AbilityHandle, FPredictionKey AbilityOriginalPredictionKey);
 
-	// Direct Input state replication. These will be called if bReplicateInputDirectly is true on the ability and is generally not a good thing to use. (Instead, prefer to use Generic Replicated Events).
+	/** Direct Input state replication. These will be called if bReplicateInputDirectly is true on the ability and is generally not a good thing to use. (Instead, prefer to use Generic Replicated Events). */
 	UFUNCTION(Server, reliable, WithValidation)
 	void ServerSetInputPressed(FGameplayAbilitySpecHandle AbilityHandle);
 
@@ -1475,50 +1335,174 @@
 	/** Called on local player always. Called on server only if bReplicateInputDirectly is set on the GameplayAbility. */
 	virtual void AbilitySpecInputReleased(FGameplayAbilitySpec& Spec);
 
-	// ---------------------------------------------------------------------
+	// ----------------------------------------------------------------------------------------------------------------
+	//  Component overrides
+	// ----------------------------------------------------------------------------------------------------------------
 
-	FORCEINLINE void SetTagMapCount(const FGameplayTag& Tag, int32 NewCount)
-	{
-		GameplayTagCountContainer.SetTagCount(Tag, NewCount);
-	}
-	
-	FORCEINLINE void UpdateTagMap(const FGameplayTag& BaseTag, int32 CountDelta)
-	{
-		if (GameplayTagCountContainer.UpdateTagCount(BaseTag, CountDelta))
-		{
-			OnTagUpdated(BaseTag, CountDelta > 0);
-		}
-	}
-	
-	FORCEINLINE void UpdateTagMap(const FGameplayTagContainer& Container, int32 CountDelta)
-	{
-		for (auto TagIt = Container.CreateConstIterator(); TagIt; ++TagIt)
-		{
-			const FGameplayTag& Tag = *TagIt;
-			UpdateTagMap(Tag, CountDelta);
-		}
-	}	
+	virtual void InitializeComponent() override;
+	virtual void UninitializeComponent() override;
+	virtual void OnComponentDestroyed(bool bDestroyingHierarchy) override;
+	virtual bool GetShouldTick() const override;
+	virtual void TickComponent(float DeltaTime, enum ELevelTick TickType, FActorComponentTickFunction *ThisTickFunction) override;
+	virtual void GetSubobjectsWithStableNamesForNetworking(TArray<UObject*>& Objs) override;
+	virtual bool ReplicateSubobjects(class UActorChannel *Channel, class FOutBunch *Bunch, FReplicationFlags *RepFlags) override;
+	/** Force owning actor to update it's replication, to make sure that gameplay cues get sent down quickly. Override to change how aggressive this is */
+	virtual void ForceReplication() override;
+	virtual void PreNetReceive() override;
+	virtual void PostNetReceive() override;
+	virtual void OnRegister() override;
+	virtual void OnUnregister() override;
+	virtual void BeginPlay() override;
 
-
-#if ENABLE_VISUAL_LOG
-	void ClearDebugInstantEffects();
-#endif // ENABLE_VISUAL_LOG
-
-	const FActiveGameplayEffect* GetActiveGameplayEffect(const FActiveGameplayEffectHandle Handle) const;
-
-	virtual AActor* GetGameplayTaskAvatar(const UGameplayTask* Task) const override;
-	AActor* GetAvatarActor() const;
-
-	/** Suppress all ability granting through GEs on this component */
-	bool bSuppressGrantAbility;
-
-	/** Suppress all GameplayCues on this component */
-	bool bSuppressGameplayCues;
-
-	/** Handle GameplayCues that may have been deferred while doing the NetDeltaSerialize and waiting for the avatar actor to get loaded */
-	void HandleDeferredGameplayCues(const FActiveGameplayEffectsContainer* GameplayEffectsContainer);
-
 protected:
 
+	/**
+	 *	The abilities we can activate. 
+	 *		-This will include CDOs for non instanced abilities and per-execution instanced abilities. 
+	 *		-Actor-instanced abilities will be the actual instance (not CDO)
+	 *		
+	 *	This array is not vital for things to work. It is a convenience thing for 'giving abilities to the actor'. But abilities could also work on things
+	 *	without an AbilitySystemComponent. For example an ability could be written to execute on a StaticMeshActor. As long as the ability doesn't require 
+	 *	instancing or anything else that the AbilitySystemComponent would provide, then it doesn't need the component to function.
+	 */
+
+	UPROPERTY(ReplicatedUsing=OnRep_ActivateAbilities, BlueprintReadOnly, Category = "Abilities")
+	FGameplayAbilitySpecContainer	ActivatableAbilities;
+
+	/** Maps from an ability spec to the target data. Used to track replicated data and callbacks */
+	TMap<FGameplayAbilitySpecHandleAndPredictionKey, FAbilityReplicatedDataCache> AbilityTargetDataMap;
+
+	/** List of gameplay tag container filters, and the delegates they call */
+	TArray<TPair<FGameplayTagContainer, FGameplayEventTagMulticastDelegate>> GameplayEventTagContainerDelegates;
+
+	/** Full list of all instance-per-execution gameplay abilities associated with this component */
+	UPROPERTY()
+	TArray<UGameplayAbility*>	AllReplicatedInstancedAbilities;
+
+	/** Will be called from GiveAbility or from OnRep. Initializes events (triggers and inputs) with the given ability */
+	virtual void OnGiveAbility(FGameplayAbilitySpec& AbilitySpec);
+
+	/** Will be called from RemoveAbility or from OnRep. Unbinds inputs with the given ability */
+	virtual void OnRemoveAbility(FGameplayAbilitySpec& AbilitySpec);
+
+	/** Called from ClearAbility, ClearAllAbilities or OnRep. Clears any triggers that should no longer exist. */
+	void CheckForClearedAbilities();
+
+	/** Cancel a specific ability spec */
+	virtual void CancelAbilitySpec(FGameplayAbilitySpec& Spec, UGameplayAbility* Ignore);
+
+	/** Creates a new instance of an ability, storing it in the spec */
+	virtual UGameplayAbility* CreateNewInstanceOfAbility(FGameplayAbilitySpec& Spec, const UGameplayAbility* Ability);
+
+	int32 AbilityScopeLockCount;
+	TArray<FGameplayAbilitySpecHandle, TInlineAllocator<2> > AbilityPendingRemoves;
+	TArray<FGameplayAbilitySpec, TInlineAllocator<2> > AbilityPendingAdds;
+
+	/** Local World time of the last ability activation. This is used for AFK/idle detection */
+	float AbilityLastActivatedTime;
+
+	UFUNCTION()
+	virtual void OnRep_ActivateAbilities();
+
+	UFUNCTION(Server, reliable, WithValidation)
+	void	ServerTryActivateAbility(FGameplayAbilitySpecHandle AbilityToActivate, bool InputPressed, FPredictionKey PredictionKey);
+
+	UFUNCTION(Server, reliable, WithValidation)
+	void	ServerTryActivateAbilityWithEventData(FGameplayAbilitySpecHandle AbilityToActivate, bool InputPressed, FPredictionKey PredictionKey, FGameplayEventData TriggerEventData);
+
+	UFUNCTION(Client, reliable)
+	void	ClientTryActivateAbility(FGameplayAbilitySpecHandle AbilityToActivate);
+
+	/** Called by ServerEndAbility and ClientEndAbility; avoids code duplication. */
+	void	RemoteEndOrCancelAbility(FGameplayAbilitySpecHandle AbilityToEnd, FGameplayAbilityActivationInfo ActivationInfo, bool bWasCanceled);
+
+	UFUNCTION(Server, reliable, WithValidation)
+	void	ServerEndAbility(FGameplayAbilitySpecHandle AbilityToEnd, FGameplayAbilityActivationInfo ActivationInfo, FPredictionKey PredictionKey);
+
+	UFUNCTION(Client, reliable)
+	void	ClientEndAbility(FGameplayAbilitySpecHandle AbilityToEnd, FGameplayAbilityActivationInfo ActivationInfo);
+
+	UFUNCTION(Server, reliable, WithValidation)
+	void    ServerCancelAbility(FGameplayAbilitySpecHandle AbilityToCancel, FGameplayAbilityActivationInfo ActivationInfo);
+
+	UFUNCTION(Client, reliable)
+	void    ClientCancelAbility(FGameplayAbilitySpecHandle AbilityToCancel, FGameplayAbilityActivationInfo ActivationInfo);
+
+	UFUNCTION(Client, Reliable)
+	void	ClientActivateAbilityFailed(FGameplayAbilitySpecHandle AbilityToActivate, int16 PredictionKey);
+	int32	ClientActivateAbilityFailedCountRecent;
+	float	ClientActivateAbilityFailedStartTime;
+
+	
+	void	OnClientActivateAbilityCaughtUp(FGameplayAbilitySpecHandle AbilityToActivate, FPredictionKey::KeyType PredictionKey);
+
+	UFUNCTION(Client, Reliable)
+	void	ClientActivateAbilitySucceed(FGameplayAbilitySpecHandle AbilityToActivate, FPredictionKey PredictionKey);
+
+	UFUNCTION(Client, Reliable)
+	void	ClientActivateAbilitySucceedWithEventData(FGameplayAbilitySpecHandle AbilityToActivate, FPredictionKey PredictionKey, FGameplayEventData TriggerEventData);
+
+	/** Implementation of ServerTryActivateAbility */
+	virtual void InternalServerTryActiveAbility(FGameplayAbilitySpecHandle AbilityToActivate, bool InputPressed, const FPredictionKey& PredictionKey, const FGameplayEventData* TriggerEventData);
+
+	/** Called when a prediction key that played a montage is rejected */
+	void OnPredictiveMontageRejected(UAnimMontage* PredictiveMontage);
+
+	/** Copy LocalAnimMontageInfo into RepAnimMontageInfo */
+	void AnimMontage_UpdateReplicatedData();
+
+	/** Data structure for replicating montage info to simulated clients */
+	UPROPERTY(ReplicatedUsing=OnRep_ReplicatedAnimMontage)
+	FGameplayAbilityRepAnimMontage RepAnimMontageInfo;
+
+	/** Cached value of rather this is a simulated actor */
+	UPROPERTY()
+	bool bCachedIsNetSimulated;
+
+	/** Set if montage rep happens while we don't have the animinstance associated with us yet */
+	UPROPERTY()
+	bool bPendingMontageRep;
+
+	/** Data structure for montages that were instigated locally (everything if server, predictive if client. replicated if simulated proxy) */
+	UPROPERTY()
+	FGameplayAbilityLocalAnimMontage LocalAnimMontageInfo;
+
+	UFUNCTION()
+	virtual void OnRep_ReplicatedAnimMontage();
+
+	/** Returns true if we are ready to handle replicated montage information */
+	virtual bool IsReadyForReplicatedMontage();
+
+	/** RPC function called from CurrentMontageSetNextSectopnName, replicates to other clients */
+	UFUNCTION(reliable, server, WithValidation)
+	void ServerCurrentMontageSetNextSectionName(UAnimMontage* ClientAnimMontage, float ClientPosition, FName SectionName, FName NextSectionName);
+
+	/** RPC function called from CurrentMontageJumpToSection, replicates to other clients */
+	UFUNCTION(reliable, server, WithValidation)
+	void ServerCurrentMontageJumpToSectionName(UAnimMontage* ClientAnimMontage, FName SectionName);
+
+	/** RPC function called from CurrentMontageSetPlayRate, replicates to other clients */
+	UFUNCTION(reliable, server, WithValidation)
+	void ServerCurrentMontageSetPlayRate(UAnimMontage* ClientAnimMontage, float InPlayRate);
+
+	/** Abilities that are triggered from a gameplay event */
+	TMap<FGameplayTag, TArray<FGameplayAbilitySpecHandle > > GameplayEventTriggeredAbilities;
+
+	/** Abilities that are triggered from a tag being added to the owner */
+	TMap<FGameplayTag, TArray<FGameplayAbilitySpecHandle > > OwnedTagTriggeredAbilities;
+
+	/** Callback that is called when an owned tag bound to an ability changes */
+	virtual void MonitoredTagChanged(const FGameplayTag Tag, int32 NewCount);
+
+	/** Returns true if the specified ability should be activated from an event in this network mode */
+	bool HasNetworkAuthorityToActivateTriggeredAbility(const FGameplayAbilitySpec &Spec) const;
+
+	virtual void OnImmunityBlockGameplayEffect(const FGameplayEffectSpec& Spec, const FActiveGameplayEffect* ImmunityGE);
+
+	// Internal gameplay cue functions
+	virtual void AddGameplayCue_Internal(const FGameplayTag GameplayCueTag, FGameplayEffectContextHandle& EffectContext, FActiveGameplayCueContainer& GameplayCueContainer);
+	virtual void AddGameplayCue_Internal(const FGameplayTag GameplayCueTag, const FGameplayCueParameters& GameplayCueParameters, FActiveGameplayCueContainer& GameplayCueContainer);
+	virtual void RemoveGameplayCue_Internal(const FGameplayTag GameplayCueTag, FActiveGameplayCueContainer& GameplayCueContainer);
+
 	/** Actually pushes the final attribute value to the attribute set's property. Should not be called by outside code since this does not go through the attribute aggregator system. */
 	void SetNumericAttribute_Internal(const FGameplayAttribute &Attribute, float& NewFloatValue);
@@ -1530,11 +1514,8 @@
 	void ExecuteGameplayEffect(FGameplayEffectSpec &Spec, FPredictionKey PredictionKey);
 
 	void CheckDurationExpired(FActiveGameplayEffectHandle Handle);
-
-	void OnAttributeGameplayEffectSpecExected(const FGameplayAttribute &Attribute, const struct FGameplayEffectSpec &Spec, struct FGameplayModifierEvaluatedData &Data);
 		
 	TArray<UGameplayTask*>&	GetAbilityActiveTasks(UGameplayAbility* Ability);
-	// --------------------------------------------
 	
 	// Contains all of the gameplay effects that are currently active on this component
 	UPROPERTY(Replicated)
@@ -1555,8 +1536,6 @@
 	TArray<uint8> BlockedAbilityBindings;
 
 	void DebugCyclicAggregatorBroadcasts(struct FAggregator* Aggregator);
-
-	// ---------------------------------------------
 	
 	// Acceleration map for all gameplay tags (OwnedGameplayTags from GEs and explicit GameplayCueTags)
 	FGameplayTagCountContainer GameplayTagCountContainer;
@@ -1570,13 +1549,6 @@
 
 	virtual void OnTagUpdated(const FGameplayTag& Tag, bool TagExists) {};
 	
-	// ---------------------------------------------
-
-	virtual void OnRegister() override;
-	virtual void OnUnregister() override;
-
-	virtual void BeginPlay() override;
-
 	const UAttributeSet*	GetAttributeSubobject(const TSubclassOf<UAttributeSet> AttributeClass) const;
 	const UAttributeSet*	GetAttributeSubobjectChecked(const TSubclassOf<UAttributeSet> AttributeClass) const;
 	const UAttributeSet*	GetOrCreateAttributeSubobject(TSubclassOf<UAttributeSet> AttributeClass);
@@ -1589,6 +1561,9 @@
 	friend struct FGameplayAbilitySpec;
 	friend struct FGameplayAbilitySpecContainer;
 	friend struct FAggregator;
+	friend struct FActiveGameplayEffectAction_Add;
+	friend struct FGameplayEffectSpec;
+	friend class AAbilitySystemDebugHUD;
 
 private:
 	FDelegateHandle MonitoredTagChangedDelegateHandle;
