Copying //UE4/Dev-Framework to //UE4/Dev-Main (Source: //UE4/Dev-Framework @ 4058146)

diff --git a/Engine/Plugins/Runtime/GameplayAbilities/Source/GameplayAbilities/Public/AbilitySystemGlobals.h b/Engine/Plugins/Runtime/GameplayAbilities/Source/GameplayAbilities/Public/AbilitySystemGlobals.h
--- a/Engine/Plugins/Runtime/GameplayAbilities/Source/GameplayAbilities/Public/AbilitySystemGlobals.h
+++ b/Engine/Plugins/Runtime/GameplayAbilities/Source/GameplayAbilities/Public/AbilitySystemGlobals.h
@@ -65,7 +65,7 @@ class GAMEPLAYABILITIES_API UAbilitySystemGlobals : public UObject
 	/** Global callback that can handle game-specific code that needs to run before applying a gameplay effect spec */
 	virtual void GlobalPreGameplayEffectSpecApply(FGameplayEffectSpec& Spec, UAbilitySystemComponent* AbilitySystemComponent);
 
-	// Stubs for WIP feature that will come to engine
+	/** Override to handle global state when gameplay effects are being applied */
 	virtual void PushCurrentAppliedGE(const FGameplayEffectSpec* Spec, UAbilitySystemComponent* AbilitySystemComponent) { }
 	virtual void SetCurrentAppliedGE(const FGameplayEffectSpec* Spec) { }
 	virtual void PopCurrentAppliedGE() { }
@@ -88,6 +88,7 @@ class GAMEPLAYABILITIES_API UAbilitySystemGlobals : public UObject
 	/** Sets a default gameplay cue tag using the asset's name. Returns true if it changed the tag. */
 	static bool DeriveGameplayCueTagFromAssetName(FString AssetName, FGameplayTag& GameplayCueTag, FName& GameplayCueName);
 
+	/** Sets a default gameplay cue tag using the asset's class*/
 	template<class T>
 	static void DeriveGameplayCueTagFromClass(T* CDO)
 	{
@@ -148,46 +149,52 @@ class GAMEPLAYABILITIES_API UAbilitySystemGlobals : public UObject
 	/** Returns true if ability costs are ignored, returns false otherwise. Always returns false in shipping builds. */
 	bool ShouldIgnoreCosts() const;
 
+	/** Called when debug strings are available, to write them to the display */
 	DECLARE_MULTICAST_DELEGATE(FOnClientServerDebugAvailable);
 	FOnClientServerDebugAvailable OnClientServerDebugAvailable;
 
 	/** Global place to accumulate debug strings for ability system component. Used when we fill up client side debug string immediately, and then wait for server to send server strings */
 	TArray<FString>	AbilitySystemDebugStrings;
 
-
-	// Helper functions for applying global scaling to various ability system tasks. This isn't meant to be a shipping feature, but to help with debugging and interation via cvar AbilitySystem.GlobalAbilityScale.
+	/** Helper functions for applying global scaling to various ability system tasks. This isn't meant to be a shipping feature, but to help with debugging and interation via cvar AbilitySystem.GlobalAbilityScale */
 	static void NonShipping_ApplyGlobalAbilityScaler_Rate(float& Rate);
 	static void NonShipping_ApplyGlobalAbilityScaler_Duration(float& Duration);
 
 	// Global Tags
 
+	/** TryActivate failed due to being dead */
 	UPROPERTY()
-	FGameplayTag ActivateFailIsDeadTag; // TryActivate failed due to being dead
+	FGameplayTag ActivateFailIsDeadTag; 
 	UPROPERTY(config)
 	FName ActivateFailIsDeadName;
 
+	/** TryActivate failed due to being on cooldown */
 	UPROPERTY()
-	FGameplayTag ActivateFailCooldownTag; // TryActivate failed due to being on cooldown
+	FGameplayTag ActivateFailCooldownTag; 
 	UPROPERTY(config)
 	FName ActivateFailCooldownName;
 
+	/** TryActivate failed due to not being able to spend costs */
 	UPROPERTY()
-	FGameplayTag ActivateFailCostTag; // TryActivate failed due to not being able to spend costs
+	FGameplayTag ActivateFailCostTag; 
 	UPROPERTY(config)
 	FName ActivateFailCostName;
 
+	/** TryActivate failed due to being blocked by other abilities */
 	UPROPERTY()
-	FGameplayTag ActivateFailTagsBlockedTag; // TryActivate failed due to being blocked by other abilities
+	FGameplayTag ActivateFailTagsBlockedTag; 
 	UPROPERTY(config)
 	FName ActivateFailTagsBlockedName;
 
+	/** TryActivate failed due to missing required tags */
 	UPROPERTY()
-	FGameplayTag ActivateFailTagsMissingTag; // TryActivate failed due to missing required tags
+	FGameplayTag ActivateFailTagsMissingTag; 
 	UPROPERTY(config)
 	FName ActivateFailTagsMissingName;
 
+	/** Failed to activate due to invalid networking settings, this is designer error */
 	UPROPERTY()
-	FGameplayTag ActivateFailNetworkingTag; // Failed to activate due to invalid networking settings, this is designer error
+	FGameplayTag ActivateFailNetworkingTag; 
 	UPROPERTY(config)
 	FName ActivateFailNetworkingName;
 
@@ -195,6 +202,7 @@ class GAMEPLAYABILITIES_API UAbilitySystemGlobals : public UObject
 	UPROPERTY(config)
 	int32	MinimalReplicationTagCountBits;
 
+	/** Initialize global tags by reading from config using the names and creating tags for use at runtime */
 	virtual void InitGlobalTags()
 	{
 		if (ActivateFailIsDeadName != NAME_None)
@@ -228,14 +236,16 @@ class GAMEPLAYABILITIES_API UAbilitySystemGlobals : public UObject
 		}
 	}
 
-	// GameplayCue Parameters
+	/** Initialize GameplayCue Parameters */
 	virtual void InitGameplayCueParameters(FGameplayCueParameters& CueParameters, const FGameplayEffectSpecForRPC &Spec);
 	virtual void InitGameplayCueParameters_GESpec(FGameplayCueParameters& CueParameters, const FGameplayEffectSpec &Spec);
 	virtual void InitGameplayCueParameters(FGameplayCueParameters& CueParameters, const FGameplayEffectContextHandle& EffectContext);
 
-	// Trigger async loading of the gameplay cue object libraries. By default, the manager will do this on creation,
-	// but that behavior can be changed by a derived class overriding ShouldAsyncLoadObjectLibrariesAtStart and returning false.
-	// In that case, this function must be called to begin the load
+	/**
+	 * Trigger async loading of the gameplay cue object libraries. By default, the manager will do this on creation,
+	 * but that behavior can be changed by a derived class overriding ShouldAsyncLoadObjectLibrariesAtStart and returning false.
+	 * In that case, this function must be called to begin the load
+	 */
 	virtual void StartAsyncLoadingObjectLibraries();
 
 	/** Simple accessor to whether gameplay modifier evaluation channels should be allowed or not */
@@ -253,6 +263,7 @@ class GAMEPLAYABILITIES_API UAbilitySystemGlobals : public UObject
 	/** Simple index-based accessor to the alias name for the specified gameplay mod evaluation channel, if any */
 	const FName& GetGameplayModEvaluationChannelAlias(int32 Index) const;
 
+	/** Path where the engine will load gameplay cue notifies from */
 	virtual TArray<FString> GetGameplayCueNotifyPaths() { return GameplayCueNotifyPaths; }
 
 protected:
@@ -287,10 +298,16 @@ protected:
 	UPROPERTY(config)
 	FSoftObjectPath GlobalCurveTableName;
 
+	UPROPERTY()
+	UCurveTable* GlobalCurveTable;
+
 	/** Holds information about the valid attributes' min and max values and stacking rules */
 	UPROPERTY(config)
 	FSoftObjectPath GlobalAttributeMetaDataTableName;
 
+	UPROPERTY()
+	UDataTable* GlobalAttributeMetaDataTable;
+
 	/** Holds default values for attribute sets, keyed off of Name/Levels. NOTE: Preserved for backwards compatibility, should use the array version below now */
 	UPROPERTY(config)
 	FSoftObjectPath GlobalAttributeSetDefaultsTableName;
@@ -299,6 +316,10 @@ protected:
 	UPROPERTY(config)
 	TArray<FSoftObjectPath> GlobalAttributeSetDefaultsTableNames;
 
+	/** Curve tables containing default values for attribute sets, keyed off of Name/Levels */
+	UPROPERTY()
+	TArray<UCurveTable*> GlobalAttributeDefaultsTables;
+
 	/** Class reference to gameplay cue manager. Use this if you want to just instantiate a class for your gameplay cue manager without having to create an asset. */
 	UPROPERTY(config)
 	FSoftObjectPath GlobalGameplayCueManagerClass;
@@ -322,19 +343,11 @@ protected:
 	UPROPERTY(config)
 	bool PredictTargetGameplayEffects;
 
-	UPROPERTY()
-	UCurveTable* GlobalCurveTable;
-
-	/** Curve tables containing default values for attribute sets, keyed off of Name/Levels */
-	UPROPERTY()
-	TArray<UCurveTable*> GlobalAttributeDefaultsTables;
-
-	UPROPERTY()
-	UDataTable* GlobalAttributeMetaDataTable;
-
+	/** Manager for all gameplay cues */
 	UPROPERTY()
 	UGameplayCueManager* GlobalGameplayCueManager;
 
+	/** Used to initialize attribute sets */
 	TSharedPtr<FAttributeSetInitter> GlobalAttributeSetInitter;
 
 	template <class T>
@@ -342,7 +355,6 @@ protected:
 
 #if WITH_EDITOR
 	void OnTableReimported(UObject* InObject);
-
 	void OnPreBeginPIE(const bool bIsSimulatingInEditor);
 #endif
 
@@ -363,7 +375,7 @@ public:
 	FOnAbilitySystemAssetFoundDelegate AbilityFindAssetInEditorCallbacks;
 };
 
-
+/** Scope object that indicates when a gameplay effect is being applied */
 struct FScopeCurrentGameplayEffectBeingApplied
 {
 	FScopeCurrentGameplayEffectBeingApplied(const FGameplayEffectSpec* Spec, UAbilitySystemComponent* AbilitySystemComponent)
diff --git a/Engine/Plugins/Runtime/GameplayAbilities/Source/GameplayAbilities/Public/AbilitySystemInterface.h b/Engine/Plugins/Runtime/GameplayAbilities/Source/GameplayAbilities/Public/AbilitySystemInterface.h
--- a/Engine/Plugins/Runtime/GameplayAbilities/Source/GameplayAbilities/Public/AbilitySystemInterface.h
+++ b/Engine/Plugins/Runtime/GameplayAbilities/Source/GameplayAbilities/Public/AbilitySystemInterface.h
@@ -20,5 +20,6 @@ class GAMEPLAYABILITIES_API IAbilitySystemInterface
 {
 	GENERATED_IINTERFACE_BODY()
 
+	/** Returns the ability system component to use for this actor. It may live on another actor, such as a Pawn using the PlayerState's component */
 	virtual UAbilitySystemComponent* GetAbilitySystemComponent() const = 0;
 };
diff --git a/Engine/Plugins/Runtime/GameplayAbilities/Source/GameplayAbilities/Public/AbilitySystemReplicationProxyInterface.h b/Engine/Plugins/Runtime/GameplayAbilities/Source/GameplayAbilities/Public/AbilitySystemReplicationProxyInterface.h
--- a/Engine/Plugins/Runtime/GameplayAbilities/Source/GameplayAbilities/Public/AbilitySystemReplicationProxyInterface.h
+++ b/Engine/Plugins/Runtime/GameplayAbilities/Source/GameplayAbilities/Public/AbilitySystemReplicationProxyInterface.h
@@ -11,9 +11,7 @@
 
 class UAbilitySystemComponent;
 
-// ---------------------------------------------------------
-
-/** Interface for actors that expose access to an ability system component */
+/** Interface for actors that act like an ability system component for replication. This can be used to group together different components for replication */
 UINTERFACE(MinimalAPI, meta = (CannotImplementInterfaceInBlueprint))
 class UAbilitySystemReplicationProxyInterface : public UInterface
 {
diff --git a/Engine/Plugins/Runtime/GameplayAbilities/Source/GameplayAbilities/Public/AttributeSet.h b/Engine/Plugins/Runtime/GameplayAbilities/Source/GameplayAbilities/Public/AttributeSet.h
--- a/Engine/Plugins/Runtime/GameplayAbilities/Source/GameplayAbilities/Public/AttributeSet.h
+++ b/Engine/Plugins/Runtime/GameplayAbilities/Source/GameplayAbilities/Public/AttributeSet.h
@@ -18,6 +18,7 @@ class UAttributeSet;
 struct FGameplayAbilityActorInfo;
 struct FAggregator;
 
+/** Place in an AttributeSet to create an attribute that can be accesed using FGameplayAttribute. It is strongly encouraged to use this instead of raw float attributes */
 USTRUCT(BlueprintType)
 struct GAMEPLAYABILITIES_API FGameplayAttributeData
 {
@@ -35,10 +36,16 @@ struct GAMEPLAYABILITIES_API FGameplayAttributeData
 	virtual ~FGameplayAttributeData()
 	{}
 
+	/** Returns the current value, which includes temporary buffs */
 	float GetCurrentValue() const;
+
+	/** Modifies current value, normally only called by ability system or during initialization */
 	virtual void SetCurrentValue(float NewValue);
 
+	/** Returns the base value which only includes permanent changes */
 	float GetBaseValue() const;
+
+	/** Modifies the permanent base value, normally only called by ability system or during initialization */
 	virtual void SetBaseValue(float NewValue);
 
 protected:
@@ -49,6 +56,7 @@ protected:
 	float CurrentValue;
 };
 
+/** Describes a FGameplayAttributeData or float property inside an attribute set. Using this provides editor UI and gelper functions */
 USTRUCT(BlueprintType)
 struct GAMEPLAYABILITIES_API FGameplayAttribute
 {
@@ -67,6 +75,7 @@ struct GAMEPLAYABILITIES_API FGameplayAttribute
 		return Attribute != nullptr;
 	}
 
+	/** Set up from a UProperty inside a set */
 	void SetUProperty(UProperty *NewProperty)
 	{
 		Attribute = NewProperty;
@@ -82,27 +91,33 @@ struct GAMEPLAYABILITIES_API FGameplayAttribute
 		}
 	}
 
+	/** Returns raw property */
 	UProperty* GetUProperty() const
 	{
 		return Attribute;
 	}
 
+	/** Returns the AttributeSet subclass holding this attribute */
 	UClass* GetAttributeSetClass() const
 	{
 		check(Attribute);
 		return CastChecked<UClass>(Attribute->GetOuter());
 	}
 
+	/** Returns true if this is one of the special attributes defined on the bBilitySystemComponent itself */
 	bool IsSystemAttribute() const;
 
 	/** Returns true if the variable associated with Property is of type FGameplayAttributeData or one of its subclasses */
 	static bool IsGameplayAttributeDataProperty(const UProperty* Property);
 
+	/** Modifies the current value of an attribute, will not modify base value if that is supported */
 	void SetNumericValueChecked(float& NewValue, class UAttributeSet* Dest) const;
 
+	/** Returns the current value of an attribute */
 	float GetNumericValue(const UAttributeSet* Src) const;
 	float GetNumericValueChecked(const UAttributeSet* Src) const;
 
+	/** Returns the AttributeData, will fail if this is a float attribute */
 	FGameplayAttributeData* GetGameplayAttributeData(UAttributeSet* Src) const;
 	FGameplayAttributeData* GetGameplayAttributeDataChecked(UAttributeSet* Src) const;
 	
@@ -116,17 +131,20 @@ struct GAMEPLAYABILITIES_API FGameplayAttribute
 		return PointerHash(InAttribute.Attribute);
 	}
 
+	/** Returns name of attribute, usually the same as the property */
 	FString GetName() const
 	{
 		return AttributeName.IsEmpty() ? *GetNameSafe(Attribute) : AttributeName;
 	}
 
+	/** Custom serialization */
 	void PostSerialize(const FArchive& Ar);
 
+	/** Name of the attribute, usually the same as property name */
 	UPROPERTY(Category = GameplayAttribute, VisibleAnywhere, BlueprintReadOnly)
 	FString AttributeName;
 
-	// In editor, this will filter out properties with meta tag "HideInDetailsView" or equal to FilterMetaStr. In non editor, it returns all properties.
+	/** In editor, this will filter out properties with meta tag "HideInDetailsView" or equal to FilterMetaStr. In non editor, it returns all properties */
 	static void GetAllAttributeProperties(TArray<UProperty*>& OutProperties, FString FilterMetaStr=FString(), bool UseEditorOnlyData=true);
 
 private:
@@ -148,19 +166,21 @@ struct TStructOpsTypeTraits< FGameplayAttribute > : public TStructOpsTypeTraitsB
 	};
 };
 
+/**
+ * Defines the set of all GameplayAttributes for your game
+ * Games should subclass this and add FGameplayAttributeData properties to represent attributes like health, damage, etc
+ * AttributeSets are added to the actors as subobjects, and then registered with the AbilitySystemComponent
+ * It often desired to have several sets per project that inherit from each other
+ * You could make a base health set, then have a player set that inherits from it and adds more attributes
+ */
 UCLASS(DefaultToInstanced, Blueprintable)
 class GAMEPLAYABILITIES_API UAttributeSet : public UObject
 {
 	GENERATED_UCLASS_BODY()
 
 public:
-	bool IsNameStableForNetworking() const override;
-
-	bool IsSupportedForNetworking() const override
-	{
-		return true;
-	}
 
+	/** Override to disable initialization for specific properties */
 	virtual bool ShouldInitProperty(bool FirstInit, UProperty* PropertyToInit) const { return true; }
 
 	/**
@@ -169,7 +189,6 @@ public:
 	 */	
 	virtual bool PreGameplayEffectExecute(struct FGameplayEffectModCallbackData &Data) { return true; }
 	
-	
 	/**
 	 *	Called just before a GameplayEffect is executed to modify the base value of an attribute. No more changes can be made.
 	 *	Note this is only called during an 'execute'. E.g., a modification to the 'base value' of an attribute. It is not called during an application of a GameplayEffect, such as a 5 ssecond +10 movement speed buff.
@@ -181,7 +200,6 @@ public:
 	 *	It is difficult to give all the information in these cases though - aggregators can change for many reasons: being added, being removed, being modified, having a modifier change, immunity, stacking rules, etc.
 	 */
 
-
 	/**
 	 *	Called just before any modification happens to an attribute. This is lower level than PreAttributeModify/PostAttribute modify.
 	 *	There is no additional context provided here since anything can trigger this. Executed effects, duration based effects, effects being removed, immunity being applied, stacking rules changing, etc.
@@ -205,16 +223,21 @@ public:
 	/** This signifies the attribute set can be ID'd by name over the network. */
 	void SetNetAddressable();
 
+	/** Initializes attribute data from a meta DataTable */
 	virtual void InitFromMetaDataTable(const UDataTable* DataTable);
 
+	/** Gets information about owning actor */
 	FORCEINLINE AActor* GetOwningActor() const { return CastChecked<AActor>(GetOuter()); }
 	UAbilitySystemComponent* GetOwningAbilitySystemComponent() const;
 	FGameplayAbilityActorInfo* GetActorInfo() const;
 
+	/** Print debug information to the log */
 	virtual void PrintDebug();
 
+	// Overrides
+	virtual bool IsNameStableForNetworking() const override;
+	virtual bool IsSupportedForNetworking() const override;
 	virtual void PreNetReceive() override;
-	
 	virtual void PostNetReceive() override;
 
 protected:
@@ -222,18 +245,7 @@ protected:
 	uint32 bNetAddressable : 1;
 };
 
-USTRUCT()
-struct GAMEPLAYABILITIES_API FGlobalCurveDataOverride
-{
-	GENERATED_USTRUCT_BODY()
-
-	TArray<UCurveTable*>	Overrides;
-};
-
-
-/**
- *	Generic numerical value in the form Coeffecient * Curve[Level] 
- */
+/** Generic numerical value in the form Value * Curve[Level] */
 USTRUCT(BlueprintType)
 struct GAMEPLAYABILITIES_API FScalableFloat
 {
@@ -259,34 +271,36 @@ struct GAMEPLAYABILITIES_API FScalableFloat
 
 public:
 
+	/** Raw value, is multiplied by curve */
 	UPROPERTY(EditAnywhere, BlueprintReadWrite, Category=ScalableFloat)
 	float	Value;
 
+	/** Curve that is evaluated at a specific level. If found, it is multipled by Value */
 	UPROPERTY(EditAnywhere, BlueprintReadWrite, Category=ScalableFloat)
 	FCurveTableRowHandle	Curve;
 
+	/** Returns the scaled value at a given level */
 	float GetValueAtLevel(float Level, const FString* ContextString = nullptr) const;
 
+	/** True if there is no curve lookup */
 	bool IsStatic() const
 	{
 		return Curve.RowName.IsNone();
 	}
 
+	/** Sets raw value */
 	void SetValue(float NewValue);
 
+	/** Overrides raw value and curve reference */
 	void SetScalingValue(float InCoeffecient, FName InRowName, UCurveTable * InTable);
 
-	void LockValueAtLevel(float Level, FGlobalCurveDataOverride *GlobalOverrides, const FString* ContextString = nullptr)
-	{
-		SetValue(GetValueAtLevel(Level, ContextString));
-	}
-
 	float GetValueChecked() const
 	{
 		check(IsStatic());
 		return Value;
 	}
 
+	/** Outputs human readable string */
 	FString ToSimpleString() const
 	{
 		if (Curve.RowName != NAME_None)
@@ -296,9 +310,9 @@ public:
 		return FString::Printf(TEXT("%.2f"), Value);
 	}
 
+	/** Error checking: checks if we have a curve table specified but no valid curve entry */
 	bool IsValid() const
-	{
-		// Error checking: checks if we have a curve table specified but no valid curve entry
+	{	
 		static const FString ContextString = TEXT("FScalableFloat::IsValid");
 		GetValueAtLevel(1.f, &ContextString);
 		bool bInvalid = (Curve.CurveTable != nullptr || Curve.RowName != NAME_None ) && (FinalCurve == nullptr);
@@ -309,7 +323,7 @@ public:
 	bool operator==(const FScalableFloat& Other) const;
 	bool operator!=(const FScalableFloat& Other) const;
 
-	//copy operator to prevent duplicate handles
+	/** copy operator to prevent duplicate handles */
 	void operator=(const FScalableFloat& Src);
 
 	/* Used to upgrade a float or int8/int16/int32 property into an FScalableFloat */
@@ -475,3 +489,53 @@ private:
 	static UProperty* ThisProperty = FindFieldChecked<UProperty>(C::StaticClass(), GET_MEMBER_NAME_CHECKED(C, P)); \
 	GetOwningAbilitySystemComponent()->SetBaseAttributeValueFromReplication(P, FGameplayAttribute(ThisProperty)); \
 }
+
+/**
+ * This defines a set of helper functions for accessing and initializing attributes, to avoid having to manually write these functions.
+ * It would creates the following functions, for attribute Health
+ *
+ *	static FGameplayAttribute UMyHealthSet::GetHealthAttribute();
+ *	FORCEINLINE float UMyHealthSet::GetHealth() const;
+ *	FORCEINLINE void UMyHealthSet::SetHealth(float NewVal);
+ *	FORCEINLINE void UMyHealthSet::InitHealth(float NewVal);
+ *
+ * To use this in your game you can define something like this, and then add game-specific functions as necessary:
+ * 
+ *	#define ATTRIBUTE_ACCESSORS(ClassName, PropertyName) \
+ *	GAMEPLAYATTRIBUTE_PROPERTY_GETTER(ClassName, PropertyName) \
+ *	GAMEPLAYATTRIBUTE_VALUE_GETTER(PropertyName) \
+ *	GAMEPLAYATTRIBUTE_VALUE_SETTER(PropertyName) \
+ *	GAMEPLAYATTRIBUTE_VALUE_INITTER(PropertyName)
+ * 
+ *	ATTRIBUTE_ACCESSORS(UMyHealthSet, Health)
+ */
+
+#define GAMEPLAYATTRIBUTE_PROPERTY_GETTER(ClassName, PropertyName) \
+	static FGameplayAttribute Get##PropertyName##Attribute() \
+	{ \
+		static UProperty* Prop = FindFieldChecked<UProperty>(ClassName::StaticClass(), GET_MEMBER_NAME_CHECKED(ClassName, PropertyName)); \
+		return Prop; \
+	}
+
+#define GAMEPLAYATTRIBUTE_VALUE_GETTER(PropertyName) \
+	FORCEINLINE float Get##PropertyName() const \
+	{ \
+		return PropertyName.GetCurrentValue(); \
+	}
+
+#define GAMEPLAYATTRIBUTE_VALUE_SETTER(PropertyName) \
+	FORCEINLINE void Set##PropertyName(float NewVal) \
+	{ \
+		UAbilitySystemComponent* AbilityComp = GetOwningAbilitySystemComponent(); \
+		if (ensure(AbilityComp)) \
+		{ \
+			AbilityComp->SetNumericAttributeBase(Get##PropertyName##Attribute(), NewVal); \
+		}; \
+	}
+
+#define GAMEPLAYATTRIBUTE_VALUE_INITTER(PropertyName) \
+	FORCEINLINE void Init##PropertyName(float NewVal) \
+	{ \
+		PropertyName.SetBaseValue(NewVal); \
+		PropertyName.SetCurrentValue(NewVal); \
+	}
\ No newline at end of file
diff --git a/Engine/Plugins/Runtime/GameplayAbilities/Source/GameplayAbilities/Public/GameplayAbilityBlueprint.h b/Engine/Plugins/Runtime/GameplayAbilities/Source/GameplayAbilities/Public/GameplayAbilityBlueprint.h
--- a/Engine/Plugins/Runtime/GameplayAbilities/Source/GameplayAbilities/Public/GameplayAbilityBlueprint.h
+++ b/Engine/Plugins/Runtime/GameplayAbilities/Source/GameplayAbilities/Public/GameplayAbilityBlueprint.h
@@ -8,8 +8,8 @@
 #include "GameplayAbilityBlueprint.generated.h"
 
 /**
- * A Gameplay Ability Blueprint is essentially a specialized Blueprint whose graphs control a gameplay ability.
- *	NOTE: This feature is EXPERIMENTAL. Use at your own risk!
+ * A Gameplay Ability Blueprint is essentially a specialized Blueprint whose graphs control a gameplay ability
+ * The ability factory should pick this for you automatically
  */
 
 UCLASS(BlueprintType)
diff --git a/Engine/Plugins/Runtime/GameplayAbilities/Source/GameplayAbilities/Public/GameplayAbilitySpec.h b/Engine/Plugins/Runtime/GameplayAbilities/Source/GameplayAbilities/Public/GameplayAbilitySpec.h
--- a/Engine/Plugins/Runtime/GameplayAbilities/Source/GameplayAbilities/Public/GameplayAbilitySpec.h
+++ b/Engine/Plugins/Runtime/GameplayAbilities/Source/GameplayAbilities/Public/GameplayAbilitySpec.h
@@ -20,22 +20,24 @@ class UGameplayAbility;
  *	GameplayEffect.h
  */
 
+/** Handle that points to a specific granted ability. These are globally unique */
 USTRUCT(BlueprintType)
 struct FGameplayAbilitySpecHandle
 {
 	GENERATED_USTRUCT_BODY()
 
 	FGameplayAbilitySpecHandle()
-	: Handle(INDEX_NONE)
+		: Handle(INDEX_NONE)
 	{
-
 	}
 
+	/** True if GenerateNewHandle was called on this handle */
 	bool IsValid() const
 	{
 		return Handle != INDEX_NONE;
 	}
 
+	/** Sets this to a valid handle */
 	void GenerateNewHandle()
 	{
 		static int32 GHandle = 1;
@@ -68,24 +70,25 @@ private:
 	int32 Handle;
 };
 
+/** Describes the status of activating this ability, this is updated as prediction is handled */
 UENUM(BlueprintType)
 namespace EGameplayAbilityActivationMode
 {
 	enum Type
 	{
-		// We are the authority activating this ability
+		/** We are the authority activating this ability */
 		Authority,
 
-		// We are not the authority but aren't predicting yet. This is a mostly invalid state to be in.
+		/** We are not the authority but aren't predicting yet. This is a mostly invalid state to be in */
 		NonAuthority,
 
-		// We are predicting the activation of this ability
+		/** We are predicting the activation of this ability */
 		Predicting,
 
-		// We are not the authority, but the authority has confirmed this activation
+		/** We are not the authority, but the authority has confirmed this activation */
 		Confirmed,
 
-		// We tried to activate it, and server told us we couldn't (even though we thought we could)
+		/** We tried to activate it, and server told us we couldn't (even though we thought we could) */
 		Rejected,
 	};
 }
@@ -98,7 +101,7 @@ enum class EGameplayEffectGrantedAbilityRemovePolicy : uint8
 	CancelAbilityImmediately,
 	/** Active abilities are allowed to finish, and then removed. */
 	RemoveAbilityOnEnd,
-	/** Granted abilties are left lone when the granting GameplayEffect is removed. */
+	/** Granted abilities are left lone when the granting GameplayEffect is removed. */
 	DoNothing,
 };
 
@@ -107,8 +110,7 @@ USTRUCT(BlueprintType)
 struct FGameplayAbilitySpecDef
 {
 	FGameplayAbilitySpecDef()
-		: Level(INDEX_NONE)
-		, InputID(INDEX_NONE)
+		: InputID(INDEX_NONE)
 		, RemovalPolicy(EGameplayEffectGrantedAbilityRemovePolicy::CancelAbilityImmediately)
 		, SourceObject(nullptr)
 	{
@@ -117,26 +119,27 @@ struct FGameplayAbilitySpecDef
 
 	GENERATED_USTRUCT_BODY()
 
+	/** What ability to grant */
 	UPROPERTY(EditDefaultsOnly, Category="Ability Definition", NotReplicated)
 	TSubclassOf<UGameplayAbility> Ability;
 	
-	// Deprecated for LevelScalableFloat
-	UPROPERTY(NotReplicated)
-	int32 Level; 
-
+	/** What level to grant this ability at */
 	UPROPERTY(EditDefaultsOnly, Category="Ability Definition", NotReplicated, DisplayName="Level")
 	FScalableFloat LevelScalableFloat; 
-	
+
+	/** Input ID to bind this ability to */
 	UPROPERTY(EditDefaultsOnly, Category="Ability Definition", NotReplicated)
 	int32 InputID;
 
+	/** What will remove this ability later */
 	UPROPERTY(EditDefaultsOnly, Category="Ability Definition", NotReplicated)
 	EGameplayEffectGrantedAbilityRemovePolicy RemovalPolicy;
-	
+
+	/** What granted this spec, not replicated or settable in editor */
 	UPROPERTY(NotReplicated)
 	UObject* SourceObject;
 
-	// SetbyCaller Magnitudes that were passed in to this ability by a GE (GE's that grant abilities). Made available so that 
+	/** SetbyCaller Magnitudes that were passed in to this ability by a GE (GE's that grant abilities) */
 	TMap<FGameplayTag, float>	SetByCallerTagMagnitudes;
 
 	/** This handle can be set if the SpecDef is used to create a real FGameplaybilitySpec */
@@ -162,7 +165,6 @@ struct GAMEPLAYABILITIES_API FGameplayAbilityActivationInfo
 		: ActivationMode(EGameplayAbilityActivationMode::Authority)
 		, bCanBeEndedByOtherInstance(false)
 	{
-
 	}
 
 	FGameplayAbilityActivationInfo(AActor* InActor)
@@ -179,6 +181,7 @@ struct GAMEPLAYABILITIES_API FGameplayAbilityActivationInfo
 	{
 	}	
 
+	/** Activation status of this ability */
 	UPROPERTY(BlueprintReadOnly, Category = "ActorInfo")
 	mutable TEnumAsByte<EGameplayAbilityActivationMode::Type>	ActivationMode;
 
@@ -186,8 +189,10 @@ struct GAMEPLAYABILITIES_API FGameplayAbilityActivationInfo
 	UPROPERTY()
 	uint8 bCanBeEndedByOtherInstance:1;
 
+	/** Called on client when activation is confirmed on server */
 	void SetActivationConfirmed();
 
+	/** Called when activation was rejected by the server */
 	void SetActivationRejected();
 
 	/** Called on client to set this as a predicted ability */
@@ -196,20 +201,20 @@ struct GAMEPLAYABILITIES_API FGameplayAbilityActivationInfo
 	/** Called on the server to set the key used by the client to predict this ability */
 	void ServerSetActivationPredictionKey(FPredictionKey PredictionKey);
 
+	/** Returns prediction key, const to avoid being able to modify it after creation */
 	const FPredictionKey& GetActivationPredictionKey() const { return PredictionKeyWhenActivated; }
 
 private:
 
-	// This was the prediction key used to activate this ability. It does not get updated
-	// if new prediction keys are generated over the course of the ability.
+	/** This was the prediction key used to activate this ability. It does not get updated if new prediction keys are generated over the course of the ability */
 	UPROPERTY()
 	FPredictionKey PredictionKeyWhenActivated;
 	
 };
 
-
-/** An activatable ability spec, hosted on the ability system component. This defines both what the ability is (what class, what level, input binding etc)
- *  and also holds runtime state that must be kept outside of the ability being instanced/activated.
+/**
+ * An activatable ability spec, hosted on the ability system component. This defines both what the ability is (what class, what level, input binding etc)
+ * and also holds runtime state that must be kept outside of the ability being instanced/activated.
  */
 USTRUCT(BlueprintType)
 struct GAMEPLAYABILITIES_API FGameplayAbilitySpec : public FFastArraySerializerItem
@@ -217,17 +222,16 @@ struct GAMEPLAYABILITIES_API FGameplayAbilitySpec : public FFastArraySerializerI
 	GENERATED_USTRUCT_BODY()
 
 	FGameplayAbilitySpec()
-	: Ability(nullptr), Level(1), InputID(INDEX_NONE), SourceObject(nullptr), ActiveCount(0), InputPressed(false), RemoveAfterActivation(false), PendingRemove(false)
-	{
-		
-	}
+		: Ability(nullptr), Level(1), InputID(INDEX_NONE), SourceObject(nullptr), ActiveCount(0), InputPressed(false), RemoveAfterActivation(false), PendingRemove(false)
+	{ }
 
-	FGameplayAbilitySpec(UGameplayAbility* InAbility, int32 InLevel=1, int32 InInputID=INDEX_NONE, UObject* InSourceObject=nullptr)
-		: Ability(InAbility), Level(InLevel), InputID(InInputID), SourceObject(InSourceObject), ActiveCount(0), InputPressed(false), RemoveAfterActivation(false), PendingRemove(false)
-	{
-		Handle.GenerateNewHandle();
-	}
-	
+	/** Version that takes an ability class */
+	FGameplayAbilitySpec(TSubclassOf<UGameplayAbility> InAbilityClass, int32 InLevel = 1, int32 InInputID = INDEX_NONE, UObject* InSourceObject = nullptr);
+
+	/** Version that takes an ability CDO, this exists for backward compatibility */
+	FGameplayAbilitySpec(UGameplayAbility* InAbility, int32 InLevel = 1, int32 InInputID = INDEX_NONE, UObject* InSourceObject = nullptr);
+
+	/** Version that takes an existing spec def */
 	FGameplayAbilitySpec(FGameplayAbilitySpecDef& InDef, int32 InGameplayEffectLevel, FActiveGameplayEffectHandle InGameplayEffectHandle = FActiveGameplayEffectHandle());
 
 	/** Handle for outside sources to refer to this spec by */
@@ -289,7 +293,7 @@ struct GAMEPLAYABILITIES_API FGameplayAbilitySpec : public FFastArraySerializerI
 	UGameplayAbility* GetPrimaryInstance() const;
 
 	/** interface function to see if the ability should replicated the ability spec or not */
-	bool ShouldReplicatedAbilitySpec() const;
+	bool ShouldReplicateAbilitySpec() const;
 
 	/** Returns all instances, which can include instance per execution abilities */
 	TArray<UGameplayAbility*> GetAbilityInstances() const
@@ -309,20 +313,26 @@ struct GAMEPLAYABILITIES_API FGameplayAbilitySpec : public FFastArraySerializerI
 	FString GetDebugString();
 };
 
-
 /** Fast serializer wrapper for above struct */
 USTRUCT(BlueprintType)
 struct GAMEPLAYABILITIES_API FGameplayAbilitySpecContainer : public FFastArraySerializer
 {
 	GENERATED_USTRUCT_BODY()
 
+	FGameplayAbilitySpecContainer()
+		: Owner(nullptr)
+	{
+	}
+
 	/** List of activatable abilities */
 	UPROPERTY()
 	TArray<FGameplayAbilitySpec> Items;
 
 	/** Component that owns this list */
+	UPROPERTY()
 	UAbilitySystemComponent* Owner;
 
+	/** Initializes Owner variable */
 	void RegisterWithOwner(UAbilitySystemComponent* Owner);
 
 	bool NetDeltaSerialize(FNetDeltaSerializeInfo & DeltaParms)
@@ -334,7 +344,7 @@ struct GAMEPLAYABILITIES_API FGameplayAbilitySpecContainer : public FFastArraySe
 	bool ShouldWriteFastArrayItem(const Type& Item, const bool bIsWritingOnClient)
 	{
 		// if we do not want the FGameplayAbilitySpec to replicated return false;
-		if (!Item.ShouldReplicatedAbilitySpec())
+		if (!Item.ShouldReplicateAbilitySpec())
 		{
 			return false;
 		}
@@ -357,9 +367,7 @@ struct TStructOpsTypeTraits< FGameplayAbilitySpecContainer > : public TStructOps
 	};
 };
 
-// ----------------------------------------------------
-
-// Used to stop us from removing abilities from an ability system component while we're iterating through the abilities
+/** Used to stop us from removing abilities from an ability system component while we're iterating through the abilities */
 struct GAMEPLAYABILITIES_API FScopedAbilityListLock
 {
 	FScopedAbilityListLock(UAbilitySystemComponent& InContainer);
@@ -371,7 +379,7 @@ private:
 
 #define ABILITYLIST_SCOPE_LOCK()	FScopedAbilityListLock ActiveScopeLock(*this);
 
-// Used to stop us from canceling or ending an ability while we're iterating through its gameplay targets
+/** Used to stop us from canceling or ending an ability while we're iterating through its gameplay targets */
 struct GAMEPLAYABILITIES_API FScopedTargetListLock
 {
 	FScopedTargetListLock(UAbilitySystemComponent& InAbilitySystemComponent, const UGameplayAbility& InAbility);
diff --git a/Engine/Plugins/Runtime/GameplayAbilities/Source/GameplayAbilities/Public/GameplayCueInterface.h b/Engine/Plugins/Runtime/GameplayAbilities/Source/GameplayAbilities/Public/GameplayCueInterface.h
--- a/Engine/Plugins/Runtime/GameplayAbilities/Source/GameplayAbilities/Public/GameplayCueInterface.h
+++ b/Engine/Plugins/Runtime/GameplayAbilities/Source/GameplayAbilities/Public/GameplayCueInterface.h
@@ -23,8 +23,10 @@ class GAMEPLAYABILITIES_API IGameplayCueInterface
 {
 	GENERATED_IINTERFACE_BODY()
 
+	/** Handle a single gameplay cue */
 	virtual void HandleGameplayCue(AActor *Self, FGameplayTag GameplayCueTag, EGameplayCueEvent::Type EventType, FGameplayCueParameters Parameters);
 
+	/** Wrapper that handles multiple cues */
 	virtual void HandleGameplayCues(AActor *Self, const FGameplayTagContainer& GameplayCueTags, EGameplayCueEvent::Type EventType, FGameplayCueParameters Parameters);
 
 	/** Returns true if the actor can currently accept gameplay cues associated with the given tag. Returns true by default. Allows actors to opt out of cues in cases such as pending death */
@@ -44,8 +46,10 @@ class GAMEPLAYABILITIES_API IGameplayCueInterface
 	UFUNCTION(BlueprintCallable, BlueprintCosmetic, Category="Ability|GameplayCue")
 	virtual void ForwardGameplayCueToParent();
 
+	/** Calls the UFunction override for a specific gameplay cue */
 	static void DispatchBlueprintCustomHandler(AActor* Actor, UFunction* Func, EGameplayCueEvent::Type EventType, FGameplayCueParameters Parameters);
 
+	/** Clears internal cache of what classes implement which functions */
 	static void ClearTagToFunctionMap();
 
 	IGameplayCueInterface() : bForwardToParent(false) {}
@@ -53,7 +57,6 @@ class GAMEPLAYABILITIES_API IGameplayCueInterface
 private:
 	/** If true, keep checking for additional handlers */
 	bool bForwardToParent;
-
 };
 
 
@@ -64,8 +67,6 @@ private:
  *	
  *	Essentially provides bare necessities to replicate GameplayCue Tags.
  */
-
-
 struct FActiveGameplayCueContainer;
 
 USTRUCT(BlueprintType)
diff --git a/Engine/Plugins/Runtime/GameplayAbilities/Source/GameplayAbilities/Public/GameplayCueManager.h b/Engine/Plugins/Runtime/GameplayAbilities/Source/GameplayAbilities/Public/GameplayCueManager.h
--- a/Engine/Plugins/Runtime/GameplayAbilities/Source/GameplayAbilities/Public/GameplayCueManager.h
+++ b/Engine/Plugins/Runtime/GameplayAbilities/Source/GameplayAbilities/Public/GameplayCueManager.h
@@ -26,6 +26,7 @@ DECLARE_DELEGATE_OneParam(FOnGameplayCueNotifySetLoaded, TArray<FSoftObjectPath>
 DECLARE_DELEGATE_OneParam(FGameplayCueProxyTick, float);
 DECLARE_DELEGATE_RetVal_TwoParams(bool, FShouldLoadGCNotifyDelegate, const FAssetData&, FName);
 
+
 /** An ObjectLibrary for the GameplayCue Notifies. Wraps 2 underlying UObjectLibraries plus options/delegates for how they are loaded */ 
 USTRUCT()
 struct FGameplayCueObjectLibrary
@@ -33,50 +34,60 @@ struct FGameplayCueObjectLibrary
 	GENERATED_BODY()
 	FGameplayCueObjectLibrary()
 	{
+		ActorObjectLibrary = nullptr;
+		StaticObjectLibrary = nullptr;
+		CueSet = nullptr;
+		AsyncPriority = 0;
+		bShouldSyncScan = false;
+		bShouldAsyncLoad = false;
+		bShouldSyncLoad = false;
 		bHasBeenInitialized = false;
 	}
 
-	// Paths to search for
+	/** Paths to search for */
 	UPROPERTY()
 	TArray<FString> Paths;
 
-	// Callback for when load finishes
+	/** Callback for when load finishes */
 	FOnGameplayCueNotifySetLoaded OnLoaded;
 
-	// Callback for "should I add this FAssetData to the set"
+	/** Callback for "should I add this FAssetData to the set" */
 	FShouldLoadGCNotifyDelegate ShouldLoad;
 
-	// Object library for actor based notifies
+	/** Object library for actor based notifies */
 	UPROPERTY()
 	UObjectLibrary* ActorObjectLibrary;
 
-	// Object library for object based notifies
+	/** Object library for object based notifies */
 	UPROPERTY()
 	UObjectLibrary* StaticObjectLibrary;
 
-	// Priority to use if async loading
+	/** Set to put the loaded asset data into. If null we will use the global set (RuntimeGameplayCueObjectLibrary.CueSet) */
+	UPROPERTY()
+	UGameplayCueSet* CueSet;
+
+	/** Priority to use if async loading */
 	TAsyncLoadPriority AsyncPriority;
 
-	// Should we force a sync scan on the asset registry in order to discover asset data, or just use what is there?
+	/** Should we force a sync scan on the asset registry in order to discover asset data, or just use what is there */
 	UPROPERTY()
 	bool bShouldSyncScan;
 
-	// Should we start async loading everything that we find (that passes ShouldLoad delegate check)
+	/** Should we start async loading everything that we find (that passes ShouldLoad delegate check) */
 	UPROPERTY()
 	bool bShouldAsyncLoad;
 
-	// Should we sync load everything that we find (that passes ShouldLoad delegate check)
+	/** Should we sync load everything that we find (that passes ShouldLoad delegate check) */
 	UPROPERTY()
 	bool bShouldSyncLoad;
 
-	// Set to put the loaded asset data into. If null we will use the global set (RuntimeGameplayCueObjectLibrary.CueSet)
-	UPROPERTY()
-	UGameplayCueSet* CueSet;
-
+	/** True if this has been initialized with correct data */
 	UPROPERTY()
 	bool bHasBeenInitialized;
 };
 
+
+/** Singleton manager object that handles dispatching gameplay cues and spawning GameplayCueNotify actors as needed */
 UCLASS()
 class GAMEPLAYABILITIES_API UGameplayCueManager : public UDataAsset
 {
@@ -85,7 +96,6 @@ class GAMEPLAYABILITIES_API UGameplayCueManager : public UDataAsset
 	// -------------------------------------------------------------
 	// Wrappers to handle replicating executed cues
 	// -------------------------------------------------------------
-
 	virtual void InvokeGameplayCueExecuted_FromSpec(UAbilitySystemComponent* OwningComponent, const FGameplayEffectSpec& Spec, FPredictionKey PredictionKey);
 	virtual void InvokeGameplayCueExecuted(UAbilitySystemComponent* OwningComponent, const FGameplayTag GameplayCueTag, FPredictionKey PredictionKey, FGameplayEffectContextHandle EffectContext);
 	virtual void InvokeGameplayCueExecuted_WithParams(UAbilitySystemComponent* OwningComponent, const FGameplayTag GameplayCueTag, FPredictionKey PredictionKey, FGameplayCueParameters GameplayCueParameters);
@@ -102,8 +112,10 @@ class GAMEPLAYABILITIES_API UGameplayCueManager : public UDataAsset
 	/** Broadcasted when ::FlushPendingCues runs: useful for custom batching/gameplay cue handling */
 	FSimpleMulticastDelegate	OnFlushPendingCues;
 
+	/** Called when manager is first created */
 	virtual void OnCreated();
 
+	/** Called when engine has completely loaded, this is a good time to finalize things */
 	virtual void OnEngineInitComplete();
 
 	/** Process a pending cue, return false if the cue should be rejected. */
diff --git a/Engine/Plugins/Runtime/GameplayAbilities/Source/GameplayAbilities/Public/GameplayCue_Types.h b/Engine/Plugins/Runtime/GameplayAbilities/Source/GameplayAbilities/Public/GameplayCue_Types.h
--- a/Engine/Plugins/Runtime/GameplayAbilities/Source/GameplayAbilities/Public/GameplayCue_Types.h
+++ b/Engine/Plugins/Runtime/GameplayAbilities/Source/GameplayAbilities/Public/GameplayCue_Types.h
@@ -15,14 +15,20 @@ class AGameplayCueNotify_Actor;
 class UAbilitySystemComponent;
 class UGameplayCueSet;
 
+
+/** Describes what type of payload is attached to a cue execution, we only replicate what is needed */
 UENUM()
 enum class EGameplayCuePayloadType : uint8
 {
+	/** Uses FGameplayEffectContextHandle inside GameplayCueParameters */
 	EffectContext,
+	/** Uses FGameplayCueParameters */
 	CueParameters,
+	/** Uses FGameplayEffectSpecForRPC */
 	FromSpec,
 };
 
+
 /** Structure to keep track of pending gameplay cues that haven't been applied yet. */
 USTRUCT()
 struct FGameplayCuePendingExecute
@@ -30,11 +36,12 @@ struct FGameplayCuePendingExecute
 	GENERATED_USTRUCT_BODY()
 
 	FGameplayCuePendingExecute()
-	: PayloadType(EGameplayCuePayloadType::EffectContext)
-	, OwningComponent(NULL)
+		: PayloadType(EGameplayCuePayloadType::EffectContext)
+		, OwningComponent(nullptr)
 	{
 	}
-	
+
+	/** List of tags, we allocate one as there is almost always exactly one tag */
 	TArray<FGameplayTag, TInlineAllocator<1> > GameplayCueTags;
 	
 	/** Prediction key that spawned this cue */
@@ -58,20 +65,25 @@ struct FGameplayCuePendingExecute
 	FGameplayCueParameters CueParameters;
 };
 
+
 /** Struct for pooling and preallocating gameplaycuenotify_actor classes. This data is per world and used to track what actors are available to recycle and which classes need to preallocate instances of those actors */
 USTRUCT()
 struct FPreallocationInfo
 {
 	GENERATED_USTRUCT_BODY()
 
+	/** Raw list of pooled instances. This relies on NotifyGameplayCueActorEndPlay always being called when actor is destroyed */
 	TMap<UClass*, TArray<AGameplayCueNotify_Actor*> >	PreallocatedInstances;
 
+	/** List of calsses that will be pooled */
 	UPROPERTY(transient)
 	TArray<TSubclassOf<AGameplayCueNotify_Actor>>	ClassesNeedingPreallocation;
 
+	/** World that owns this list */
 	FObjectKey OwningWorldKey;
 };
 
+
 /** Struct that is used by the gameplaycue manager to tie an instanced gameplaycue to the calling gamecode. Usually this is just the target actor, but can also be unique per instigator/sourceobject */
 struct FGCNotifyActorKey
 {
diff --git a/Engine/Plugins/Runtime/GameplayAbilities/Source/GameplayAbilities/Public/GameplayEffect.h b/Engine/Plugins/Runtime/GameplayAbilities/Source/GameplayAbilities/Public/GameplayEffect.h
--- a/Engine/Plugins/Runtime/GameplayAbilities/Source/GameplayAbilities/Public/GameplayEffect.h
+++ b/Engine/Plugins/Runtime/GameplayAbilities/Source/GameplayAbilities/Public/GameplayEffect.h
@@ -557,6 +557,7 @@ struct FGameplayEffectCue
 	}
 };
 
+/** Structure that is used to combine tags from parent and child blueprints in a safe way */
 USTRUCT(BlueprintType)
 struct GAMEPLAYABILITIES_API FInheritedTagContainer
 {
diff --git a/Engine/Plugins/Runtime/GameplayAbilities/Source/GameplayAbilities/Public/GameplayEffectTypes.h b/Engine/Plugins/Runtime/GameplayAbilities/Source/GameplayAbilities/Public/GameplayEffectTypes.h
--- a/Engine/Plugins/Runtime/GameplayAbilities/Source/GameplayAbilities/Public/GameplayEffectTypes.h
+++ b/Engine/Plugins/Runtime/GameplayAbilities/Source/GameplayAbilities/Public/GameplayEffectTypes.h
@@ -19,7 +19,7 @@
 #if SKILL_SYSTEM_AGGREGATOR_DEBUG
 	#define SKILL_AGG_DEBUG( Format, ... ) *FString::Printf(Format, ##__VA_ARGS__)
 #else
-	#define SKILL_AGG_DEBUG( Format, ... ) NULL
+	#define SKILL_AGG_DEBUG( Format, ... ) nullptr
 #endif
 
 class Error;
@@ -29,14 +29,13 @@ struct FActiveGameplayEffect;
 struct FGameplayEffectModCallbackData;
 struct FGameplayEffectSpec;
 
+/** Wrappers to convert enum to string. These are fairly slow */
 GAMEPLAYABILITIES_API FString EGameplayModOpToString(int32 Type);
-
 GAMEPLAYABILITIES_API FString EGameplayModToString(int32 Type);
-
 GAMEPLAYABILITIES_API FString EGameplayModEffectToString(int32 Type);
-
 GAMEPLAYABILITIES_API FString EGameplayCueEventToString(int32 Type);
 
+
 /** Valid gameplay modifier evaluation channels; Displayed and renamed via game-specific aliases and options */
 UENUM()
 enum class EGameplayModEvaluationChannel : uint8
@@ -56,6 +55,7 @@ enum class EGameplayModEvaluationChannel : uint8
 	Channel_MAX UMETA(Hidden)
 };
 
+
 /** Struct representing evaluation channel settings for a gameplay modifier */
 USTRUCT()
 struct GAMEPLAYABILITIES_API FGameplayModEvaluationChannelSettings
@@ -89,9 +89,11 @@ protected:
 	friend class FGameplayModEvaluationChannelSettingsDetails;
 };
 
+
 UENUM(BlueprintType)
 namespace EGameplayModOp
 {
+	/** Defines the ways that mods will modify attributes. Numeric ones operate on the existing value, override ignores it */
 	enum Type
 	{		
 		/** Numeric. */
@@ -155,6 +157,7 @@ enum class EGameplayEffectStackingType : uint8
 	AggregateByTarget,
 };
 
+
 /**
  * This handle is required for things outside of FActiveGameplayEffectsContainer to refer to a specific active GameplayEffect
  *	For example if a skill needs to create an active effect and then destroy that specific effect that it created, it has to do so
@@ -179,23 +182,29 @@ struct GAMEPLAYABILITIES_API FActiveGameplayEffectHandle
 
 	}
 
+	/** True if this is tracking an active ongoing gameplay effect */
 	bool IsValid() const
 	{
 		return Handle != INDEX_NONE;
 	}
 
+	/** True if this applied a gameplay effect. This can be true when it is not valid if it applied an instant effect */
 	bool WasSuccessfullyApplied() const
 	{
 		return bPassedFiltersAndWasExecuted;
 	}
 
+	/** Creates a new handle, will be set to successfully applied */
 	static FActiveGameplayEffectHandle GenerateNewHandle(UAbilitySystemComponent* OwningComponent);
 
+	/** Resets the map that supports GetOwningAbilitySystemComponent */
 	static void ResetGlobalHandleMap();
 
+	/** Returns the ability system component that created this handle */
 	UAbilitySystemComponent* GetOwningAbilitySystemComponent();
 	const UAbilitySystemComponent* GetOwningAbilitySystemComponent() const;
 
+	/** Remove this from the GetOwningAbilitySystemComponent map */
 	void RemoveFromGlobalMap();
 
 	bool operator==(const FActiveGameplayEffectHandle& Other) const
@@ -232,6 +241,8 @@ private:
 	bool bPassedFiltersAndWasExecuted;
 };
 
+
+/** Data that describes what happened in an attribute modification. This is passed to ability set callbacks */
 USTRUCT(BlueprintType)
 struct FGameplayModifierEvaluatedData
 {
@@ -254,6 +265,7 @@ struct FGameplayModifierEvaluatedData
 	{
 	}
 
+	/** What attribute was modified */
 	UPROPERTY()
 	FGameplayAttribute Attribute;
 
@@ -261,6 +273,7 @@ struct FGameplayModifierEvaluatedData
 	UPROPERTY()
 	TEnumAsByte<EGameplayModOp::Type> ModifierOp;
 
+	/** The raw magnitude of the applied attribute, this is generally before being clamped */
 	UPROPERTY()
 	float Magnitude;
 
@@ -268,6 +281,7 @@ struct FGameplayModifierEvaluatedData
 	UPROPERTY()
 	FActiveGameplayEffectHandle	Handle;
 
+	/** True if something was evaluated */
 	UPROPERTY()
 	bool IsValid;
 
@@ -277,6 +291,7 @@ struct FGameplayModifierEvaluatedData
 	}
 };
 
+
 /** Struct defining gameplay attribute capture options for gameplay effects */
 USTRUCT(BlueprintType)
 struct GAMEPLAYABILITIES_API FGameplayEffectAttributeCaptureDefinition
@@ -328,10 +343,11 @@ struct GAMEPLAYABILITIES_API FGameplayEffectAttributeCaptureDefinition
 	FString ToSimpleString() const;
 };
 
+
 /**
- * FGameplayEffectContext
- *	Data struct for an instigator and related data. This is still being fleshed out. We will want to track actors but also be able to provide some level of tracking for actors that are destroyed.
- *	We may need to store some positional information as well.
+ * Data structure that stores an instigator and related data, such as positions and targets
+ * Games can subclass this structure and add game-specific information
+ * It is passed throughout effect execution so it is a great place to track transient information about an execution
  */
 USTRUCT()
 struct GAMEPLAYABILITIES_API FGameplayEffectContext
@@ -372,11 +388,13 @@ struct GAMEPLAYABILITIES_API FGameplayEffectContext
 		return Instigator.Get();
 	}
 
-	/** returns the CDO of the ability used to instigate this context */
+	/** Returns the CDO of the ability used to instigate this context */
 	const UGameplayAbility* GetAbility() const;
 
+	/** Returns the specific instance that instigated this, may not always be set */
 	const UGameplayAbility* GetAbilityInstance_NotReplicated() const;
 
+	/** Gets the ability level this was evaluated at */
 	int32 GetAbilityLevel() const
 	{
 		return AbilityLevel;
@@ -394,6 +412,7 @@ struct GAMEPLAYABILITIES_API FGameplayEffectContext
 		return EffectCauser.Get();
 	}
 
+	/** Modify the effect causer actor, useful when that information is added after creation */
 	void SetEffectCauser(AActor* InEffectCauser)
 	{
 		EffectCauser = InEffectCauser;
@@ -424,42 +443,52 @@ struct GAMEPLAYABILITIES_API FGameplayEffectContext
 		return SourceObject.Get();
 	}
 
+	/** Add actors to the stored actor list */
 	virtual void AddActors(const TArray<TWeakObjectPtr<AActor>>& IActor, bool bReset = false);
 
+	/** Add a hit result for targeting */
 	virtual void AddHitResult(const FHitResult& InHitResult, bool bReset = false);
 
+	/** Returns actor list, may be empty */
 	virtual const TArray<TWeakObjectPtr<AActor>>& GetActors() const
 	{
 		return Actors;
 	}
 
+	/** Returns hit result, this can be null */
 	virtual const FHitResult* GetHitResult() const
 	{
 		return const_cast<FGameplayEffectContext*>(this)->GetHitResult();
 	}
 
+	/** Returns hit result, this can be null */
 	virtual FHitResult* GetHitResult()
 	{
 		return HitResult.Get();
 	}
 
+	/** Adds an origin point */
 	virtual void AddOrigin(FVector InOrigin);
 
+	/** Returns origin point, may be invalid if HasOrigin is false */
 	virtual const FVector& GetOrigin() const
 	{
 		return WorldOrigin;
 	}
 
+	/** Returns true if GetOrigin will give valid information */
 	virtual bool HasOrigin() const
 	{
 		return bHasWorldOrigin;
 	}
 
+	/** Returns debug string */
 	virtual FString ToString() const
 	{
 		return Instigator.IsValid() ? Instigator->GetName() : FString(TEXT("NONE"));
 	}
 
+	/** Returns the actual struct used for serialization, subclasses must override this! */
 	virtual UScriptStruct* GetScriptStruct() const
 	{
 		return FGameplayEffectContext::StaticStruct();
@@ -479,10 +508,13 @@ struct GAMEPLAYABILITIES_API FGameplayEffectContext
 		return NewContext;
 	}
 
+	/** True if this was instigated by a locally controlled actor */
 	virtual bool IsLocallyControlled() const;
 
+	/** True if this was instigated by a locally controlled player */
 	virtual bool IsLocallyControlledPlayer() const;
 
+	/** Custom serialization, subclasses must override this */
 	virtual bool NetSerialize(FArchive& Ar, class UPackageMap* Map, bool& bOutSuccess);
 
 protected:
@@ -496,14 +528,15 @@ protected:
 	UPROPERTY()
 	TWeakObjectPtr<AActor> EffectCauser;
 
-	/** the ability CDO that is responsible for this effect context (replicated) */
+	/** The ability CDO that is responsible for this effect context (replicated) */
 	UPROPERTY()
 	TWeakObjectPtr<UGameplayAbility> AbilityCDO;
 
-	/** the ability instance that is responsible for this effect context (NOT replicated) */
+	/** The ability instance that is responsible for this effect context (NOT replicated) */
 	UPROPERTY(NotReplicated)
 	TWeakObjectPtr<UGameplayAbility> AbilityInstanceNotReplicated;
 
+	/** The level this was executed at */
 	UPROPERTY()
 	int32 AbilityLevel;
 
@@ -515,12 +548,14 @@ protected:
 	UPROPERTY(NotReplicated)
 	TWeakObjectPtr<UAbilitySystemComponent> InstigatorAbilitySystemComponent;
 
+	/** Actors referenced by this context */
 	UPROPERTY()
 	TArray<TWeakObjectPtr<AActor>> Actors;
 
-	/** Trace information - may be NULL in many cases */
+	/** Trace information - may be nullptr in many cases */
 	TSharedPtr<FHitResult>	HitResult;
 
+	/** Stored origin, may be invalid if bHasWorldOrigin is false */
 	UPROPERTY()
 	FVector	WorldOrigin;
 
@@ -542,6 +577,7 @@ struct TStructOpsTypeTraits< FGameplayEffectContext > : public TStructOpsTypeTra
 	};
 };
 
+
 /**
  * Handle that wraps a FGameplayEffectContext or subclass, to allow it to be polymorphic and replicate properly
  */
@@ -580,14 +616,14 @@ struct FGameplayEffectContextHandle
 		return Data.IsValid();
 	}
 
+	/** Returns Raw effet context, may be null */
 	FGameplayEffectContext* Get()
 	{
-		return IsValid() ? Data.Get() : NULL;
+		return IsValid() ? Data.Get() : nullptr;
 	}
-
 	const FGameplayEffectContext* Get() const
 	{
-		return IsValid() ? Data.Get() : NULL;
+		return IsValid() ? Data.Get() : nullptr;
 	}
 
 	/** Returns the list of gameplay tags applicable to this effect, defaults to the owner's tags */
@@ -608,6 +644,7 @@ struct FGameplayEffectContextHandle
 		}
 	}
 
+	/** Sets Abilit instance and CDO parameters on context */
 	void SetAbility(const UGameplayAbility* InGameplayAbility)
 	{
 		if (IsValid())
@@ -623,7 +660,7 @@ struct FGameplayEffectContextHandle
 		{
 			return Data->GetInstigator();
 		}
-		return NULL;
+		return nullptr;
 	}
 
 	/** Returns the Ability CDO */
@@ -646,6 +683,7 @@ struct FGameplayEffectContextHandle
 		return nullptr;
 	}
 
+	/** Returns level this was executed at */
 	int32 GetAbilityLevel() const
 	{
 		if (IsValid())
@@ -662,7 +700,7 @@ struct FGameplayEffectContextHandle
 		{
 			return Data->GetInstigatorAbilitySystemComponent();
 		}
-		return NULL;
+		return nullptr;
 	}
 
 	/** Returns the physical actor tied to the application of this effect */
@@ -672,7 +710,7 @@ struct FGameplayEffectContextHandle
 		{
 			return Data->GetEffectCauser();
 		}
-		return NULL;
+		return nullptr;
 	}
 
 	/** Should always return the original instigator that started the whole chain. Subclasses can override what this does */
@@ -682,7 +720,7 @@ struct FGameplayEffectContextHandle
 		{
 			return Data->GetOriginalInstigator();
 		}
-		return NULL;
+		return nullptr;
 	}
 
 	/** Returns the ability system component of the instigator that started the whole chain */
@@ -692,7 +730,7 @@ struct FGameplayEffectContextHandle
 		{
 			return Data->GetOriginalInstigatorAbilitySystemComponent();
 		}
-		return NULL;
+		return nullptr;
 	}
 
 	/** Sets the object this effect was created from. */
@@ -724,6 +762,7 @@ struct FGameplayEffectContextHandle
 		return false;
 	}
 
+	/** Returns if the instigator is locally controlled and a player */
 	bool IsLocallyControlledPlayer() const
 	{
 		if (IsValid())
@@ -733,6 +772,7 @@ struct FGameplayEffectContextHandle
 		return false;
 	}
 
+	/** Add actors to the stored actor list */
 	void AddActors(const TArray<TWeakObjectPtr<AActor>>& InActors, bool bReset = false)
 	{
 		if (IsValid())
@@ -741,6 +781,7 @@ struct FGameplayEffectContextHandle
 		}
 	}
 
+	/** Add a hit result for targeting */
 	void AddHitResult(const FHitResult& InHitResult, bool bReset = false)
 	{
 		if (IsValid())
@@ -749,11 +790,13 @@ struct FGameplayEffectContextHandle
 		}
 	}
 
+	/** Returns actor list, may be empty */
 	const TArray<TWeakObjectPtr<AActor>> GetActors()
 	{
 		return Data->GetActors();
 	}
 
+	/** Returns hit result, this can be null */
 	const FHitResult* GetHitResult() const
 	{
 		if (IsValid())
@@ -763,6 +806,7 @@ struct FGameplayEffectContextHandle
 		return nullptr;
 	}
 
+	/** Adds an origin point */
 	void AddOrigin(FVector InOrigin)
 	{
 		if (IsValid())
@@ -771,6 +815,7 @@ struct FGameplayEffectContextHandle
 		}
 	}
 
+	/** Returns origin point, may be invalid if HasOrigin is false */
 	virtual const FVector& GetOrigin() const
 	{
 		if (IsValid())
@@ -780,6 +825,7 @@ struct FGameplayEffectContextHandle
 		return FVector::ZeroVector;
 	}
 
+	/** Returns true if GetOrigin will give valid information */
 	virtual bool HasOrigin() const
 	{
 		if (IsValid())
@@ -789,6 +835,7 @@ struct FGameplayEffectContextHandle
 		return false;
 	}
 
+	/** Returns debug string */
 	FString ToString() const
 	{
 		return IsValid() ? Data->ToString() : FString(TEXT("NONE"));
@@ -828,10 +875,10 @@ struct FGameplayEffectContextHandle
 		return !(FGameplayEffectContextHandle::operator==(Other));
 	}
 
+	/** Custom serializer, handles polymorphism of context */
 	bool NetSerialize(FArchive& Ar, class UPackageMap* Map, bool& bOutSuccess);
 
 private:
-
 	TSharedPtr<FGameplayEffectContext> Data;
 };
 
@@ -847,10 +894,7 @@ struct TStructOpsTypeTraits<FGameplayEffectContextHandle> : public TStructOpsTyp
 };
 
 
-/**
- * FGameplayEffectRemovalInfo
- *	Data struct for containing information pertinent to GameplayEffects as they are removed.
- */
+/** Data struct for containing information pertinent to GameplayEffects as they are removed */
 USTRUCT(BlueprintType)
 struct FGameplayEffectRemovalInfo
 {
@@ -868,9 +912,9 @@ struct FGameplayEffectRemovalInfo
 	UPROPERTY(EditAnywhere, BlueprintReadWrite, Category="Removal")
 	FGameplayEffectContextHandle EffectContext;
 };
-// -----------------------------------------------------------
 
 
+/** Metadata about a gameplay cue execution */
 USTRUCT(BlueprintType)
 struct GAMEPLAYABILITIES_API FGameplayCueParameters
 {
@@ -887,7 +931,6 @@ struct GAMEPLAYABILITIES_API FGameplayCueParameters
 
 	/** Projects can override this via UAbilitySystemGlobals */
 	FGameplayCueParameters(const struct FGameplayEffectSpecForRPC &Spec);
-
 	FGameplayCueParameters(const struct FGameplayEffectContextHandle& EffectContext);
 
 	/** Magnitude of source gameplay effect, normalzed from 0-1. Use this for "how strong is the gameplay effect" (0=min, 1=,max) */
@@ -918,9 +961,11 @@ struct GAMEPLAYABILITIES_API FGameplayCueParameters
 	UPROPERTY(BlueprintReadWrite, Category=GameplayCue)
 	FGameplayTagContainer AggregatedTargetTags;
 
+	/** Location cue took place at */
 	UPROPERTY(BlueprintReadWrite, Category=GameplayCue)
 	FVector_NetQuantize10 Location;
 
+	/** Normal of impact that caused cue */
 	UPROPERTY(BlueprintReadWrite, Category=GameplayCue)
 	FVector_NetQuantizeNormal Normal;
 
@@ -952,17 +997,22 @@ struct GAMEPLAYABILITIES_API FGameplayCueParameters
 	UPROPERTY(BlueprintReadWrite, Category = GameplayCue)
 	TWeakObjectPtr<USceneComponent> TargetAttachComponent;
 
+	/** Optimized serializer */
 	bool NetSerialize(FArchive& Ar, class UPackageMap* Map, bool& bOutSuccess);
 
-	bool IsInstigatorLocallyControlled() const;
+	/** Returns true if this is locally controlled, using fallback actor if nothing else available */
+	bool IsInstigatorLocallyControlled(AActor* FallbackActor = nullptr) const;
 
-	// Fallback actor is used if the parameters have nullptr for instigator and effect causer
+	/** Fallback actor is used if the parameters have nullptr for instigator and effect causer */
 	bool IsInstigatorLocallyControlledPlayer(AActor* FallbackActor=nullptr) const;
 
+	/** Returns the actor that instigated this originally, generally attached to an ability system component */
 	AActor* GetInstigator() const;
 
+	/** Returns the actor that physically caused the damage, could be a projectile or weapon */
 	AActor* GetEffectCauser() const;
 
+	/** Returns the object that originally caused this, game-specific but usually not an actor */
 	const UObject* GetSourceObject() const;
 };
 
@@ -975,15 +1025,24 @@ struct TStructOpsTypeTraits<FGameplayCueParameters> : public TStructOpsTypeTrait
 	};
 };
 
+
 UENUM(BlueprintType)
 namespace EGameplayCueEvent
 {
+	/** Indicates what type of action happend to a specific gameplay cue tag. Sometimes you will get multiple events at once */
 	enum Type
 	{
-		OnActive,		// Called when GameplayCue is activated.
-		WhileActive,	// Called when GameplayCue is active, even if it wasn't actually just applied (Join in progress, etc)
-		Executed,		// Called when a GameplayCue is executed: instant effects or periodic tick
-		Removed			// Called when GameplayCue is removed
+		/** Called when GameplayCue is activated */
+		OnActive,
+
+		/** Called when GameplayCue is active, even if it wasn't actually just applied (Join in progress, etc) */
+		WhileActive,
+
+		/** Called when a GameplayCue is executed: instant effects or periodic tick */
+		Executed,
+
+		/** Called when GameplayCue is removed */
+		Removed
 	};
 }
 
@@ -998,6 +1057,7 @@ DECLARE_MULTICAST_DELEGATE_ThreeParams(FOnActiveGameplayEffectStackChange, FActi
 DECLARE_MULTICAST_DELEGATE_ThreeParams(FOnActiveGameplayEffectTimeChange, FActiveGameplayEffectHandle, float /*NewStartTime*/, float /*NewDuration*/);
 DECLARE_MULTICAST_DELEGATE_TwoParams(FOnActiveGameplayEffectInhibitionChanged, FActiveGameplayEffectHandle, bool /*bIsInhibited*/);
 
+/** Callback struct for different types of gameplay effect changes */
 struct FActiveGameplayEffectEvents
 {
 	FOnActiveGameplayEffectRemoved DEPRECATED_OnEffectRemoved;
@@ -1010,8 +1070,16 @@ struct FActiveGameplayEffectEvents
 // This is deprecated, use FOnGameplayAttributeValueChange
 DECLARE_MULTICAST_DELEGATE_TwoParams(FOnGameplayAttributeChange, float, const FGameplayEffectModCallbackData*);
 
+
+/** Temporary parameter struct used when an attribute has changed */
 struct FOnAttributeChangeData
 {
+	FOnAttributeChangeData()
+		: NewValue(0.0f)
+		, OldValue(0.0f)
+		, GEModData(nullptr)
+	{ }
+
 	FGameplayAttribute Attribute;
 
 	float	NewValue;
@@ -1025,12 +1093,11 @@ DECLARE_DELEGATE_RetVal(FGameplayTagContainer, FGetGameplayTags);
 
 DECLARE_DELEGATE_RetVal_OneParam(FOnGameplayEffectTagCountChanged&, FRegisterGameplayTagChangeDelegate, FGameplayTag);
 
-// -----------------------------------------------------------
-
 
 UENUM(BlueprintType)
 namespace EGameplayTagEventType
 {
+	/** Rather a tag was added or removed, used in callbacks */
 	enum Type
 	{		
 		/** Event only happens when tag is new or completely removed */
@@ -1046,7 +1113,6 @@ namespace EGameplayTagEventType
  * while simultaneously tracking the count of parent tags as well. Events/delegates are fired whenever the tag counts
  * of any tag (explicit or parent) are modified.
  */
-
 struct GAMEPLAYABILITIES_API FGameplayTagCountContainer
 {	
 	FGameplayTagCountContainer()
@@ -1255,8 +1321,6 @@ private:
 };
 
 
-// -----------------------------------------------------------
-
 /** Encapsulate require and ignore tags */
 USTRUCT(BlueprintType)
 struct GAMEPLAYABILITIES_API FGameplayTagRequirements
@@ -1271,14 +1335,18 @@ struct GAMEPLAYABILITIES_API FGameplayTagRequirements
 	UPROPERTY(EditAnywhere, BlueprintReadWrite, Category = GameplayModifier)
 	FGameplayTagContainer IgnoreTags;
 
+	/** True if all required tags and no ignore tags found */
 	bool	RequirementsMet(const FGameplayTagContainer& Container) const;
-	bool	IsEmpty() const;
 
-	static FGetGameplayTags	SnapshotTags(FGetGameplayTags TagDelegate);
+	/** True if neither RequireTags or IgnoreTags has any tags */
+	bool	IsEmpty() const;
 
+	/** Return debug string */
 	FString ToString() const;
 };
 
+
+/** Structure used to combine tags from different sources during effect execution */
 USTRUCT()
 struct GAMEPLAYABILITIES_API FTagContainerAggregator
 {
@@ -1324,12 +1392,15 @@ struct GAMEPLAYABILITIES_API FTagContainerAggregator
 		return *this;
 	}
 
+	/** Returns tags from the source or target actor */
 	FGameplayTagContainer& GetActorTags();
 	const FGameplayTagContainer& GetActorTags() const;
 
+	/** Get tags that came from the effect spec */
 	FGameplayTagContainer& GetSpecTags();
 	const FGameplayTagContainer& GetSpecTags() const;
 
+	/** Returns combination of spec and actor tags */
 	const FGameplayTagContainer* GetAggregatedTags() const;
 
 private:
@@ -1357,10 +1428,9 @@ struct GAMEPLAYABILITIES_API FGameplayEffectSpecHandle
 	FGameplayEffectSpecHandle();
 	FGameplayEffectSpecHandle(FGameplayEffectSpec* DataPtr);
 
+	/** Internal pointer to effect spec */
 	TSharedPtr<FGameplayEffectSpec>	Data;
 
-	bool IsValidCache;
-
 	void Clear()
 	{
 		Data.Reset();
@@ -1404,9 +1474,8 @@ struct TStructOpsTypeTraits<FGameplayEffectSpecHandle> : public TStructOpsTypeTr
 	};
 };
 
-// -----------------------------------------------------------
-
 
+/** Map that stores count of tags, in a form that is optimized for replication */
 USTRUCT()
 struct GAMEPLAYABILITIES_API FMinimalReplicationTagCountMap
 {
diff --git a/Engine/Plugins/Runtime/GameplayAbilities/Source/GameplayAbilities/Public/TickableAttributeSetInterface.h b/Engine/Plugins/Runtime/GameplayAbilities/Source/GameplayAbilities/Public/TickableAttributeSetInterface.h
--- a/Engine/Plugins/Runtime/GameplayAbilities/Source/GameplayAbilities/Public/TickableAttributeSetInterface.h
+++ b/Engine/Plugins/Runtime/GameplayAbilities/Source/GameplayAbilities/Public/TickableAttributeSetInterface.h
@@ -7,7 +7,7 @@
 #include "UObject/Interface.h"
 #include "TickableAttributeSetInterface.generated.h"
 
-/** Interface for actors which can be "spotted" by a player */
+/** Interface for attribute sets that need to tick to update state. These can be expensive */
 UINTERFACE(MinimalAPI, meta = (CannotImplementInterfaceInBlueprint))
 class UTickableAttributeSetInterface : public UInterface
 {
diff --git a/Engine/Plugins/Runtime/MeshReconstruction/DummyMeshReconstructor/Source/DummyMeshReconstructor/Public/DummyMeshReconstructor.h b/Engine/Plugins/Runtime/MeshReconstruction/DummyMeshReconstructor/Source/DummyMeshReconstructor/Public/DummyMeshReconstructor.h
--- a/Engine/Plugins/Runtime/MeshReconstruction/DummyMeshReconstructor/Source/DummyMeshReconstructor/Public/DummyMeshReconstructor.h
+++ b/Engine/Plugins/Runtime/MeshReconstruction/DummyMeshReconstructor/Source/DummyMeshReconstructor/Public/DummyMeshReconstructor.h
@@ -15,19 +15,14 @@ class DUMMYMESHRECONSTRUCTOR_API UDummyMeshReconstructor : public UMeshReconstru
 
 public:
 	//~ UMeshReconstructorBase
-	UFUNCTION(BlueprintCallable, Category="Mesh Reconstruction")
 	virtual void StartReconstruction() override;
 	
-	UFUNCTION(BlueprintCallable, Category="Mesh Reconstruction")
 	virtual void StopReconstruction() override;
 	
-	UFUNCTION(BlueprintCallable, Category="Mesh Reconstruction")
 	virtual void PauseReconstruction() override;
 	
-	UFUNCTION(BlueprintCallable, Category="Mesh Reconstruction")
 	virtual bool IsReconstructionStarted() const override;
 	
-	UFUNCTION(BlueprintCallable, Category="Mesh Reconstruction")
 	virtual bool IsReconstructionPaused() const override;
 
 	virtual FMRMeshConfiguration ConnectMRMesh(UMRMeshComponent* Mesh) override;
diff --git a/Engine/Plugins/Runtime/MixedRealityFramework/Source/MixedRealityFramework/Public/MixedRealityCaptureDevice.h b/Engine/Plugins/Runtime/MixedRealityFramework/Source/MixedRealityFramework/Public/MixedRealityCaptureDevice.h
--- a/Engine/Plugins/Runtime/MixedRealityFramework/Source/MixedRealityFramework/Public/MixedRealityCaptureDevice.h
+++ b/Engine/Plugins/Runtime/MixedRealityFramework/Source/MixedRealityFramework/Public/MixedRealityCaptureDevice.h
@@ -113,7 +113,6 @@ public:
 	static UAsyncTask_OpenMRCaptureDevice* OpenMRCaptureDevice(const FMediaCaptureDeviceInfo& DeviceId, UMediaPlayer* Target, FMRCaptureFeedDelegate::FDelegate OpenedCallback);
 
 private:
-	UFUNCTION()
 	virtual void OnVideoFeedOpened(FString DeviceUrl) override;
 };
 
@@ -132,7 +131,6 @@ public:
 	static UAsyncTask_OpenMRCaptureFeed* OpenMRCaptureFeed(const FMRCaptureDeviceIndex& FeedRef, UMediaPlayer* Target, FMRCaptureFeedDelegate::FDelegate OpenedCallback);
 
 private:
-	UFUNCTION()
 	virtual void OnVideoFeedOpened(FString DeviceUrl) override;
 
 	FMRCaptureDeviceIndex DesiredFeedRef;
diff --git a/Engine/Source/Editor/AnimGraph/Private/AnimGraphNode_SubInstance.cpp b/Engine/Source/Editor/AnimGraph/Private/AnimGraphNode_SubInstance.cpp
--- a/Engine/Source/Editor/AnimGraph/Private/AnimGraphNode_SubInstance.cpp
+++ b/Engine/Source/Editor/AnimGraph/Private/AnimGraphNode_SubInstance.cpp
@@ -199,18 +199,7 @@ void UAnimGraphNode_SubInstance::ReallocatePinsDuringReconstruction(TArray<UEdGr
 	for(FName& RemovedPropertyName : BeginExposableNames)
 	{
 		KnownExposableProperties.Remove(RemovedPropertyName);
-		ExposedPropertyNames.Remove(RemovedPropertyName);
 	}
-
-	// Make sure that any old pins that linked to properties are told not to be orphans
-	for(UEdGraphPin* OldPin : OldPins)
-	{
-		if(OldPin && !AnimGraphDefaultSchema->IsPosePin(OldPin->PinType))
-		{
-			OldPin->SetSavePinIfOrphaned(false);
-		}
-	}
-
 }
 
 void UAnimGraphNode_SubInstance::PostEditChangeProperty(FPropertyChangedEvent& PropertyChangedEvent)
diff --git a/Engine/Source/Editor/BehaviorTreeEditor/Private/BehaviorTreeGraphNode_Composite.cpp b/Engine/Source/Editor/BehaviorTreeEditor/Private/BehaviorTreeGraphNode_Composite.cpp
--- a/Engine/Source/Editor/BehaviorTreeEditor/Private/BehaviorTreeGraphNode_Composite.cpp
+++ b/Engine/Source/Editor/BehaviorTreeEditor/Private/BehaviorTreeGraphNode_Composite.cpp
@@ -42,7 +42,7 @@ FText UBehaviorTreeGraphNode_Composite::GetTooltipText() const
 			NSLOCTEXT("BehaviorTreeEditor", "CompositeNodeScopeTooltip", "This node is a local scope for decorators.\nAll observing decorators (Lower Priority or Both) will be removed when execution flow leaves this branch."));
 	}
 
-	return Super::GetDescription();
+	return Super::GetTooltipText();
 }
 
 void UBehaviorTreeGraphNode_Composite::GetContextMenuActions(const FGraphNodeContextMenuBuilder& Context) const
diff --git a/Engine/Source/Editor/BlueprintGraph/Classes/K2Node_CustomEvent.h b/Engine/Source/Editor/BlueprintGraph/Classes/K2Node_CustomEvent.h
--- a/Engine/Source/Editor/BlueprintGraph/Classes/K2Node_CustomEvent.h
+++ b/Engine/Source/Editor/BlueprintGraph/Classes/K2Node_CustomEvent.h
@@ -50,6 +50,7 @@ class UK2Node_CustomEvent : public UK2Node_Event
 	//~ End UK2Node Interface
 
 	//~ Begin UK2Node_EditablePinBase Interface
+	virtual bool ShouldUseConstRefParams() const override { return true; }
 	virtual UEdGraphPin* CreatePinFromUserDefinition(const TSharedPtr<FUserPinInfo> NewPinInfo) override;
 	virtual bool CanCreateUserDefinedPin(const FEdGraphPinType& InPinType, EEdGraphPinDirection InDesiredDirection, FText& OutErrorMessage) override;
 	virtual bool ModifyUserDefinedPinDefaultValue(TSharedPtr<FUserPinInfo> PinInfo, const FString& NewDefaultValue) override;
diff --git a/Engine/Source/Editor/BlueprintGraph/Classes/K2Node_EditablePinBase.h b/Engine/Source/Editor/BlueprintGraph/Classes/K2Node_EditablePinBase.h
--- a/Engine/Source/Editor/BlueprintGraph/Classes/K2Node_EditablePinBase.h
+++ b/Engine/Source/Editor/BlueprintGraph/Classes/K2Node_EditablePinBase.h
@@ -161,5 +161,10 @@ class UK2Node_EditablePinBase : public UK2Node
 	 * Can this node have pass-by-reference parameters?
 	 */
 	virtual bool CanUseRefParams() const { return false; }
+
+	/**
+	 * Should this node require 'const' for pass-by-reference parameters?
+	 */
+	virtual bool ShouldUseConstRefParams() const { return false;  }
 };
 
diff --git a/Engine/Source/Editor/BlueprintGraph/Classes/K2Node_FunctionEntry.h b/Engine/Source/Editor/BlueprintGraph/Classes/K2Node_FunctionEntry.h
--- a/Engine/Source/Editor/BlueprintGraph/Classes/K2Node_FunctionEntry.h
+++ b/Engine/Source/Editor/BlueprintGraph/Classes/K2Node_FunctionEntry.h
@@ -55,6 +55,7 @@ class UK2Node_FunctionEntry : public UK2Node_FunctionTerminator
 
 	//~ Begin UK2Node_EditablePinBase Interface
 	virtual bool CanUseRefParams() const override { return true; }
+	virtual bool ShouldUseConstRefParams() const override;
 	virtual bool ModifyUserDefinedPinDefaultValue(TSharedPtr<FUserPinInfo> PinInfo, const FString& NewDefaultValue) override;
 	//~ End UK2Node_EditablePinBase Interface
 
diff --git a/Engine/Source/Editor/BlueprintGraph/Private/K2Node_CustomEvent.cpp b/Engine/Source/Editor/BlueprintGraph/Private/K2Node_CustomEvent.cpp
--- a/Engine/Source/Editor/BlueprintGraph/Private/K2Node_CustomEvent.cpp
+++ b/Engine/Source/Editor/BlueprintGraph/Private/K2Node_CustomEvent.cpp
@@ -117,40 +117,6 @@ void UK2Node_CustomEvent::Serialize(FArchive& Ar)
 	if (Ar.IsLoading())
 	{
 		CachedNodeTitle.MarkDirty();
-
-		const UEdGraphSchema_K2* K2Schema = GetDefault<UEdGraphSchema_K2>();
-
-		// Ensure that array type inputs and non-array type pass-by-reference inputs are also marked 'const' - since arrays
-		// are implicitly passed by reference, and since events do not have return values/outputs, this equates to marking
-		// the parameter as 'const Type&' in native code. Also note that since UHT already blocks non-const reference types
-		// from being compiled into a MC delegate signature, any existing custom event param pins that were implicitly
-		// created via "Assign" in the Blueprint editor's context menu will previously have had 'const' set for its pin type.
-		//
-		// This ensures that (a) we don't emit the "no reference will be returned" note on custom event nodes with array
-		// inputs added by the user via the Details panel, and (b) we don't emit the "no reference will be returned" warning
-		// on custom event nodes with struct/object inputs added by the user in the Details tab and also set to pass-by-reference.
-		//
-		// Blueprint details customization now sets 'bIsConst' in new custom event node placements - see OnRefCheckStateChanged().
-		for (auto Pin : Pins)
-		{
-			if (Pin)
-			{
-				if (Pin->Direction == EGPD_Output
-					&& !Pin->PinType.bIsConst
-					&& !K2Schema->IsExecPin(*Pin)
-					&& !K2Schema->IsDelegateCategory(Pin->PinType.PinCategory))
-				{
-					for (TSharedPtr<FUserPinInfo>& PinInfo : UserDefinedPins)
-					{
-						if (PinInfo->PinName == Pin->PinName)
-						{
-							Pin->PinType.bIsConst = PinInfo->PinType.bIsConst = (PinInfo->PinType.IsContainer() || PinInfo->PinType.bIsReference);
-							break;
-						}
-					}
-				}
-			}
-		}
 	}
 }
 
diff --git a/Engine/Source/Editor/BlueprintGraph/Private/K2Node_EditablePinBase.cpp b/Engine/Source/Editor/BlueprintGraph/Private/K2Node_EditablePinBase.cpp
--- a/Engine/Source/Editor/BlueprintGraph/Private/K2Node_EditablePinBase.cpp
+++ b/Engine/Source/Editor/BlueprintGraph/Private/K2Node_EditablePinBase.cpp
@@ -7,6 +7,14 @@
 #include "EdGraphSchema_K2.h"
 #include "Kismet2/KismetDebugUtilities.h"
 
+// Ensure that the UserDefinedPin's "desired direction" matches the direction of 
+// the EdGraphPin that it corresponds to. Somehow it is possible for these to get 
+// out of sync, and we're not entirely sure how/why.
+//
+// @TODO: Determine how these get out of sync and fix that up so we can guard it 
+//        with a version check and not have to do this anymore for updated assets
+#define ALWAYS_VALIDATE_DESIRED_PIN_DIRECTION_ON_LOAD 1
+
 FArchive& operator<<(FArchive& Ar, FUserPinInfo& Info)
 {
 	Ar.UsingCustomVersion(FFrameworkObjectVersion::GUID);
@@ -166,29 +174,85 @@ void UK2Node_EditablePinBase::Serialize(FArchive& Ar)
 {
 	Super::Serialize(Ar);
 
+	Ar.UsingCustomVersion(FFrameworkObjectVersion::GUID);
+
 	TArray<FUserPinInfo> SerializedItems;
 	if (Ar.IsLoading())
 	{
 		Ar << SerializedItems;
 
+		const UEdGraphSchema_K2* K2Schema = GetDefault<UEdGraphSchema_K2>();
+
 		UserDefinedPins.Empty(SerializedItems.Num());
 		for (int32 Index = 0; Index < SerializedItems.Num(); ++Index)
 		{
 			TSharedPtr<FUserPinInfo> PinInfo = MakeShareable(new FUserPinInfo(SerializedItems[Index]));
 
-			// Ensure that the UserDefinedPin's "desired direction" matches the direction of 
-			// the EdGraphPin that it corresponds to. Somehow it is possible for these to get 
-			// out of sync, and we're not entirely sure how/why.
+#if !ALWAYS_VALIDATE_DESIRED_PIN_DIRECTION_ON_LOAD
+			// @TODO - Replace with a version check here if/when we are able to avoid this validation step on updated assets.
+			const bool bValidateDesiredPinDirection = true;
+#endif
+			// Ensure that array type inputs and non-array type pass-by-reference inputs are also marked 'const' for both custom
+			// event signatures and interface functions that have no return value. Since arrays are implicitly passed by reference,
+			// and since events do not have return values/outputs, this equates to marking the parameter as 'const Type&' in native.
+			// code. Also note that since UHT already blocks non-const reference types from being compiled into a MC delegate signature,
+			// any existing custom event param pins that were implicitly created via "Assign" in the Blueprint editor's context menu
+			// will previously have had 'const' set for its pin type.
 			//
-			// @TODO: Determine how these get out of sync and fix that up so we can guard this 
-			//        with a version check, and not have to do this for updated assets
-			if (UEdGraphPin* NodePin = FindPin(PinInfo->PinName))
+			// This ensures that (a) we don't emit the "no reference will be returned" note on custom event and implemented interface
+			// event nodes with array inputs added by the user via the Details panel, and (b) we don't emit the "no reference will be
+			// returned" warning on custom event and implemented interface event nodes with struct/object inputs added by the user in
+			// the Details tab that are also explicitly set to pass-by-reference. That message is intended to convey one should not
+			// expect the input to also be treated like an output - if the value is modified inside the event, it won't be reflected
+			// back out to the caller through the reference. The message should only be seen for implemented event signatures with I/O
+			// parameters that are explicitly passed by reference and declared in native C++ code using 'UPARAM(ref)' markup instead
+			// of 'const Type&' - i.e. the 'const' form should be used to declare input ref args in native events with no return value.
+			//
+			// However, on the Blueprint side for new custom event and new interface event signature input arguments, we do not currently
+			// expose a 'const' qualifier for pass-by-reference parameters, so we're implicitly adding one here to older input ref pin
+			// types to be consistent with the native side. The Blueprint compiler currently ignores 'const' in terms of whether or not
+			// the referenced value or object is actually treated as read-only in the event's implementation, but this may change later.
+			//
+			// Note that Blueprint details customization will set 'const' for all new custom event and implemented interface event node
+			// placements with an array type or pass-by-reference input param. See FBlueprintGraphArgumentLayout::PinInfoChanged().
+			//
+			const bool bValidateConstRefPinTypes = Ar.CustomVer(FFrameworkObjectVersion::GUID) < FFrameworkObjectVersion::EditableEventsUseConstRefParameters
+				&& ShouldUseConstRefParams();
+
+			// Avoid the FindPin() call if we don't need to do it.
+#if !ALWAYS_VALIDATE_DESIRED_PIN_DIRECTION_ON_LOAD
+			if (bValidateDesiredPinDirection || bValidateConstRefPinTypes)
+#endif
 			{
-				// NOTE: the second FindPin call here to keep us from altering a pin with the same 
-				//       name but different direction (in case there is two)
-				if (PinInfo->DesiredPinDirection != NodePin->Direction && FindPin(PinInfo->PinName, PinInfo->DesiredPinDirection) == nullptr)
+				if (UEdGraphPin* NodePin = FindPin(PinInfo->PinName))
 				{
-					PinInfo->DesiredPinDirection = NodePin->Direction;
+#if !ALWAYS_VALIDATE_DESIRED_PIN_DIRECTION_ON_LOAD
+					if (bValidateDesiredPinDirection)
+#endif
+					{
+						// NOTE: the second FindPin call here to keep us from altering a pin with the same 
+						//       name but different direction (in case there is two)
+						if (PinInfo->DesiredPinDirection != NodePin->Direction && FindPin(PinInfo->PinName, PinInfo->DesiredPinDirection) == nullptr)
+						{
+							PinInfo->DesiredPinDirection = NodePin->Direction;
+						}
+					}
+
+					if (bValidateConstRefPinTypes)
+					{
+						// Note that we should only get here if ShouldUseConstRefParams() indicated this node represents an event function with no outputs (above).
+						if (!NodePin->PinType.bIsConst
+							&& NodePin->Direction == EGPD_Output
+							&& !K2Schema->IsExecPin(*NodePin)
+							&& !K2Schema->IsDelegateCategory(NodePin->PinType.PinCategory))
+						{
+							// Add 'const' to either an array pin type (always passed by reference) or a pin type that's explicitly flagged to be passed by reference.
+							NodePin->PinType.bIsConst = NodePin->PinType.IsArray() || NodePin->PinType.bIsReference;
+
+							// Also mirror the flag into the UserDefinedPins array.
+							PinInfo->PinType.bIsConst = NodePin->PinType.bIsConst;
+						}
+					}
 				}
 			}
 
diff --git a/Engine/Source/Editor/BlueprintGraph/Private/K2Node_FunctionEntry.cpp b/Engine/Source/Editor/BlueprintGraph/Private/K2Node_FunctionEntry.cpp
--- a/Engine/Source/Editor/BlueprintGraph/Private/K2Node_FunctionEntry.cpp
+++ b/Engine/Source/Editor/BlueprintGraph/Private/K2Node_FunctionEntry.cpp
@@ -12,15 +12,18 @@
 #include "EdGraph/EdGraphSchema.h"
 #include "EdGraphSchema_K2.h"
 #include "K2Node_CallFunction.h"
+#include "K2Node_FunctionResult.h"
 #include "K2Node_MakeArray.h"
 #include "K2Node_MakeVariable.h"
 #include "K2Node_VariableSet.h"
 #include "Kismet2/BlueprintEditorUtils.h"
+#include "Kismet2/KismetEditorUtilities.h"
 #include "EdGraphUtilities.h"
 #include "BPTerminal.h"
 #include "UObject/PropertyPortFlags.h"
 #include "KismetCompilerMisc.h"
 #include "KismetCompiler.h"
+#include "Misc/OutputDeviceNull.h"
 
 #define LOCTEXT_NAMESPACE "K2Node_FunctionEntry"
 
@@ -238,7 +241,42 @@ void UK2Node_FunctionEntry::Serialize(FArchive& Ar)
 
 	Ar.UsingCustomVersion(FBlueprintsObjectVersion::GUID);
 
-	if (Ar.IsLoading())
+	if (Ar.IsSaving())
+	{
+		for (FBPVariableDescription& LocalVariable : LocalVariables)
+		{
+			if (!LocalVariable.DefaultValue.IsEmpty())
+			{
+				// If looking for references during save, expand any default values on the local variables
+				if (Ar.IsObjectReferenceCollector() && LocalVariable.VarType.PinCategory == UEdGraphSchema_K2::PC_Struct && LocalVariable.VarType.PinSubCategoryObject.IsValid())
+				{
+					UScriptStruct* Struct = Cast<UScriptStruct>(LocalVariable.VarType.PinSubCategoryObject.Get());
+
+					if (Struct)
+					{
+						TSharedPtr<FStructOnScope> StructData = MakeShareable(new FStructOnScope(Struct));
+
+						// Import the literal text to a dummy struct and then serialize that. Hard object references will not properly import, this is only useful for soft references!
+						FOutputDeviceNull NullOutput;
+						Struct->ImportText(*LocalVariable.DefaultValue, StructData->GetStructMemory(), nullptr, PPF_SerializedAsImportText, &NullOutput, LocalVariable.VarName.ToString());
+						Struct->SerializeItem(Ar, StructData->GetStructMemory(), nullptr);
+					}
+				}
+
+				if (LocalVariable.VarType.PinCategory == UEdGraphSchema_K2::PC_SoftObject || LocalVariable.VarType.PinCategory == UEdGraphSchema_K2::PC_SoftClass)
+				{
+					FSoftObjectPath TempRef(LocalVariable.DefaultValue);
+
+					// Serialize the asset reference, this will do the save fixup. It won't actually serialize the string if this is a real archive like linkersave
+					FSoftObjectPathSerializationScope DisableSerialize(NAME_None, NAME_None, ESoftObjectPathCollectType::AlwaysCollect, ESoftObjectPathSerializeType::SkipSerializeIfArchiveHasSize);
+					Ar << TempRef;
+
+					LocalVariable.DefaultValue = TempRef.ToString();
+				}
+			}
+		}
+	}
+	else if (Ar.IsLoading())
 	{
 		if (Ar.CustomVer(FFrameworkObjectVersion::GUID) < FFrameworkObjectVersion::LocalVariablesBlueprintVisible)
 		{
@@ -584,5 +622,36 @@ bool UK2Node_FunctionEntry::ModifyUserDefinedPinDefaultValue(TSharedPtr<FUserPin
 	return false;
 }
 
+bool UK2Node_FunctionEntry::ShouldUseConstRefParams() const
+{
+	// Interface functions with no outputs will be implemented as events. As with native interface functions with no outputs, the entry
+	// node is expected to use 'const Type&' for input parameters that are passed by reference. See UEditablePinBase::PostLoad() for details.
+	if (const UEdGraph* OwningGraph = GetGraph())
+	{
+		const UBlueprint* OwningBlueprint = FBlueprintEditorUtils::FindBlueprintForGraph(OwningGraph);
+		if (OwningBlueprint && OwningBlueprint->BlueprintType == BPTYPE_Interface)
+		{
+			// Find paired result node and check for outputs.
+			for (UEdGraphNode* Node : OwningGraph->Nodes)
+			{
+				if (UK2Node_FunctionResult* ResultNode = Cast<UK2Node_FunctionResult>(Node))
+				{
+					// This might be called from the super's Serialize() method for older assets, so make sure the result node's pins have been loaded.
+					if (ResultNode->HasAnyFlags(RF_NeedLoad))
+					{
+						GetLinker()->Preload(ResultNode);
+					}
+
+					return ResultNode->UserDefinedPins.Num() == 0;
+				}
+			}
+
+			// No result node, so there are no outputs.
+			return true;
+		}
+	}
+	
+	return false;
+}
 
 #undef LOCTEXT_NAMESPACE
diff --git a/Engine/Source/Editor/BlueprintGraph/Private/K2Node_Message.cpp b/Engine/Source/Editor/BlueprintGraph/Private/K2Node_Message.cpp
--- a/Engine/Source/Editor/BlueprintGraph/Private/K2Node_Message.cpp
+++ b/Engine/Source/Editor/BlueprintGraph/Private/K2Node_Message.cpp
@@ -34,11 +34,9 @@ FText UK2Node_Message::GetNodeTitle(ENodeTitleType::Type TitleType) const
 			}
 			else
 			{
-				FFormatNamedArguments Args;
-				Args.Add(TEXT("NodeName"), NodeNameText);
-				Args.Add(TEXT("OuterClassName"), FText::FromString(Function->GetOuterUClass()->GetName()));
+				// For consistency, use the function call node title format
+				FText NodeTitle = Super::GetNodeTitle(TitleType);
 
-				FText NodeTitle = FText::Format(NSLOCTEXT("K2Node", "CallInterfaceContext", "{NodeName}\nUsing Interface {OuterClassName}"), Args);
 				// FText::Format() is slow, so we cache this to save on performance
 				CachedNodeTitles.SetCachedTitle(TitleType, NodeTitle, this);
 			}
diff --git a/Engine/Source/Editor/BlueprintGraph/Private/K2Node_Variable.cpp b/Engine/Source/Editor/BlueprintGraph/Private/K2Node_Variable.cpp
--- a/Engine/Source/Editor/BlueprintGraph/Private/K2Node_Variable.cpp
+++ b/Engine/Source/Editor/BlueprintGraph/Private/K2Node_Variable.cpp
@@ -97,23 +97,7 @@ bool UK2Node_Variable::CreatePinForVariable(EEdGraphPinDirection Direction, FNam
 	}
 	else
 	{
-		if (!VariableReference.IsLocalScope())
-		{
-			UBlueprint* OwnerBP = GetBlueprint();
-
-			FString AssetName = OwnerBP ? OwnerBP->GetPathName() : TEXT("Unknown BP");
-
-			FString WarningMsg = FString::Printf(TEXT("'%s' variable not found. Base class was probably changed. %s"), *GetVarNameString(), *AssetName);
-
-			if (OwnerBP && OwnerBP->CurrentMessageLog)
-			{
-				OwnerBP->CurrentMessageLog->Warning(*FString::Printf(TEXT("@@: %s"), *WarningMsg), this);
-			}
-			else
-			{
-				Message_Warn(*WarningMsg);
-			}			
-		}
+		// Don't handle warning or error logging here, that needs to be done in ValidateNodeDuringCompilation
 		return false;
 	}
 
diff --git a/Engine/Source/Editor/Kismet/Private/BlueprintDebugger.cpp b/Engine/Source/Editor/Kismet/Private/BlueprintDebugger.cpp
new file mode 100644
--- /dev/null
+++ b/Engine/Source/Editor/Kismet/Private/BlueprintDebugger.cpp
@@ -0,0 +1,326 @@
+// Copyright 1998-2018 Epic Games, Inc. All Rights Reserved.
+
+#include "BlueprintDebugger.h"
+
+#include "CallStackViewer.h"
+#include "Debugging/SKismetDebuggingView.h"
+#include "EditorStyleSet.h"
+#include "Framework/Application/SlateApplication.h"
+#include "Framework/Commands/Commands.h"
+#include "Framework/Docking/TabManager.h"
+#include "Framework/MultiBox/MultiBoxBuilder.h"
+#include "Layout/WidgetPath.h"
+#include "WatchPointViewer.h"
+#include "Widgets/Docking/SDockTab.h"
+#include "WorkspaceMenuStructure.h"
+#include "WorkspaceMenuStructureModule.h"
+
+#define LOCTEXT_NAMESPACE "BlueprintDebugger"
+
+static bool CanCloseBlueprintDebugger()
+{
+	return !GIntraFrameDebuggingGameThread;
+}
+
+struct FBlueprintDebuggerCommands : public TCommands<FBlueprintDebuggerCommands>
+{
+	FBlueprintDebuggerCommands()
+		: TCommands<FBlueprintDebuggerCommands>(
+			TEXT("BlueprintDebugger"), // Context name for fast lookup
+			LOCTEXT("BlueprintDebugger", "Blueprint Debugger"), // Localized context name for displaying
+			NAME_None, // Parent
+			FCoreStyle::Get().GetStyleSetName() // Icon Style Set
+		)
+	{
+	}
+
+	// TCommand<> interface
+	virtual void RegisterCommands() override;
+	// End of TCommand<> interface
+
+	TSharedPtr<FUICommandInfo> ShowCallStackViewer;
+	TSharedPtr<FUICommandInfo> ShowWatchViewer;
+	TSharedPtr<FUICommandInfo> ShowExecutionTrace;
+};
+
+void FBlueprintDebuggerCommands::RegisterCommands()
+{
+	UI_COMMAND(ShowCallStackViewer, "Call Stack", "Toggles visibility of the Call Stack window", EUserInterfaceActionType::Check, FInputChord());
+	UI_COMMAND(ShowWatchViewer, "Watches", "Toggles visibility of the Watches window", EUserInterfaceActionType::Check, FInputChord());
+	UI_COMMAND(ShowExecutionTrace, "Execution Trace", "Toggles visibility of the Execution Trace window", EUserInterfaceActionType::Check, FInputChord());
+}
+
+struct FBlueprintDebuggerImpl
+{
+	FBlueprintDebuggerImpl();
+	~FBlueprintDebuggerImpl();
+
+	/** Function registered with tab manager to create the bluepring debugger */
+	TSharedRef<SDockTab> CreateBluprintDebuggerTab(const FSpawnTabArgs& Args);
+
+	TSharedPtr<FTabManager> DebuggingToolsTabManager;
+	TSharedPtr<FTabManager::FLayout> BlueprintDebuggerLayout;
+
+private:
+	// prevent copying:
+	FBlueprintDebuggerImpl(const FBlueprintDebuggerImpl&);
+	FBlueprintDebuggerImpl(FBlueprintDebuggerImpl&&);
+	FBlueprintDebuggerImpl& operator=(FBlueprintDebuggerImpl const&);
+	FBlueprintDebuggerImpl& operator=(FBlueprintDebuggerImpl&&);
+};
+
+const FName DebuggerAppName = FName(TEXT("DebuggerApp"));
+
+FBlueprintDebuggerImpl::FBlueprintDebuggerImpl()
+{
+	const IWorkspaceMenuStructure& MenuStructure = WorkspaceMenu::GetMenuStructure();
+
+	FBlueprintDebuggerCommands::Register();
+	FGlobalTabmanager::Get()->RegisterNomadTabSpawner(DebuggerAppName, FOnSpawnTab::CreateRaw(this, &FBlueprintDebuggerImpl::CreateBluprintDebuggerTab))
+		.SetDisplayName(NSLOCTEXT("BlueprintDebugger", "TabTitle", "Blueprint Debugger"))
+		.SetTooltipText(NSLOCTEXT("BlueprintDebugger", "TooltipText", "Open the Blueprint Debugger tab."))
+		.SetGroup(MenuStructure.GetDeveloperToolsDebugCategory())
+		.SetIcon(FSlateIcon(FEditorStyle::GetStyleSetName(), "BlueprintDebugger.TabIcon"));
+}
+
+FBlueprintDebuggerImpl::~FBlueprintDebuggerImpl()
+{
+	if (FSlateApplication::IsInitialized())
+	{
+		FGlobalTabmanager::Get()->UnregisterNomadTabSpawner(DebuggerAppName);
+	}
+
+	const IWorkspaceMenuStructure& MenuStructure = WorkspaceMenu::GetMenuStructure();
+
+	if (DebuggingToolsTabManager.IsValid())
+	{
+		FGlobalTabmanager::Get()->UnregisterTabSpawner(DebuggerAppName);
+		BlueprintDebuggerLayout = TSharedPtr<FTabManager::FLayout>();
+		DebuggingToolsTabManager = TSharedPtr<FTabManager>();
+	}
+	FBlueprintDebuggerCommands::Unregister();
+}
+
+TSharedRef<SDockTab> FBlueprintDebuggerImpl::CreateBluprintDebuggerTab(const FSpawnTabArgs& Args)
+{
+	const TSharedRef<SDockTab> NomadTab = SNew(SDockTab)
+		.TabRole(ETabRole::NomadTab)
+		.OnCanCloseTab(SDockTab::FCanCloseTab::CreateStatic(&CanCloseBlueprintDebugger))
+		.Label(NSLOCTEXT("BlueprintDebugger", "TabTitle", "Blueprint Debugger"));
+
+	if (!DebuggingToolsTabManager.IsValid())
+	{
+		DebuggingToolsTabManager = FGlobalTabmanager::Get()->NewTabManager(NomadTab);
+		DebuggingToolsTabManager->SetOnPersistLayout(
+			FTabManager::FOnPersistLayout::CreateStatic(
+				[](const TSharedRef<FTabManager::FLayout>& InLayout)
+				{
+					if (InLayout->GetPrimaryArea().Pin().IsValid())
+					{
+						FLayoutSaveRestore::SaveToConfig(GEditorLayoutIni, InLayout);
+					}
+				}
+			)
+		);
+	}
+	else
+	{
+		ensure(BlueprintDebuggerLayout.IsValid());
+	}
+
+	const FName ExecutionFlowTabName = FName(TEXT("ExecutionFlowApp"));
+	const FName CallStackTabName = CallStackViewer::GetTabName();
+	const FName WatchViewerTabName = WatchViewer::GetTabName();
+
+	if (!BlueprintDebuggerLayout.IsValid())
+	{
+		DebuggingToolsTabManager->RegisterTabSpawner(
+			ExecutionFlowTabName,
+			FOnSpawnTab::CreateStatic(
+				[](const FSpawnTabArgs&)->TSharedRef<SDockTab>
+				{
+				return SNew(SDockTab)
+					.TabRole(ETabRole::PanelTab)
+					.OnCanCloseTab(SDockTab::FCanCloseTab::CreateStatic(&CanCloseBlueprintDebugger))
+					.Label(NSLOCTEXT("BlueprintExecutionFlow", "TabTitle", "Execution Flow"))
+					[
+						SNew(SKismetDebuggingView)
+					];
+				}
+			)
+		)
+		.SetDisplayName(NSLOCTEXT("BlueprintDebugger", "ExecutionFlowTabTitle", "Blueprint Execution Flow"))
+		.SetTooltipText(NSLOCTEXT("BlueprintDebugger", "ExecutionFlowTooltipText", "Open the Blueprint Execution Flow tab."));
+
+		CallStackViewer::RegisterTabSpawner(*DebuggingToolsTabManager);
+		WatchViewer::RegisterTabSpawner(*DebuggingToolsTabManager);
+
+		BlueprintDebuggerLayout = FTabManager::NewLayout("Standalone_BlueprintDebugger_Layout_v1")
+		->AddArea
+		(
+			FTabManager::NewPrimaryArea()
+			->SetOrientation(Orient_Vertical)
+			->Split
+			(
+				FTabManager::NewStack()
+				->SetSizeCoefficient(.4f)
+				->SetHideTabWell(true)
+				->AddTab(CallStackTabName, ETabState::OpenedTab)
+				->AddTab(WatchViewerTabName, ETabState::OpenedTab)
+				->AddTab(ExecutionFlowTabName, ETabState::OpenedTab)
+				->SetForegroundTab(CallStackTabName)
+			)
+		);
+	}
+
+	BlueprintDebuggerLayout = FLayoutSaveRestore::LoadFromConfig(GEditorLayoutIni, BlueprintDebuggerLayout.ToSharedRef());
+
+	TSharedRef<SWidget> TabContents = DebuggingToolsTabManager->RestoreFrom(BlueprintDebuggerLayout.ToSharedRef(), TSharedPtr<SWindow>()).ToSharedRef();
+
+	// build command list for tab restoration menu:
+	TSharedPtr<FUICommandList> CommandList = MakeShareable(new FUICommandList());
+
+	TWeakPtr<FTabManager> DebuggingToolsManagerWeak = DebuggingToolsTabManager;
+
+	const auto ToggleTabVisibility = [](TWeakPtr<FTabManager> InDebuggingToolsManagerWeak, FName InTabName)
+	{
+		TSharedPtr<FTabManager> InDebuggingToolsManager = InDebuggingToolsManagerWeak.Pin();
+		if (InDebuggingToolsManager.IsValid())
+		{
+			TSharedPtr<SDockTab> ExistingTab = InDebuggingToolsManager->FindExistingLiveTab(InTabName);
+			if (ExistingTab.IsValid())
+			{
+				ExistingTab->RequestCloseTab();
+			}
+			else
+			{
+				InDebuggingToolsManager->InvokeTab(InTabName);
+			}
+		}
+	};
+
+	const auto IsTabVisible = [](TWeakPtr<FTabManager> InDebuggingToolsManagerWeak, FName InTabName)
+	{
+		TSharedPtr<FTabManager> InDebuggingToolsManager = InDebuggingToolsManagerWeak.Pin();
+		if (InDebuggingToolsManager.IsValid())
+		{
+			return InDebuggingToolsManager->FindExistingLiveTab(InTabName).IsValid();
+		}
+		return false;
+	};
+
+	CommandList->MapAction(
+		FBlueprintDebuggerCommands::Get().ShowCallStackViewer,
+		FExecuteAction::CreateStatic(
+			ToggleTabVisibility,
+			DebuggingToolsManagerWeak,
+			CallStackTabName
+		),
+		FCanExecuteAction::CreateStatic(
+			[]() { return true; }
+		),
+		FIsActionChecked::CreateStatic(
+			IsTabVisible,
+			DebuggingToolsManagerWeak,
+			CallStackTabName
+		)
+	);
+
+	CommandList->MapAction(
+		FBlueprintDebuggerCommands::Get().ShowWatchViewer,
+		FExecuteAction::CreateStatic(
+			ToggleTabVisibility,
+			DebuggingToolsManagerWeak,
+			WatchViewerTabName
+		),
+		FCanExecuteAction::CreateStatic(
+			[]() { return true; }
+		),
+		FIsActionChecked::CreateStatic(
+			IsTabVisible,
+			DebuggingToolsManagerWeak,
+			WatchViewerTabName
+		)
+	);
+
+	CommandList->MapAction(
+		FBlueprintDebuggerCommands::Get().ShowExecutionTrace,
+		FExecuteAction::CreateStatic(
+			ToggleTabVisibility,
+			DebuggingToolsManagerWeak,
+			ExecutionFlowTabName
+		),
+		FCanExecuteAction::CreateStatic(
+			[]() { return true; }
+		),
+		FIsActionChecked::CreateStatic(
+			IsTabVisible,
+			DebuggingToolsManagerWeak,
+			ExecutionFlowTabName
+		)
+	);
+
+	TWeakPtr<SWidget> OwningWidgetWeak = NomadTab;
+	TabContents->SetOnMouseButtonUp(
+		FPointerEventHandler::CreateStatic(
+			[]( /** The geometry of the widget*/
+				const FGeometry&,
+				/** The Mouse Event that we are processing */
+				const FPointerEvent& PointerEvent,
+				TWeakPtr<SWidget> InOwnerWeak,
+				TSharedPtr<FUICommandList> InCommandList) -> FReply
+			{
+				if (PointerEvent.GetEffectingButton() == EKeys::RightMouseButton)
+				{
+					// if the tab manager is still available then make a context window that allows users to
+					// show and hide tabs:
+					TSharedPtr<SWidget> InOwner = InOwnerWeak.Pin();
+					if (InOwner.IsValid())
+					{
+						FMenuBuilder MenuBuilder(true, InCommandList);
+
+						MenuBuilder.PushCommandList(InCommandList.ToSharedRef());
+						{
+							MenuBuilder.AddMenuEntry(FBlueprintDebuggerCommands::Get().ShowCallStackViewer);
+							MenuBuilder.AddMenuEntry(FBlueprintDebuggerCommands::Get().ShowWatchViewer);
+							MenuBuilder.AddMenuEntry(FBlueprintDebuggerCommands::Get().ShowExecutionTrace);
+						}
+						MenuBuilder.PopCommandList();
+
+
+						FWidgetPath WidgetPath = PointerEvent.GetEventPath() != nullptr ? *PointerEvent.GetEventPath() : FWidgetPath();
+						FSlateApplication::Get().PushMenu(InOwner.ToSharedRef(), WidgetPath, MenuBuilder.MakeWidget(), PointerEvent.GetScreenSpacePosition(), FPopupTransitionEffect(FPopupTransitionEffect::ContextMenu));
+
+						return FReply::Handled();
+					}
+				}
+				
+				return FReply::Unhandled();
+			}
+			, OwningWidgetWeak
+			, CommandList
+		)
+	);
+
+	NomadTab->SetContent(
+		SNew(SBorder)
+		.BorderImage(FEditorStyle::GetBrush("ToolPanel.DarkGroupBorder"))
+		.Padding(FMargin(0.f, 2.f))
+		[
+			TabContents
+		]
+	);
+
+	return NomadTab;
+}
+
+FBlueprintDebugger::FBlueprintDebugger()
+	: Impl(MakeUnique<FBlueprintDebuggerImpl>())
+{
+}
+
+FBlueprintDebugger::~FBlueprintDebugger()
+{
+}
+
+#undef LOCTEXT_NAMESPACE 
