Some ability system comment cleanup, came out of notes from documentation

diff --git a/Engine/Plugins/Runtime/GameplayAbilities/Source/GameplayAbilities/Public/Abilities/GameplayAbility.h b/Engine/Plugins/Runtime/GameplayAbilities/Source/GameplayAbilities/Public/Abilities/GameplayAbility.h
--- a/Engine/Plugins/Runtime/GameplayAbilities/Source/GameplayAbilities/Public/Abilities/GameplayAbility.h
+++ b/Engine/Plugins/Runtime/GameplayAbilities/Source/GameplayAbilities/Public/Abilities/GameplayAbility.h
@@ -29,7 +29,7 @@ class UGameplayTasksComponent;
  *	The main features provided by the AbilitySystem for GameplayAbilities are: 
  *		-CanUse functionality:
  *			-Cooldowns
- *			-Resources (mana, stamina, etc)
+ *			-Costs (mana, stamina, etc)
  *			-etc
  *			
  *		-Replication support
@@ -37,34 +37,24 @@ class UGameplayTasksComponent;
  *			-Client prediction for ability activation
  *			
  *		-Instancing support
- *			-Abilities can be non-instanced (default)
+ *			-Abilities can be non-instanced (native only)
  *			-Instanced per owner
- *			-Instanced per execution
+ *			-Instanced per execution (default)
  *			
  *		-Basic, extendable support for:
  *			-Input binding
  *			-'Giving' abilities (that can be used) to actors
  *	
- *	
- *	The intention is for programmers to create these non instanced abilities in C++. Designers can then
- *	extend them as data assets (E.g., they can change default properties, they cannot implement blueprint graphs).
- *	
- *	See GameplayAbility_Montage for example.
+ *
+ *	See GameplayAbility_Montage for an example of a non-instanced ability
  *		-Plays a montage and applies a GameplayEffect to its target while the montage is playing.
  *		-When finished, removes GameplayEffect.
- *		
  *	
  *	Note on replication support:
  *		-Non instanced abilities have limited replication support. 
  *			-Cannot have state (obviously) so no replicated properties
  *			-RPCs on the ability class are not possible either.
  *			
- *			-However: generic RPC functionality can be achieved through the UAbilitySystemAttribute.
- *				-E.g.: ServerTryActivateAbility(class UGameplayAbility* AbilityToActivate, int32 PredictionKey)
- *				
- *	A lot is possible with non instanced abilities but care must be taken.
- *	
- *	
  *	To support state or event replication, an ability must be instanced. This can be done with the InstancingPolicy property.
  */
 
@@ -120,14 +110,13 @@ public:
 	//		TryActivateAbility()	- Attempts to activate the ability. Calls CanActivateAbility(). Input events can call this directly.
 	//								- Also handles instancing-per-execution logic and replication/prediction calls.
 	//		
-	//		CallActivate()			- Protected, non virtual function. Does some boilerplate 'pre activate' stuff, then calls Activate()
+	//		CallActivateAbility()	- Protected, non virtual function. Does some boilerplate 'pre activate' stuff, then calls Activate()
 	//
-	//		Activate()				- What the abilities *does*. This is what child classes want to override.
+	//		ActivateAbility()		- What the abilities *does*. This is what child classes want to override.
 	//	
-	//		Commit()				- Commits reources/cooldowns etc. Activate() must call this!
+	//		CommitAbility()			- Commits reources/cooldowns etc. Activate() must call this!
 	//		
 	//		CancelAbility()			- Interrupts the ability (from an outside source).
-	//									-We may want to add some info on what/who cancelled.
 	//
 	//		EndAbility()			- The ability has ended. This is intended to be called by the ability to end itself.
 	//	
@@ -860,7 +849,6 @@ protected:
 	/** Decreases the scope lock count. Runs the waiting to execute delegates if the count drops to zero. */
 	void DecrementListLock() const;
 
-
 public:
 	/** Setter for the bMarkPendingKillOnAbilityEnd */
 	void SetMarkPendingKillOnAbilityEnd(bool bInMarkPendingKillOnAbilityEnd) { bMarkPendingKillOnAbilityEnd = bInMarkPendingKillOnAbilityEnd; }
diff --git a/Engine/Plugins/Runtime/GameplayAbilities/Source/GameplayAbilities/Public/AbilitySystemComponent.h b/Engine/Plugins/Runtime/GameplayAbilities/Source/GameplayAbilities/Public/AbilitySystemComponent.h
--- a/Engine/Plugins/Runtime/GameplayAbilities/Source/GameplayAbilities/Public/AbilitySystemComponent.h
+++ b/Engine/Plugins/Runtime/GameplayAbilities/Source/GameplayAbilities/Public/AbilitySystemComponent.h
@@ -34,22 +34,17 @@ class UInputComponent;
  *	UAbilitySystemComponent	
  *
  *	A component to easily interface with the 3 aspects of the AbilitySystem:
- *		-GameplayAbilities
- *		-GameplayEffects
- *		-GameplayAttributes
- *		
- *	This component will make life easier for interfacing with these subsystems, but is not completely required. The main functions are:
  *	
  *	GameplayAbilities:
  *		-Provides a way to give/assign abilities that can be used (by a player or AI for example)
  *		-Provides management of instanced abilities (something must hold onto them)
  *		-Provides replication functionality
- *			-Ability state must always be replicated on the UGameplayAbility itself, but UAbilitySystemComponent can provide RPC replication
- *			for non-instanced gameplay abilities. (Explained more in GameplayAbility.h).
+ *			-Ability state must always be replicated on the UGameplayAbility itself, but UAbilitySystemComponent provides RPC replication
+ *			for the actual activation of abilities
  *			
  *	GameplayEffects:
  *		-Provides an FActiveGameplayEffectsContainer for holding active GameplayEffects
- *		-Provides methods for apply GameplayEffect to a target or to self
+ *		-Provides methods for applying GameplayEffects to a target or to self
  *		-Provides wrappers for querying information in FActiveGameplayEffectsContainers (duration, magnitude, etc)
  *		-Provides methods for clearing/remove GameplayEffects
  *		
@@ -57,7 +52,6 @@ class UInputComponent;
  *		-Provides methods for allocating and initializing attribute sets
  *		-Provides methods for getting AttributeSets
  *  
- * 
  */
 
 /** Called when a targeting actor rejects target confirmation */
@@ -167,7 +161,7 @@ class GAMEPLAYABILITIES_API UAbilitySystemComponent : public UGameplayTasksCompo
 	float GetNumericAttributeBase(const FGameplayAttribute &Attribute) const;
 
 	/**
-	 *	Applies an inplace mod to the given attribute. This correctly update the attribute's aggregator, updates the attribute set property,
+	 *	Applies an in-place mod to the given attribute. This correctly update the attribute's aggregator, updates the attribute set property,
 	 *	and invokes the OnDirty callbacks.
 	 *	
 	 *	This does not invoke Pre/PostGameplayEffectExecute calls on the attribute set. This does no tag checking, application requirements, immunity, etc.
@@ -1157,10 +1151,6 @@ class GAMEPLAYABILITIES_API UAbilitySystemComponent : public UGameplayTasksCompo
 
 	// ----------------------------------------------------------------------------------------------------------------
 	//	AnimMontage Support
-	//	
-	//	TODO:
-	//	-Continously update RepAnimMontageInfo on server for join in progress clients.
-	//	-Some missing functionality may still be needed (GetCurrentSectionTime, etc)	
 	// ----------------------------------------------------------------------------------------------------------------	
 
 	/** Plays a montage and handles replication and prediction based on passed in ability/activation info */
@@ -1521,10 +1511,11 @@ protected:
 		
 	TArray<UGameplayTask*>&	GetAbilityActiveTasks(UGameplayAbility* Ability);
 	
-	// Contains all of the gameplay effects that are currently active on this component
+	/** Contains all of the gameplay effects that are currently active on this component */
 	UPROPERTY(Replicated)
 	FActiveGameplayEffectsContainer	ActiveGameplayEffects;
 
+	/** List of all active gameplay cues, including ones applied manually */
 	UPROPERTY(Replicated)
 	FActiveGameplayCueContainer	ActiveGameplayCues;
 
@@ -1541,7 +1532,7 @@ protected:
 
 	void DebugCyclicAggregatorBroadcasts(struct FAggregator* Aggregator);
 	
-	// Acceleration map for all gameplay tags (OwnedGameplayTags from GEs and explicit GameplayCueTags)
+	/** Acceleration map for all gameplay tags (OwnedGameplayTags from GEs and explicit GameplayCueTags) */
 	FGameplayTagCountContainer GameplayTagCountContainer;
 
 	UPROPERTY(Replicated)
diff --git a/Engine/Plugins/Runtime/GameplayAbilities/Source/GameplayAbilities/Public/GameplayEffect.h b/Engine/Plugins/Runtime/GameplayAbilities/Source/GameplayAbilities/Public/GameplayEffect.h
--- a/Engine/Plugins/Runtime/GameplayAbilities/Source/GameplayAbilities/Public/GameplayEffect.h
+++ b/Engine/Plugins/Runtime/GameplayAbilities/Source/GameplayAbilities/Public/GameplayEffect.h
@@ -104,9 +104,9 @@ public:
 	 * @note:	This function assumes (and asserts on) the existence of the required captured attribute within the spec.
 	 *			It is the responsibility of the caller to verify that the spec is properly setup before calling this function.
 	 *			
-	 *	@param InRelevantSpec	Gameplay effect spec providing the backing attribute capture
+	 * @param InRelevantSpec	Gameplay effect spec providing the backing attribute capture
 	 *	
-	 *	@return Evaluated magnitude based upon the spec & calculation policy
+	 * @return Evaluated magnitude based upon the spec & calculation policy
 	 */
 	float CalculateMagnitude(const FGameplayEffectSpec& InRelevantSpec) const;
 
@@ -172,9 +172,9 @@ public:
 	 * @note:	This function assumes (and asserts on) the existence of the required captured attribute within the spec.
 	 *			It is the responsibility of the caller to verify that the spec is properly setup before calling this function.
 	 *			
-	 *	@param InRelevantSpec	Gameplay effect spec providing the backing attribute capture
+	 * @param InRelevantSpec	Gameplay effect spec providing the backing attribute capture
 	 *	
-	 *	@return Evaluated magnitude based upon the spec & calculation policy
+	 * @return Evaluated magnitude based upon the spec & calculation policy
 	 */
 	float CalculateMagnitude(const FGameplayEffectSpec& InRelevantSpec) const;
 
@@ -338,7 +338,6 @@ protected:
 	UPROPERTY(EditDefaultsOnly, Category=Magnitude)
 	FSetByCallerFloat SetByCallerMagnitude;
 
-	// @hack: @todo: This is temporary to aid in post-load fix-up w/o exposing members publicly
 	friend class UGameplayEffect;
 	friend class FGameplayEffectModifierMagnitudeDetails;
 };
@@ -461,7 +460,6 @@ struct GAMEPLAYABILITIES_API FGameplayEffectExecutionDefinition
  * FGameplayModifierInfo
  *	Tells us "Who/What we" modify
  *	Does not tell us how exactly
- *
  */
 USTRUCT(BlueprintType)
 struct GAMEPLAYABILITIES_API FGameplayModifierInfo
@@ -482,7 +480,6 @@ struct GAMEPLAYABILITIES_API FGameplayModifierInfo
 	UPROPERTY(EditDefaultsOnly, Category=GameplayModifier)
 	TEnumAsByte<EGameplayModOp::Type> ModifierOp;
 
-	// @todo: Remove this after content resave
 	/** Now "deprecated," though being handled in a custom manner to avoid engine version bump. */
 	UPROPERTY()
 	FScalableFloat Magnitude;
@@ -643,7 +640,6 @@ struct FModifierSpec
 
 private:
 
-	// @todo: Probably need to make the modifier info private so people don't go gunking around in the magnitude
 	/** In the event that the modifier spec requires custom magnitude calculations, this is the authoritative, last evaluated value of the magnitude */
 	UPROPERTY()
 	float EvaluatedMagnitude;
@@ -903,7 +899,6 @@ struct GAMEPLAYABILITIES_API FGameplayEffectSpec
 
 	// --------------------------------------------------------------------------------------------------------------------------
 	//	IMPORTANT: Any state added to FGameplayEffectSpec must be handled in the move/copy constructor/operator!
-	//	(When VS2012/2013 support is dropped, we can use compiler generated operators, but until then these need to be maintained manually!)
 	// --------------------------------------------------------------------------------------------------------------------------
 
 	FGameplayEffectSpec();
@@ -920,16 +915,16 @@ struct GAMEPLAYABILITIES_API FGameplayEffectSpec
 
 	FGameplayEffectSpec& operator=(const FGameplayEffectSpec& Other);
 
-	// Can be called manually but it is preferred to use the 3 parameter constructor
+	/** Can be called manually but it is preferred to use the 3 parameter constructor */
 	void Initialize(const UGameplayEffect* InDef, const FGameplayEffectContextHandle& InEffectContext, float Level = FGameplayEffectConstants::INVALID_LEVEL);
 
-	// Initialize the spec as a linked spec. The original spec's context is preserved except for the original GE asset tags, which are stripped out
+	/** Initialize the spec as a linked spec. The original spec's context is preserved except for the original GE asset tags, which are stripped out */
 	void InitializeFromLinkedSpec(const UGameplayEffect* InDef, const FGameplayEffectSpec& OriginalSpec);
 
-	// Copies SetbyCallerMagnitudes from OriginalSpec into this
+	/** Copies SetbyCallerMagnitudes from OriginalSpec into this */
 	void CopySetByCallerMagnitudes(const FGameplayEffectSpec& OriginalSpec);
 
-	// Copies SetbuCallerMagnitudes, but only if magnitudes don't exist in our map (slower but preserves data)
+	/** Copies SetbuCallerMagnitudes, but only if magnitudes don't exist in our map (slower but preserves data) */
 	void MergeSetByCallerMagnitudes(const TMap<FGameplayTag, float>& Magnitudes);
 
 	/**
@@ -972,10 +967,10 @@ struct GAMEPLAYABILITIES_API FGameplayEffectSpec
 		return EffectContext;
 	}
 
-	// Appends all tags granted by this gameplay effect spec
+	/** Appends all tags granted by this gameplay effect spec */
 	void GetAllGrantedTags(OUT FGameplayTagContainer& Container) const;
 
-	// Appends all tags that apply to this gameplay effect spec
+	/** Appends all tags that apply to this gameplay effect spec */
 	void GetAllAssetTags(OUT FGameplayTagContainer& Container) const;
 
 	/** Sets the magnitude of a SetByCaller modifier */
@@ -1020,6 +1015,7 @@ struct GAMEPLAYABILITIES_API FGameplayEffectSpec
 	 */
 	float GetModifierMagnitude(int32 ModifierIdx, bool bFactorInStackCount) const;
 
+	/** Fills out the modifier magnitudes inside the Modifier Specs */
 	void CalculateModifierMagnitudes();
 
 	/** Recapture attributes from source and target for cloning */
@@ -1040,8 +1036,6 @@ private:
 
 public:
 
-	// -----------------------------------------------------------------------
-
 	/** GameplayEfect definition. The static data that this spec points to. */
 	UPROPERTY()
 	const UGameplayEffect* Def;
@@ -1057,26 +1051,27 @@ public:
 	/** other effects that need to be applied to the target if this effect is successful */
 	TArray< FGameplayEffectSpecHandle > TargetEffectSpecs;
 
-	// The duration in seconds of this effect
-	// instantaneous effects should have a duration of FGameplayEffectConstants::INSTANT_APPLICATION
-	// effects that last forever should have a duration of FGameplayEffectConstants::INFINITE_DURATION
+	/**
+	 * The duration in seconds of this effect
+	 * instantaneous effects should have a duration of FGameplayEffectConstants::INSTANT_APPLICATION
+	 * effects that last forever should have a duration of FGameplayEffectConstants::INFINITE_DURATION
+	 */
 	UPROPERTY()
 	float Duration;
 
-	// The period in seconds of this effect.
-	// Nonperiodic effects should have a period of FGameplayEffectConstants::NO_PERIOD
+	/** The period in seconds of this effect, nonperiodic effects should have a period of FGameplayEffectConstants::NO_PERIOD */
 	UPROPERTY()
 	float Period;
 
-	// The chance, in a 0.0-1.0 range, that this GameplayEffect will be applied to the target Attribute or GameplayEffect.
+	/** The chance, in a 0.0-1.0 range, that this GameplayEffect will be applied to the target Attribute or GameplayEffect */
 	UPROPERTY()
 	float ChanceToApplyToTarget;
 
-	// Captured Source Tags on GameplayEffectSpec creation.	
+	/** Captured Source Tags on GameplayEffectSpec creation */
 	UPROPERTY(NotReplicated)
 	FTagContainerAggregator	CapturedSourceTags;
 
-	// Tags from the target, captured during execute	
+	/** Tags from the target, captured during execute */
 	UPROPERTY(NotReplicated)
 	FTagContainerAggregator	CapturedTargetTags;
 
@@ -1087,10 +1082,12 @@ public:
 	/** Tags that are on this effect spec and that did not come from the UGameplayEffect def. These are replicated. */
 	UPROPERTY()
 	FGameplayTagContainer DynamicAssetTags;
-	
+
+	/** The calcuated modifiers for this effect */
 	UPROPERTY()
 	TArray<FModifierSpec> Modifiers;
 
+	/** Total number of stacks of this effect */
 	UPROPERTY()
 	int32 StackCount;
 
@@ -1106,6 +1103,7 @@ public:
 	UPROPERTY(NotReplicated)
 	uint32 bDurationLocked : 1;
 
+	/** List of abilities granted by this effect */
 	UPROPERTY()
 	TArray<FGameplayAbilitySpecDef> GrantedAbilitySpecs;
 
@@ -1114,10 +1112,12 @@ public:
 	TMap<FGameplayTag, float>	SetByCallerTagMagnitudes;
 
 private:
-	
+
+	/** This tells us how we got here (who / what applied us) */
 	UPROPERTY()
-	FGameplayEffectContextHandle EffectContext; // This tells us how we got here (who / what applied us)
-	
+	FGameplayEffectContextHandle EffectContext; 
+
+	/** The level this effect was applied at */
 	UPROPERTY()
 	float Level;	
 };
@@ -1141,7 +1141,7 @@ struct GAMEPLAYABILITIES_API FGameplayEffectSpecForRPC
 	TArray<FGameplayEffectModifiedAttribute> ModifiedAttributes;
 
 	UPROPERTY()
-	FGameplayEffectContextHandle EffectContext; // This tells us how we got here (who / what applied us)
+	FGameplayEffectContextHandle EffectContext;
 
 	UPROPERTY()
 	FGameplayTagContainer AggregatedSourceTags;
@@ -1182,7 +1182,6 @@ struct GAMEPLAYABILITIES_API FGameplayEffectSpecForRPC
  *	-Start time
  *  -When to execute next
  *  -Replication callbacks
- *
  */
 USTRUCT(BlueprintType)
 struct GAMEPLAYABILITIES_API FActiveGameplayEffect : public FFastArraySerializerItem
@@ -1191,7 +1190,6 @@ struct GAMEPLAYABILITIES_API FActiveGameplayEffect : public FFastArraySerializer
 
 	// ---------------------------------------------------------------------------------------------------------------------------------
 	//  IMPORTANT: Any new state added to FActiveGameplayEffect must be handled in the copy/move constructor/operator
-	//	(When VS2012/2013 support is dropped, we can use compiler generated operators, but until then these need to be maintained manually)
 	// ---------------------------------------------------------------------------------------------------------------------------------
 
 	FActiveGameplayEffect();
@@ -1236,7 +1234,7 @@ struct GAMEPLAYABILITIES_API FActiveGameplayEffect : public FFastArraySerializer
 	void PostReplicatedAdd(const struct FActiveGameplayEffectsContainer &InArray);
 	void PostReplicatedChange(const struct FActiveGameplayEffectsContainer &InArray);
 
-	// Debug string used by Fast Array serialization
+	/** Debug string used by Fast Array serialization */
 	FString GetDebugString();
 
 	/** Refreshes the cached StartWorldTime for this effect. To be used when the server/client world time delta changes significantly to keep the start time in sync. */
@@ -1357,10 +1355,7 @@ public:
 	/** Returns true if the query is empty/default. E.g., it has no data set. */
 	bool IsEmpty() const;
 
-	/** 
-	 * Shortcuts for easily creating common query types 
-	 * @todo: add more as dictated by use cases
-	 */
+	// Shortcuts for easily creating common query types 
 
 	/** Creates an effect query that will match if there are any common tags between the given tags and an ActiveGameplayEffect's owning tags */
 	static FGameplayEffectQuery MakeQuery_MatchAnyOwningTags(const FGameplayTagContainer& InTags);
@@ -1434,16 +1429,16 @@ struct FActiveGameplayEffectQuery
 	/** used to reject matches with InheritableGameplayEffectTags */
 	const FGameplayTagContainer* EffectTagContainer_Rejection;
 
-	// Matches on GameplayEffects which modify given attribute
+	/** Matches on GameplayEffects which modify given attribute */
 	FGameplayAttribute ModifyingAttribute;
 
-	// Matches on GameplayEffects which come from this source
+	/** Matches on GameplayEffects which come from this source */
 	const UObject* EffectSource;
 
-	// Matches on GameplayEffects with this definition
+	/** Matches on GameplayEffects with this definition */
 	const UGameplayEffect* EffectDef;
 
-	// Handles to ignore as matches, even if other criteria is met
+	/** Handles to ignore as matches, even if other criteria is met */
 	TArray<FActiveGameplayEffectHandle> IgnoreHandles;
 };
 
@@ -1469,16 +1464,10 @@ struct FCustomModifierDependencyHandle
  *  
  * This should only be used by UAbilitySystemComponent. All of this could just live in UAbilitySystemComponent except that we need a distinct USTRUCT to implement FFastArraySerializer.
  *
- * The preferred way to iterate through the ActiveGameplayEffectContainer is with CreateConstITerator/CreateIterator or stl style range iteration:
+ * The preferred way to iterate through the ActiveGameplayEffectContainer is with CreateConstIterator/CreateIterator or stl style range iteration:
  * 
- * 
- *	for (const FActiveGameplayEffect& Effect : this)
- *	{
- *	}
- *
- *	for (auto It = CreateConstIterator(); It; ++It) 
- *	{
- *	}
+ *	for (const FActiveGameplayEffect& Effect : this) {}
+ *	for (auto It = CreateConstIterator(); It; ++It) {}
  *
  */
 USTRUCT()
@@ -1515,7 +1504,7 @@ struct GAMEPLAYABILITIES_API FActiveGameplayEffectsContainer : public FFastArray
 	};
 
 #if ENABLE_VISUAL_LOG
-	// Stores a record of gameplay effects that have executed and their results. Useful for debugging.
+	/** Stores a record of gameplay effects that have executed and their results. Useful for debugging */
 	TArray<DebugExecutedGameplayEffectData> DebugExecutedGameplayEffects;
 
 	void GrabDebugSnapshot(FVisualLogEntry* Snapshot) const;
@@ -1604,8 +1593,6 @@ struct GAMEPLAYABILITIES_API FActiveGameplayEffectsContainer : public FFastArray
 
 	void Uninitialize();
 
-	// ------------------------------------------------
-
 	bool CanApplyAttributeModifiers(const UGameplayEffect *GameplayEffect, float Level, const FGameplayEffectContextHandle& EffectContext);
 	
 	TArray<float> GetActiveEffectsTimeRemaining(const FGameplayEffectQuery& Query) const;
@@ -1766,7 +1753,6 @@ private:
 	
 	TMap<FGameplayAttribute, FOnGameplayAttributeValueChange> AttributeValueChangeDelegates;
 
-
 	TMap<FGameplayTag, TSet<FActiveGameplayEffectHandle> >	ActiveEffectTagDependencies;
 
 	/** Mapping of custom gameplay modifier magnitude calculation class to dependency handles for triggering updates on external delegates firing */
@@ -1876,6 +1862,22 @@ public:
 	static const float NO_PERIOD;	
 	static const float INVALID_LEVEL;
 
+	virtual void GetOwnedGameplayTags(FGameplayTagContainer& TagContainer) const override;
+	virtual void PostInitProperties() override;
+	virtual void PostLoad() override;
+	virtual void PreSave(const class ITargetPlatform* TargetPlatform) override;
+#if WITH_EDITOR
+	void PostEditChangeProperty(FPropertyChangedEvent& PropertyChangedEvent) override;
+#endif
+
+	/** Update all the InheritedTagContainers */
+	void UpdateInheritedTagProperties();
+
+	/** Check for any errors */
+	void ValidateGameplayEffect();
+
+	// ---------------------------------------------------------------------------------------------------------------------------------
+
 #if WITH_EDITORONLY_DATA
 	/** Template to derive starting values and editing customization from */
 	UPROPERTY()
@@ -1886,11 +1888,6 @@ public:
 	bool ShowAllProperties;
 #endif
 
-	virtual void PostInitProperties() override;
-#if WITH_EDITOR
-	void PostEditChangeProperty(FPropertyChangedEvent& PropertyChangedEvent) override;
-#endif
-
 	/** Policy for the duration of this effect */
 	UPROPERTY(EditDefaultsOnly, Category=GameplayEffect)
 	EGameplayEffectDurationType DurationPolicy;
@@ -1949,21 +1946,6 @@ public:
 	UPROPERTY(EditDefaultsOnly, BlueprintReadOnly, Category=Expiration)
 	TArray<TSubclassOf<UGameplayEffect>> RoutineExpirationEffectClasses;
 
-	// ------------------------------------------------
-	// Gameplay tag interface
-	// ------------------------------------------------
-
-	/** Overridden to return requirements tags */
-	virtual void GetOwnedGameplayTags(FGameplayTagContainer& TagContainer) const override;
-
-	void UpdateInheritedTagProperties();
-	void ValidateGameplayEffect();
-
-	virtual void PostLoad() override;
-	virtual void PreSave(const class ITargetPlatform* TargetPlatform) override;
-
-	// ----------------------------------------------
-
 	/** If true, cues will only trigger when GE modifiers succeed being applied (whether through modifiers or executions) */
 	UPROPERTY(EditDefaultsOnly, BlueprintReadOnly, Category = Display)
 	bool bRequireModifierSuccessToTriggerCues;
