Comments only. Making function names referenced in comments more accurate.

diff --git a/Engine/Plugins/Runtime/GameplayAbilities/Source/GameplayAbilities/Public/Abilities/GameplayAbility.h b/Engine/Plugins/Runtime/GameplayAbilities/Source/GameplayAbilities/Public/Abilities/GameplayAbility.h
--- a/Engine/Plugins/Runtime/GameplayAbilities/Source/GameplayAbilities/Public/Abilities/GameplayAbility.h
+++ b/Engine/Plugins/Runtime/GameplayAbilities/Source/GameplayAbilities/Public/Abilities/GameplayAbility.h
@@ -110,11 +110,11 @@ public:
 	//		TryActivateAbility()	- Attempts to activate the ability. Calls CanActivateAbility(). Input events can call this directly.
 	//								- Also handles instancing-per-execution logic and replication/prediction calls.
 	//		
-	//		CallActivateAbility()	- Protected, non virtual function. Does some boilerplate 'pre activate' stuff, then calls Activate()
+	//		CallActivateAbility()	- Protected, non virtual function. Does some boilerplate 'pre activate' stuff, then calls ActivateAbility()
 	//
 	//		ActivateAbility()		- What the abilities *does*. This is what child classes want to override.
 	//	
-	//		CommitAbility()			- Commits reources/cooldowns etc. Activate() must call this!
+	//		CommitAbility()			- Commits reources/cooldowns etc. ActivateAbility() must call this!
 	//		
 	//		CancelAbility()			- Interrupts the ability (from an outside source).
 	//
