[NO NON-PII COMMIT MESSAGE]

diff --git a/Engine/Plugins/Runtime/GameplayAbilities/Source/GameplayAbilities/Public/Abilities/Tasks/AbilityTask_ApplyRootMotion_Base.h b/Engine/Plugins/Runtime/GameplayAbilities/Source/GameplayAbilities/Public/Abilities/Tasks/AbilityTask_ApplyRootMotion_Base.h
--- a/Engine/Plugins/Runtime/GameplayAbilities/Source/GameplayAbilities/Public/Abilities/Tasks/AbilityTask_ApplyRootMotion_Base.h
+++ b/Engine/Plugins/Runtime/GameplayAbilities/Source/GameplayAbilities/Public/Abilities/Tasks/AbilityTask_ApplyRootMotion_Base.h
@@ -12,23 +12,20 @@ enum class ERootMotionFinishVelocityMode : uint8;
 DECLARE_MULTICAST_DELEGATE_TwoParams(FOnTargetActorSwapped, AActor*, AActor*);
 
 /** Base class for ability tasks that apply root motion */
-UCLASS(MinimalAPI)
-class UAbilityTask_ApplyRootMotion_Base : public UAbilityTask
+UCLASS()
+class GAMEPLAYABILITIES_API UAbilityTask_ApplyRootMotion_Base : public UAbilityTask
 {
 	GENERATED_UCLASS_BODY()
 
 	virtual void InitSimulatedTask(UGameplayTasksComponent& InGameplayTasksComponent) override;
 
 	//..See notes on delegate definition FOnTargetActorSwapped.
-	static GAMEPLAYABILITIES_API FOnTargetActorSwapped OnTargetActorSwapped;
+	static FOnTargetActorSwapped OnTargetActorSwapped;
 
 protected:
 
 	virtual void SharedInitAndApply() {};
-
-	bool HasTimedOut() const;
-
-protected:
+	virtual bool HasTimedOut() const;
 
 	UPROPERTY(Replicated)
 	FName ForceName;
