PR #4945: Move FGameplayAbilitySpecHandle::GenerateNewHandle body to cpp to fix issue where it would make duplicate copies per module on some platforms

diff --git a/Engine/Plugins/Runtime/GameplayAbilities/Source/GameplayAbilities/Private/GameplayAbilityTypes.cpp b/Engine/Plugins/Runtime/GameplayAbilities/Source/GameplayAbilities/Private/GameplayAbilityTypes.cpp
--- a/Engine/Plugins/Runtime/GameplayAbilities/Source/GameplayAbilities/Private/GameplayAbilityTypes.cpp
+++ b/Engine/Plugins/Runtime/GameplayAbilities/Source/GameplayAbilities/Private/GameplayAbilityTypes.cpp
@@ -246,6 +246,15 @@ FGameplayAbilitySpec::FGameplayAbilitySpec(FGameplayAbilitySpecDef& InDef, int32
 
 // ----------------------------------------------------
 
+void FGameplayAbilitySpecHandle::GenerateNewHandle()
+{
+	// Must be in C++ to avoid duplicate statics accross execution units
+	static int32 GHandle = 1;
+	Handle = GHandle++;
+}
+
+// ----------------------------------------------------
+
 FScopedAbilityListLock::FScopedAbilityListLock(UAbilitySystemComponent& InAbilitySystemComponent)
 	: AbilitySystemComponent(InAbilitySystemComponent)
 {
diff --git a/Engine/Plugins/Runtime/GameplayAbilities/Source/GameplayAbilities/Public/GameplayAbilitySpec.h b/Engine/Plugins/Runtime/GameplayAbilities/Source/GameplayAbilities/Public/GameplayAbilitySpec.h
--- a/Engine/Plugins/Runtime/GameplayAbilities/Source/GameplayAbilities/Public/GameplayAbilitySpec.h
+++ b/Engine/Plugins/Runtime/GameplayAbilities/Source/GameplayAbilities/Public/GameplayAbilitySpec.h
@@ -38,11 +38,7 @@ struct FGameplayAbilitySpecHandle
 	}
 
 	/** Sets this to a valid handle */
-	void GenerateNewHandle()
-	{
-		static int32 GHandle = 1;
-		Handle = GHandle++;
-	}
+	void GenerateNewHandle();
 
 	bool operator==(const FGameplayAbilitySpecHandle& Other) const
 	{
