Added scoped prediction windows to fix a bug where gameplay cue execution did not work following a delay inside locally predicted abilities.

diff --git a/Engine/Plugins/Runtime/GameplayAbilities/Source/GameplayAbilities/Private/AbilitySystemComponent.cpp b/Engine/Plugins/Runtime/GameplayAbilities/Source/GameplayAbilities/Private/AbilitySystemComponent.cpp
--- a/Engine/Plugins/Runtime/GameplayAbilities/Source/GameplayAbilities/Private/AbilitySystemComponent.cpp
+++ b/Engine/Plugins/Runtime/GameplayAbilities/Source/GameplayAbilities/Private/AbilitySystemComponent.cpp
@@ -1084,12 +1084,16 @@ void UAbilitySystemComponent::InvokeGameplayCueEvent(const FGameplayTag Gameplay
 
 void UAbilitySystemComponent::ExecuteGameplayCue(const FGameplayTag GameplayCueTag, FGameplayEffectContextHandle EffectContext)
 {
+	FScopedPredictionWindow ScopedPredictionWindow(this, true);
+
 	// Send to the wrapper on the cue manager
 	UAbilitySystemGlobals::Get().GetGameplayCueManager()->InvokeGameplayCueExecuted(this, GameplayCueTag, ScopedPredictionKey, EffectContext);
 }
 
 void UAbilitySystemComponent::ExecuteGameplayCue(const FGameplayTag GameplayCueTag, const FGameplayCueParameters& GameplayCueParameters)
 {
+	FScopedPredictionWindow ScopedPredictionWindow(this, true);
+
 	// Send to the wrapper on the cue manager
 	UAbilitySystemGlobals::Get().GetGameplayCueManager()->InvokeGameplayCueExecuted_WithParams(this, GameplayCueTag, ScopedPredictionKey, GameplayCueParameters);
 }
