Unshelved from pending changelist '4299645':

PR #5003: Changed SendGameplayEvent()'s category from "Abilities" to "Ability" for consistency (Contributed by [REDACTED_PERSON])

diff --git a/Engine/Plugins/Runtime/GameplayAbilities/Source/GameplayAbilities/Public/Abilities/GameplayAbility.h b/Engine/Plugins/Runtime/GameplayAbilities/Source/GameplayAbilities/Public/Abilities/GameplayAbility.h
--- a/Engine/Plugins/Runtime/GameplayAbilities/Source/GameplayAbilities/Public/Abilities/GameplayAbility.h
+++ b/Engine/Plugins/Runtime/GameplayAbilities/Source/GameplayAbilities/Public/Abilities/GameplayAbility.h
@@ -500,7 +500,7 @@ protected:
 	bool bHasBlueprintShouldAbilityRespondToEvent;
 
 	/** Sends a gameplay event, also creates a prediction window */
-	UFUNCTION(BlueprintCallable, Category = "Abilities")
+	UFUNCTION(BlueprintCallable, Category = Ability)
 	virtual void SendGameplayEvent(FGameplayTag EventTag, FGameplayEventData Payload);
 
 	// --------------------------------------
