FGameplayTagBlueprintPropertyMap unregister;

* FGameplayTagBlueprintPropertyMap will now unregister it's event delegates if it gets initialized again.

diff --git a/Engine/Plugins/Runtime/GameplayAbilities/Source/GameplayAbilities/Private/GameplayEffectTypes.cpp b/Engine/Plugins/Runtime/GameplayAbilities/Source/GameplayAbilities/Private/GameplayEffectTypes.cpp
--- a/Engine/Plugins/Runtime/GameplayAbilities/Source/GameplayAbilities/Private/GameplayEffectTypes.cpp
+++ b/Engine/Plugins/Runtime/GameplayAbilities/Source/GameplayAbilities/Private/GameplayEffectTypes.cpp
@@ -502,16 +502,7 @@ FGameplayTagBlueprintPropertyMap::FGameplayTagBlueprintPropertyMap()
 
 FGameplayTagBlueprintPropertyMap::~FGameplayTagBlueprintPropertyMap()
 {
-	if (UAbilitySystemComponent* ASC = CachedASC.Get())
-	{
-		for (FGameplayTagBlueprintPropertyMapping& Mapping : PropertyMappings)
-		{
-			if (Mapping.PropertyToEdit && Mapping.TagToMap.IsValid())
-			{
-				ASC->UnregisterGameplayTagEvent(Mapping.DelegateHandle, Mapping.TagToMap, GetGameplayTagEventType(Mapping.PropertyToEdit));
-			}
-		}
-	}
+	Unregister();
 }
 
 #if WITH_EDITOR
@@ -569,6 +560,11 @@ void FGameplayTagBlueprintPropertyMap::Initialize(UObject* Owner, UAbilitySystem
 		return;
 	}
 
+	if (CachedOwner.IsValid())
+	{
+		Unregister();
+	}
+
 	CachedOwner = Owner;
 	CachedASC = ASC;
 
@@ -598,6 +594,26 @@ void FGameplayTagBlueprintPropertyMap::Initialize(UObject* Owner, UAbilitySystem
 	}
 }
 
+void FGameplayTagBlueprintPropertyMap::Unregister()
+{
+	if (UAbilitySystemComponent* ASC = CachedASC.Get())
+	{
+		for (FGameplayTagBlueprintPropertyMapping& Mapping : PropertyMappings)
+		{
+			if (Mapping.PropertyToEdit && Mapping.TagToMap.IsValid())
+			{
+				ASC->UnregisterGameplayTagEvent(Mapping.DelegateHandle, Mapping.TagToMap, GetGameplayTagEventType(Mapping.PropertyToEdit));
+			}
+
+			Mapping.PropertyToEdit = nullptr;
+			Mapping.DelegateHandle.Reset();
+		}
+	}
+
+	CachedOwner = nullptr;
+	CachedASC = nullptr;
+}
+
 void FGameplayTagBlueprintPropertyMap::GameplayTagEventCallback(const FGameplayTag Tag, int32 NewCount)
 {
 	UObject* Owner = CachedOwner.Get();
diff --git a/Engine/Plugins/Runtime/GameplayAbilities/Source/GameplayAbilities/Public/GameplayEffectTypes.h b/Engine/Plugins/Runtime/GameplayAbilities/Source/GameplayAbilities/Public/GameplayEffectTypes.h
--- a/Engine/Plugins/Runtime/GameplayAbilities/Source/GameplayAbilities/Public/GameplayEffectTypes.h
+++ b/Engine/Plugins/Runtime/GameplayAbilities/Source/GameplayAbilities/Public/GameplayEffectTypes.h
@@ -1383,6 +1383,8 @@ public:
 
 protected:
 
+	void Unregister();
+
 	void GameplayTagEventCallback(const FGameplayTag Tag, int32 NewCount);
 
 	bool IsPropertyTypeValid(const UProperty* Property) const;
