- Adding extra logging to help track down ensure
#rnx

diff --git a/Engine/Plugins/Runtime/GameplayAbilities/Source/GameplayAbilities/Private/Abilities/GameplayAbility.cpp b/Engine/Plugins/Runtime/GameplayAbilities/Source/GameplayAbilities/Private/Abilities/GameplayAbility.cpp
--- a/Engine/Plugins/Runtime/GameplayAbilities/Source/GameplayAbilities/Private/Abilities/GameplayAbility.cpp
+++ b/Engine/Plugins/Runtime/GameplayAbilities/Source/GameplayAbilities/Private/Abilities/GameplayAbility.cpp
@@ -532,6 +532,7 @@ void UGameplayAbility::CancelAbility(const FGameplayAbilitySpecHandle Handle, co
 	{
 		if (ScopeLockCount > 0)
 		{
+			UE_LOG(LogAbilitySystem, Verbose, TEXT("Attempting to cancel Ability %s but ScopeLockCount was greater than 0, adding cancel to the WaitingToExecute Array"), *GetName());
 			WaitingToExecute.Add(FPostLockDelegate::CreateUObject(this, &UGameplayAbility::CancelAbility, Handle, ActorInfo, ActivationInfo, bReplicateCancelAbility));
 			return;
 		}
