- I'm thinking this has something to do with our ability system component being cleaned up with our pawn but the latent action system is holding onto a refererence to our gameplay ability. I'm adding additional logging and an IsValid check to see if that's the case

diff --git a/Engine/Plugins/Runtime/GameplayAbilities/Source/GameplayAbilities/Private/Abilities/GameplayAbility.cpp b/Engine/Plugins/Runtime/GameplayAbilities/Source/GameplayAbilities/Private/Abilities/GameplayAbility.cpp
--- a/Engine/Plugins/Runtime/GameplayAbilities/Source/GameplayAbilities/Private/Abilities/GameplayAbility.cpp
+++ b/Engine/Plugins/Runtime/GameplayAbilities/Source/GameplayAbilities/Private/Abilities/GameplayAbility.cpp
@@ -562,6 +562,7 @@ bool UGameplayAbility::IsEndAbilityValid(const FGameplayAbilitySpecHandle Handle
 	// Ending an AbilityState may cause this to be invoked again
 	if (bIsActive == false && GetInstancingPolicy() != EGameplayAbilityInstancingPolicy::NonInstanced)
 	{
+		UE_LOG(LogAbilitySystem, Verbose, TEXT("IsEndAbilityValid returning false on Ability %s due to EndAbility being called multiple times"), *GetName());
 		return false;
 	}
 
@@ -569,6 +570,7 @@ bool UGameplayAbility::IsEndAbilityValid(const FGameplayAbilitySpecHandle Handle
 	UAbilitySystemComponent* AbilityComp = ActorInfo ? ActorInfo->AbilitySystemComponent.Get() : nullptr;
 	if (AbilityComp == nullptr)
 	{
+		UE_LOG(LogAbilitySystem, Verbose, TEXT("IsEndAbilityValid returning false on Ability %s due to AbilitySystemComponent being invalid"), *GetName());
 		return false;
 	}
 
@@ -578,6 +580,7 @@ bool UGameplayAbility::IsEndAbilityValid(const FGameplayAbilitySpecHandle Handle
 
 	if (!bIsSpecActive)
 	{
+		UE_LOG(LogAbilitySystem, Verbose, TEXT("IsEndAbilityValid returning false on Ability %s due spec not being active"), *GetName());
 		return false;
 	}
 
@@ -590,6 +593,7 @@ void UGameplayAbility::EndAbility(const FGameplayAbilitySpecHandle Handle, const
 	{
 		if (ScopeLockCount > 0)
 		{
+			UE_LOG(LogAbilitySystem, Verbose, TEXT("Attempting to end Ability %s but ScopeLockCount was greater than 0, adding end to the WaitingToExecute Array"), *GetName());
 			WaitingToExecute.Add(FPostLockDelegate::CreateUObject(this, &UGameplayAbility::EndAbility, Handle, ActorInfo, ActivationInfo, bReplicateEndAbility, bWasCancelled));
 			return;
 		}
