Fix for Crash in the ability system.
 - Ensure the attribute set is valid before attempting to use it

#rnx

diff --git a/Engine/Plugins/Runtime/GameplayAbilities/Source/GameplayAbilities/Private/GameplayEffect.cpp b/Engine/Plugins/Runtime/GameplayAbilities/Source/GameplayAbilities/Private/GameplayEffect.cpp
--- a/Engine/Plugins/Runtime/GameplayAbilities/Source/GameplayAbilities/Private/GameplayEffect.cpp
+++ b/Engine/Plugins/Runtime/GameplayAbilities/Source/GameplayAbilities/Private/GameplayEffect.cpp
@@ -2581,14 +2581,18 @@ float FActiveGameplayEffectsContainer::GetAttributeBaseValue(FGameplayAttribute
 	float BaseValue = 0.f;
 	if (Owner)
 	{
+		const UAttributeSet* AttributeSet = Owner->GetAttributeSubobject(Attribute.GetAttributeSetClass());
+		if (!ensureMsgf(AttributeSet, TEXT("FActiveGameplayEffectsContainer::SetAttributeBaseValue: Unable to get attribute set for attribute %s"), *Attribute.AttributeName))
+		{
+			return BaseValue;
+		}
+
 		const FAggregatorRef* RefPtr = AttributeAggregatorMap.Find(Attribute);
 		// if this attribute is of type FGameplayAttributeData then use the base value stored there
 		if (FGameplayAttribute::IsGameplayAttributeDataProperty(Attribute.GetUProperty()))
 		{
 			const UStructProperty* StructProperty = Cast<UStructProperty>(Attribute.GetUProperty());
 			check(StructProperty);
-			const UAttributeSet* AttributeSet = Owner->GetAttributeSubobject(Attribute.GetAttributeSetClass());
-			ensure(AttributeSet);
 			const FGameplayAttributeData* DataPtr = StructProperty->ContainerPtrToValuePtr<FGameplayAttributeData>(AttributeSet);
 			if (DataPtr)
 			{
