Fix for certain gameplay effects getting duplicated in playback of server side replays

diff --git a/Engine/Plugins/Runtime/GameplayAbilities/Source/GameplayAbilities/Private/GameplayEffect.cpp b/Engine/Plugins/Runtime/GameplayAbilities/Source/GameplayAbilities/Private/GameplayEffect.cpp
--- a/Engine/Plugins/Runtime/GameplayAbilities/Source/GameplayAbilities/Private/GameplayEffect.cpp
+++ b/Engine/Plugins/Runtime/GameplayAbilities/Source/GameplayAbilities/Private/GameplayEffect.cpp
@@ -3708,8 +3708,8 @@ bool FActiveGameplayEffectsContainer::NetDeltaSerialize(FNetDeltaSerializeInfo&
 			{
 				UNetConnection* Connection = Client->GetConnection();
 
-				// Even in mixed mode, we should always replicate out to replays so it has all information.
-				if (Connection->GetDriver()->NetDriverName != NAME_DemoNetDriver)
+				// Even in mixed mode, we should always replicate out to client side recorded replays so it has all information.
+				if (Connection->GetDriver()->NetDriverName != NAME_DemoNetDriver || IsNetAuthority())
 				{
 					// In mixed mode, we only want to replicate to the owner of this channel, minimal replication
 					// data will go to everyone else.
