Added an enum to decide how periodic effects should respond when they are no longer inhibited. By default they will not adjust the timing of their execution (this matches the old behaviour). The other options are to reset the period to start when they are no longer inhibited or to execute immediately and reset the timer.

diff --git a/Engine/Plugins/Runtime/GameplayAbilities/Source/GameplayAbilities/Private/GameplayEffect.cpp b/Engine/Plugins/Runtime/GameplayAbilities/Source/GameplayAbilities/Private/GameplayEffect.cpp
--- a/Engine/Plugins/Runtime/GameplayAbilities/Source/GameplayAbilities/Private/GameplayEffect.cpp
+++ b/Engine/Plugins/Runtime/GameplayAbilities/Source/GameplayAbilities/Private/GameplayEffect.cpp
@@ -52,6 +52,7 @@ UGameplayEffect::UGameplayEffect(const FObjectInitializer& ObjectInitializer)
 {
 	DurationPolicy = EGameplayEffectDurationType::Instant;
 	bExecutePeriodicEffectOnApplication = true;
+	PeriodicInhibitionPolicy = EGameplayEffectPeriodInhibitionRemovedPolicy::NeverReset;
 	ChanceToApplyToTarget.SetValue(1.f);
 	StackingType = EGameplayEffectStackingType::None;
 	StackLimitCount = 0;
@@ -2742,7 +2743,7 @@ FActiveGameplayEffect* FActiveGameplayEffectsContainer::ApplyGameplayEffectSpec(
 		UAbilitySystemGlobals::Get().SetCurrentAppliedGE(&ExistingSpec);
 		
 		// How to apply multiple stacks at once? What if we trigger an overflow which can reject the application?
-		// We still want to apply the stacks that didnt push us over, but we also want to call HandleActiveGameplayEffectStackOverflow.
+		// We still want to apply the stacks that didn't push us over, but we also want to call HandleActiveGameplayEffectStackOverflow.
 		
 		// For now: call HandleActiveGameplayEffectStackOverflow only if we are ALREADY at the limit. Else we just clamp stack limit to max.
 		if (ExistingSpec.StackCount == ExistingSpec.Def->StackLimitCount)
@@ -3073,6 +3074,22 @@ void FActiveGameplayEffectsContainer::AddActiveGameplayEffectGrantedTagsAndModif
 			}
 		}
 	}
+	else
+	{
+		if (Effect.Spec.Def->PeriodicInhibitionPolicy != EGameplayEffectPeriodInhibitionRemovedPolicy::NeverReset && Owner)
+		{
+			FTimerManager& TimerManager = Owner->GetWorld()->GetTimerManager();
+			FTimerDelegate Delegate = FTimerDelegate::CreateUObject(Owner, &UAbilitySystemComponent::ExecutePeriodicEffect, Effect.Handle);
+
+			// The timer manager moves things from the pending list to the active list after checking the active list on the first tick so we need to execute here
+			if (Effect.Spec.Def->PeriodicInhibitionPolicy == EGameplayEffectPeriodInhibitionRemovedPolicy::ExecuteAndResetPeriod)
+			{
+				TimerManager.SetTimerForNextTick(Delegate);
+			}
+
+			TimerManager.SetTimer(Effect.PeriodHandle, Delegate, Effect.Spec.GetPeriod(), true);
+		}
+	}
 
 	// Update our owner with the tags this GameplayEffect grants them
 	Owner->UpdateTagMap(Effect.Spec.Def->InheritableOwnedTagsContainer.CombinedTags, 1);
diff --git a/Engine/Plugins/Runtime/GameplayAbilities/Source/GameplayAbilities/Public/GameplayEffect.h b/Engine/Plugins/Runtime/GameplayAbilities/Source/GameplayAbilities/Public/GameplayEffect.h
--- a/Engine/Plugins/Runtime/GameplayAbilities/Source/GameplayAbilities/Public/GameplayEffect.h
+++ b/Engine/Plugins/Runtime/GameplayAbilities/Source/GameplayAbilities/Public/GameplayEffect.h
@@ -628,6 +628,20 @@ enum class EGameplayEffectStackingExpirationPolicy : uint8
 	RefreshDuration,
 };
 
+/** Enumeration of policies for dealing with the period of a gameplay effect when inhibition is removed */
+UENUM()
+enum class EGameplayEffectPeriodInhibitionRemovedPolicy : uint8
+{
+	/** Does not reset. The period timing will continue as if the inhibition hadn't occurred. */
+	NeverReset,
+
+	/** Resets the period. The next execution will occur one full period from when inhibition is removed. */
+	ResetPeriod,
+
+	/** Executes immediately and resets the period. */
+	ExecuteAndResetPeriod,
+};
+
 /** Holds evaluated magnitude from a GameplayEffect modifier */
 USTRUCT()
 struct FModifierSpec
@@ -1896,6 +1910,9 @@ public:
 	UPROPERTY(EditDefaultsOnly, BlueprintReadOnly, Category=Period)
 	bool bExecutePeriodicEffectOnApplication;
 
+	UPROPERTY(EditDefaultsOnly, BlueprintReadOnly, Category=Period)
+	EGameplayEffectPeriodInhibitionRemovedPolicy PeriodicInhibitionPolicy;
+
 	/** Array of modifiers that will affect the target of this effect */
 	UPROPERTY(EditDefaultsOnly, BlueprintReadOnly, Category=GameplayEffect)
 	TArray<FGameplayModifierInfo> Modifiers;
