Add a unique id for ability target data that gets submitted to the server.
Add the ability to track if hits get replaced.

diff --git a/Engine/Plugins/Runtime/GameplayAbilities/Source/GameplayAbilities/Private/GameplayAbilityTargetTypes.cpp b/Engine/Plugins/Runtime/GameplayAbilities/Source/GameplayAbilities/Private/GameplayAbilityTargetTypes.cpp
--- a/Engine/Plugins/Runtime/GameplayAbilities/Source/GameplayAbilities/Private/GameplayAbilityTargetTypes.cpp
+++ b/Engine/Plugins/Runtime/GameplayAbilities/Source/GameplayAbilities/Private/GameplayAbilityTargetTypes.cpp
@@ -136,6 +136,8 @@ FGameplayAbilityTargetDataHandle FGameplayAbilityTargetingLocationInfo::MakeTarg
 
 bool FGameplayAbilityTargetDataHandle::NetSerialize(FArchive& Ar, class UPackageMap* Map, bool& bOutSuccess)
 {
+	Ar << UniqueId;
+
 	uint8 DataNum;
 	if (Ar.IsSaving())
 	{
diff --git a/Engine/Plugins/Runtime/GameplayAbilities/Source/GameplayAbilities/Public/Abilities/GameplayAbilityTargetTypes.h b/Engine/Plugins/Runtime/GameplayAbilities/Source/GameplayAbilities/Public/Abilities/GameplayAbilityTargetTypes.h
--- a/Engine/Plugins/Runtime/GameplayAbilities/Source/GameplayAbilities/Public/Abilities/GameplayAbilityTargetTypes.h
+++ b/Engine/Plugins/Runtime/GameplayAbilities/Source/GameplayAbilities/Public/Abilities/GameplayAbilityTargetTypes.h
@@ -198,11 +198,13 @@ struct GAMEPLAYABILITIES_API FGameplayAbilityTargetDataHandle
 		Data.Add(TSharedPtr<FGameplayAbilityTargetData>(DataPtr));
 	}
 
-	FGameplayAbilityTargetDataHandle(FGameplayAbilityTargetDataHandle&& Other) : Data(MoveTemp(Other.Data))	{ }
-	FGameplayAbilityTargetDataHandle(const FGameplayAbilityTargetDataHandle& Other) : Data(Other.Data) { }
+	FGameplayAbilityTargetDataHandle(FGameplayAbilityTargetDataHandle&& Other) : UniqueId(Other.UniqueId), Data(MoveTemp(Other.Data))	{ }
+	FGameplayAbilityTargetDataHandle(const FGameplayAbilityTargetDataHandle& Other) : UniqueId(Other.UniqueId), Data(Other.Data) { }
 
-	FGameplayAbilityTargetDataHandle& operator=(FGameplayAbilityTargetDataHandle&& Other) { Data = MoveTemp(Other.Data); return *this; }
-	FGameplayAbilityTargetDataHandle& operator=(const FGameplayAbilityTargetDataHandle& Other) { Data = Other.Data; return *this; }
+	FGameplayAbilityTargetDataHandle& operator=(FGameplayAbilityTargetDataHandle&& Other) { UniqueId = Other.UniqueId; Data = MoveTemp(Other.Data); return *this; }
+	FGameplayAbilityTargetDataHandle& operator=(const FGameplayAbilityTargetDataHandle& Other) { UniqueId = Other.UniqueId; Data = Other.Data; return *this; }
+
+	uint8 UniqueId = 0;
 
 	/** Raw storage of target data, do not modify this directly */
 	TArray<TSharedPtr<FGameplayAbilityTargetData>, TInlineAllocator<1> >	Data;
@@ -617,6 +619,8 @@ struct GAMEPLAYABILITIES_API FGameplayAbilityTargetData_SingleTargetHit : public
 
 	virtual void ReplaceHitWith(AActor* NewHitActor, const FHitResult* NewHitResult)
 	{
+		bHitReplaced = true;
+
 		HitResult = FHitResult();
 		if (NewHitResult != nullptr)
 		{
@@ -630,6 +634,9 @@ struct GAMEPLAYABILITIES_API FGameplayAbilityTargetData_SingleTargetHit : public
 	UPROPERTY()
 	FHitResult	HitResult;
 
+	UPROPERTY()
+	bool bHitReplaced = false;
+
 	bool NetSerialize(FArchive& Ar, class UPackageMap* Map, bool& bOutSuccess);
 
 	virtual UScriptStruct* GetScriptStruct() const override
