Check if the GE owner is authoritative before modifying the timers for periodic effects. This fixes a log spam issue.

diff --git a/Engine/Plugins/Runtime/GameplayAbilities/Source/GameplayAbilities/Private/GameplayEffect.cpp b/Engine/Plugins/Runtime/GameplayAbilities/Source/GameplayAbilities/Private/GameplayEffect.cpp
--- a/Engine/Plugins/Runtime/GameplayAbilities/Source/GameplayAbilities/Private/GameplayEffect.cpp
+++ b/Engine/Plugins/Runtime/GameplayAbilities/Source/GameplayAbilities/Private/GameplayEffect.cpp
@@ -3076,7 +3076,7 @@ void FActiveGameplayEffectsContainer::AddActiveGameplayEffectGrantedTagsAndModif
 	}
 	else
 	{
-		if (Effect.Spec.Def->PeriodicInhibitionPolicy != EGameplayEffectPeriodInhibitionRemovedPolicy::NeverReset && Owner)
+		if (Effect.Spec.Def->PeriodicInhibitionPolicy != EGameplayEffectPeriodInhibitionRemovedPolicy::NeverReset && Owner && Owner->IsOwnerActorAuthoritative())
 		{
 			FTimerManager& TimerManager = Owner->GetWorld()->GetTimerManager();
 			FTimerDelegate Delegate = FTimerDelegate::CreateUObject(Owner, &UAbilitySystemComponent::ExecutePeriodicEffect, Effect.Handle);
