Reorder code in PreActivate so that we set the current ability info before applying tags and potentially blocking and canceling other abilities. This avoids potential bugs where you end up in a loop that calls back into the current ability and tries to do something that needs the current info set.

diff --git a/Engine/Plugins/Runtime/GameplayAbilities/Source/GameplayAbilities/Private/Abilities/GameplayAbility.cpp b/Engine/Plugins/Runtime/GameplayAbilities/Source/GameplayAbilities/Private/Abilities/GameplayAbility.cpp
--- a/Engine/Plugins/Runtime/GameplayAbilities/Source/GameplayAbilities/Private/Abilities/GameplayAbility.cpp
+++ b/Engine/Plugins/Runtime/GameplayAbilities/Source/GameplayAbilities/Private/Abilities/GameplayAbility.cpp
@@ -745,6 +745,10 @@ void UGameplayAbility::PreActivate(const FGameplayAbilitySpecHandle Handle, cons
 
 	RemoteInstanceEnded = false;
 
+	// This must be called before we start applying tags and blocking or canceling other abilities.
+	// We could set off a chain that results in calling functions on this ability that rely on the current info being set.
+	SetCurrentInfo(Handle, ActorInfo, ActivationInfo);
+
 	Comp->HandleChangeAbilityCanBeCanceled(AbilityTags, this, true);
 	Comp->ApplyAbilityBlockAndCancelTags(AbilityTags, this, true, BlockAbilitiesWithTag, true, CancelAbilitiesWithTag);
 	Comp->AddLooseGameplayTags(ActivationOwnedTags);
@@ -754,8 +758,6 @@ void UGameplayAbility::PreActivate(const FGameplayAbilitySpecHandle Handle, cons
 		OnGameplayAbilityEnded.Add(*OnGameplayAbilityEndedDelegate);
 	}
 
-	SetCurrentInfo(Handle, ActorInfo, ActivationInfo);
-
 	Comp->NotifyAbilityActivated(Handle, this);
 }
 
