Remove old temporary log in GameplayAbility.cpp that can be spammy.

diff --git a/Engine/Plugins/Runtime/GameplayAbilities/Source/GameplayAbilities/Private/GameplayEffect.cpp b/Engine/Plugins/Runtime/GameplayAbilities/Source/GameplayAbilities/Private/GameplayEffect.cpp
--- a/Engine/Plugins/Runtime/GameplayAbilities/Source/GameplayAbilities/Private/GameplayEffect.cpp
+++ b/Engine/Plugins/Runtime/GameplayAbilities/Source/GameplayAbilities/Private/GameplayEffect.cpp
@@ -3122,8 +3122,6 @@ void FActiveGameplayEffectsContainer::AddActiveGameplayEffectGrantedTagsAndModif
 				AbilitySpecDef.SetByCallerTagMagnitudes = Effect.Spec.SetByCallerTagMagnitudes;
 
 				Owner->GiveAbility( FGameplayAbilitySpec(AbilitySpecDef, Effect.Spec.GetLevel(), Effect.Handle) );
-
-				ABILITY_LOG(Display, TEXT("::AddActiveGameplayEffectGrantedTagsAndModifiers granted ability %s (Handle %s) from GE %s (Handle: %s)"), *GetNameSafe(AbilitySpecDef.Ability), *AbilitySpecDef.AssignedHandle.ToString(), *Effect.GetDebugString(), *Effect.Handle.ToString());
 			}
 		}	
 	}
