When removing many GameplayTags, call FillParentTags only once.  This adds up for very large GameplayTagContainers.  ~9X faster

diff --git a/Engine/Plugins/Runtime/GameplayAbilities/Source/GameplayAbilities/Private/GameplayEffectTypes.cpp b/Engine/Plugins/Runtime/GameplayAbilities/Source/GameplayAbilities/Private/GameplayEffectTypes.cpp
--- a/Engine/Plugins/Runtime/GameplayAbilities/Source/GameplayAbilities/Private/GameplayEffectTypes.cpp
+++ b/Engine/Plugins/Runtime/GameplayAbilities/Source/GameplayAbilities/Private/GameplayEffectTypes.cpp
@@ -417,7 +417,7 @@ void FGameplayTagCountContainer::Reset()
 	OnAnyTagChangeDelegate.Clear();
 }
 
-bool FGameplayTagCountContainer::UpdateTagMap_Internal(const FGameplayTag& Tag, int32 CountDelta)
+bool FGameplayTagCountContainer::UpdateTagMap_Internal(const FGameplayTag& Tag, int32 CountDelta, bool bDeferParentTagsOnRemove)
 {
 	const bool bTagAlreadyExplicitlyExists = ExplicitTags.HasTagExact(Tag);
 
@@ -451,7 +451,7 @@ bool FGameplayTagCountContainer::UpdateTagMap_Internal(const FGameplayTag& Tag,
 	if (ExistingCount <= 0)
 	{
 		// Remove from the explicit list
-		ExplicitTags.RemoveTag(Tag);
+		ExplicitTags.RemoveTag(Tag, bDeferParentTagsOnRemove);
 	}
 
 	// Check if change delegates are required to fire for the tag or any of its parents based on the count change
diff --git a/Engine/Plugins/Runtime/GameplayAbilities/Source/GameplayAbilities/Public/AbilitySystemComponent.h b/Engine/Plugins/Runtime/GameplayAbilities/Source/GameplayAbilities/Public/AbilitySystemComponent.h
--- a/Engine/Plugins/Runtime/GameplayAbilities/Source/GameplayAbilities/Public/AbilitySystemComponent.h
+++ b/Engine/Plugins/Runtime/GameplayAbilities/Source/GameplayAbilities/Public/AbilitySystemComponent.h
@@ -539,10 +539,43 @@ class GAMEPLAYABILITIES_API UAbilitySystemComponent : public UGameplayTasksCompo
 	/** Update the number of instances of a given tag and calls callback */
 	FORCEINLINE void UpdateTagMap(const FGameplayTagContainer& Container, int32 CountDelta)
 	{
-		for (auto TagIt = Container.CreateConstIterator(); TagIt; ++TagIt)
+		// For removal, reorder calls so that FillParentTags is only called once
+		if (CountDelta > 0)
 		{
-			const FGameplayTag& Tag = *TagIt;
-			UpdateTagMap(Tag, CountDelta);
+			for (auto TagIt = Container.CreateConstIterator(); TagIt; ++TagIt)
+			{
+				const FGameplayTag& Tag = *TagIt;
+				if (GameplayTagCountContainer.UpdateTagCount(Tag, CountDelta))
+				{
+					OnTagUpdated(Tag, true);
+				}
+			}
+		}
+		else if (CountDelta < 0)
+		{
+			// Defer FillParentTags until all Tags have been removed
+			TArray<FGameplayTag> RemovedTags;
+			RemovedTags.Reserve(Container.Num()); // pre-allocate max number (if all are removed)
+
+			for (auto TagIt = Container.CreateConstIterator(); TagIt; ++TagIt)
+			{
+				const FGameplayTag& Tag = *TagIt;
+				if (GameplayTagCountContainer.UpdateTagCount(Tag, CountDelta, true))
+				{
+					RemovedTags.Add(Tag);
+				}
+			}
+			
+			if (RemovedTags.Num() > 0)
+			{
+				GameplayTagCountContainer.FillParentTags();
+			}
+
+			// Notify last in case OnTagUpdated queries this container
+			for (FGameplayTag& Tag : RemovedTags)
+			{
+				OnTagUpdated(Tag, false);
+			}
 		}
 	}
 
diff --git a/Engine/Plugins/Runtime/GameplayAbilities/Source/GameplayAbilities/Public/GameplayEffectTypes.h b/Engine/Plugins/Runtime/GameplayAbilities/Source/GameplayAbilities/Public/GameplayEffectTypes.h
--- a/Engine/Plugins/Runtime/GameplayAbilities/Source/GameplayAbilities/Public/GameplayEffectTypes.h
+++ b/Engine/Plugins/Runtime/GameplayAbilities/Source/GameplayAbilities/Public/GameplayEffectTypes.h
@@ -1196,9 +1196,15 @@ struct GAMEPLAYABILITIES_API FGameplayTagCountContainer
 	{
 		if (CountDelta != 0)
 		{
+			bool bUpdatedAny = false;
 			for (auto TagIt = Container.CreateConstIterator(); TagIt; ++TagIt)
 			{
-				UpdateTagMap_Internal(*TagIt, CountDelta);
+				bUpdatedAny |= UpdateTagMap_Internal(*TagIt, CountDelta, true);
+			}
+
+			if (bUpdatedAny && CountDelta < 0)
+			{
+				ExplicitTags.FillParentTags();
 			}
 		}
 	}
@@ -1206,16 +1212,17 @@ struct GAMEPLAYABILITIES_API FGameplayTagCountContainer
 	/**
 	 * Update the specified tag by the specified delta, potentially causing an additional or removal from the explicit tag list
 	 * 
-	 * @param Tag			Tag to update
-	 * @param CountDelta	Delta of the tag count to apply
+	 * @param Tag						Tag to update
+	 * @param CountDelta				Delta of the tag count to apply
+	 * @param bDeferParentTagsOnRemove	Skip calling FillParentTags for performance (must be handled by calling code)
 	 * 
 	 * @return True if tag was *either* added or removed. (E.g., we had the tag and now dont. or didnt have the tag and now we do. We didn't just change the count (1 count -> 2 count would return false).
 	 */
-	FORCEINLINE bool UpdateTagCount(const FGameplayTag& Tag, int32 CountDelta)
+	FORCEINLINE bool UpdateTagCount(const FGameplayTag& Tag, int32 CountDelta, bool bDeferParentTagsOnRemove=false)
 	{
 		if (CountDelta != 0)
 		{
-			return UpdateTagMap_Internal(Tag, CountDelta);
+			return UpdateTagMap_Internal(Tag, CountDelta, bDeferParentTagsOnRemove);
 		}
 
 		return false;
@@ -1296,6 +1303,12 @@ struct GAMEPLAYABILITIES_API FGameplayTagCountContainer
 
 	void Reset();
 
+	/** Fills in ParentTags from GameplayTags */
+	void FillParentTags()
+	{
+		ExplicitTags.FillParentTags();
+	}
+
 private:
 
 	struct FDelegateInfo
@@ -1320,7 +1333,7 @@ private:
 	FGameplayTagContainer ExplicitTags;
 
 	/** Internal helper function to adjust the explicit tag list & corresponding maps/delegates/etc. as necessary */
-	bool UpdateTagMap_Internal(const FGameplayTag& Tag, int32 CountDelta);
+	bool UpdateTagMap_Internal(const FGameplayTag& Tag, int32 CountDelta, bool bDeferParentTagsOnRemove=false);
 };
 
 
diff --git a/Engine/Source/Runtime/GameplayTags/Classes/GameplayTagContainer.h b/Engine/Source/Runtime/GameplayTags/Classes/GameplayTagContainer.h
--- a/Engine/Source/Runtime/GameplayTags/Classes/GameplayTagContainer.h
+++ b/Engine/Source/Runtime/GameplayTags/Classes/GameplayTagContainer.h
@@ -542,9 +542,10 @@ struct GAMEPLAYTAGS_API FGameplayTagContainer
 	/**
 	 * Tag to remove from the container
 	 * 
-	 * @param TagToRemove	Tag to remove from the container
+	 * @param TagToRemove		Tag to remove from the container
+	 * @param bDeferParentTags	Skip calling FillParentTags for performance (must be handled by calling code)
 	 */
-	bool RemoveTag(FGameplayTag TagToRemove);
+	bool RemoveTag(FGameplayTag TagToRemove, bool bDeferParentTags=false);
 
 	/**
 	 * Removes all tags in TagsToRemove from this container
@@ -616,6 +617,9 @@ struct GAMEPLAYTAGS_API FGameplayTagContainer
 		return GameplayTags.Num() > 0 ? GameplayTags.Last() : FGameplayTag();
 	}
 
+	/** Fills in ParentTags from GameplayTags */
+	void FillParentTags();
+
 	/** An empty Gameplay Tag Container */
 	static const FGameplayTagContainer EmptyContainer;
 
@@ -824,9 +828,6 @@ protected:
 	/** Adds parent tags for a single tag */
 	void AddParentsForTag(const FGameplayTag& Tag);
 
-	/** Fills in ParentTags from GameplayTags */
-	void FillParentTags();
-
 	/** Array of gameplay tags */
 	UPROPERTY(BlueprintReadWrite, Category=GameplayTags, SaveGame) // Change to VisibleAnywhere after fixing up games
 	TArray<FGameplayTag> GameplayTags;
diff --git a/Engine/Source/Runtime/GameplayTags/Private/GameplayTagContainer.cpp b/Engine/Source/Runtime/GameplayTags/Private/GameplayTagContainer.cpp
--- a/Engine/Source/Runtime/GameplayTags/Private/GameplayTagContainer.cpp
+++ b/Engine/Source/Runtime/GameplayTags/Private/GameplayTagContainer.cpp
@@ -821,7 +821,7 @@ bool FGameplayTagContainer::AddLeafTag(const FGameplayTag& TagToAdd)
 
 DECLARE_CYCLE_STAT(TEXT("FGameplayTagContainer::RemoveTag"), STAT_FGameplayTagContainer_RemoveTag, STATGROUP_GameplayTags);
 
-bool FGameplayTagContainer::RemoveTag(FGameplayTag TagToRemove)
+bool FGameplayTagContainer::RemoveTag(FGameplayTag TagToRemove, bool bDeferParentTags)
 {
 	SCOPE_CYCLE_COUNTER(STAT_FGameplayTagContainer_RemoveTag);
 
@@ -829,8 +829,11 @@ bool FGameplayTagContainer::RemoveTag(FGameplayTag TagToRemove)
 
 	if (NumChanged > 0)
 	{
-		// Have to recompute parent table from scratch because there could be duplicates providing the same parent tag
-		FillParentTags();
+		if (!bDeferParentTags)
+		{
+			// Have to recompute parent table from scratch because there could be duplicates providing the same parent tag
+			FillParentTags();
+		}
 		return true;
 	}
 	return false;
