PR #5769: Added constructor to prevent using of uninitialized memory (Contributed by [REDACTED_PERSON])

diff --git a/Engine/Plugins/Runtime/GameplayAbilities/Source/GameplayAbilities/Public/GameplayCueTranslator.h b/Engine/Plugins/Runtime/GameplayAbilities/Source/GameplayAbilities/Public/GameplayCueTranslator.h
--- a/Engine/Plugins/Runtime/GameplayAbilities/Source/GameplayAbilities/Public/GameplayCueTranslator.h
+++ b/Engine/Plugins/Runtime/GameplayAbilities/Source/GameplayAbilities/Public/GameplayCueTranslator.h
@@ -57,6 +57,11 @@ struct FGameplayCueParameters;
 #if WITH_EDITOR
 struct GAMEPLAYABILITIES_API FGameplayCueTranslationEditorOnlyData
 {
+	FGameplayCueTranslationEditorOnlyData()
+		:UniqueID(0), Enabled(true)
+	{
+	}
+
 	FName	EditorDescription;	// For pretty/custom printing
 	FString ToolTip;			// additional info for tooltip text (exactly where did this rule come from?)
 	int32	UniqueID;			// For filtering overrides by translation rule. Set by the GameplaycueTranslator.
