Added new aggregator stat and a helper #define for FScopedAggregatorOnDirtyBatch

#rnx

diff --git a/Engine/Plugins/Runtime/GameplayAbilities/Source/GameplayAbilities/Private/AbilitySystemStats.cpp b/Engine/Plugins/Runtime/GameplayAbilities/Source/GameplayAbilities/Private/AbilitySystemStats.cpp
--- a/Engine/Plugins/Runtime/GameplayAbilities/Source/GameplayAbilities/Private/AbilitySystemStats.cpp
+++ b/Engine/Plugins/Runtime/GameplayAbilities/Source/GameplayAbilities/Private/AbilitySystemStats.cpp
@@ -23,6 +23,7 @@ DEFINE_STAT(STAT_InitAttributeSetDefaults);
 DEFINE_STAT(STAT_TickAbilityTasks);
 DEFINE_STAT(STAT_FindAbilitySpecFromHandle);
 DEFINE_STAT(STAT_AggregatorEvaluate);
+DEFINE_STAT(STAT_AggregatorBroadcastOnDirty);
 DEFINE_STAT(STAT_HasApplicationImmunityToSpec);
 DEFINE_STAT(STAT_HasMatchingGameplayTag);
 DEFINE_STAT(STAT_HandleGameplayCueNotifyStatic);
diff --git a/Engine/Plugins/Runtime/GameplayAbilities/Source/GameplayAbilities/Private/GameplayEffectAggregator.cpp b/Engine/Plugins/Runtime/GameplayAbilities/Source/GameplayAbilities/Private/GameplayEffectAggregator.cpp
--- a/Engine/Plugins/Runtime/GameplayAbilities/Source/GameplayAbilities/Private/GameplayEffectAggregator.cpp
+++ b/Engine/Plugins/Runtime/GameplayAbilities/Source/GameplayAbilities/Private/GameplayEffectAggregator.cpp
@@ -4,6 +4,7 @@
 #include "UObject/UObjectHash.h"
 #include "UObject/UObjectIterator.h"
 #include "AbilitySystemComponent.h"
+#include "AbilitySystemStats.h"
 
 void FAggregatorMod::UpdateQualifies(const FAggregatorEvaluateParameters& Parameters) const
 {
@@ -567,6 +568,7 @@ void FAggregator::BroadcastOnDirty()
 		return;
 	}
 
+	SCOPE_CYCLE_COUNTER(STAT_AggregatorBroadcastOnDirty);
 
 	// --------------------------------------------------
 	//	The code below attempts to avoid recursion issues: an aggregator is dirty and while it is broadcasting this out, someone dirties it again.
diff --git a/Engine/Plugins/Runtime/GameplayAbilities/Source/GameplayAbilities/Public/AbilitySystemStats.h b/Engine/Plugins/Runtime/GameplayAbilities/Source/GameplayAbilities/Public/AbilitySystemStats.h
--- a/Engine/Plugins/Runtime/GameplayAbilities/Source/GameplayAbilities/Public/AbilitySystemStats.h
+++ b/Engine/Plugins/Runtime/GameplayAbilities/Source/GameplayAbilities/Public/AbilitySystemStats.h
@@ -28,6 +28,7 @@ DECLARE_CYCLE_STAT_EXTERN(TEXT("InitAttributeSetDefaults"), STAT_InitAttributeSe
 DECLARE_CYCLE_STAT_EXTERN(TEXT("TickAbilityTasks"), STAT_TickAbilityTasks, STATGROUP_AbilitySystem, );
 DECLARE_CYCLE_STAT_EXTERN(TEXT("FindAbilitySpecFromHandle"), STAT_FindAbilitySpecFromHandle, STATGROUP_AbilitySystem, );
 DECLARE_CYCLE_STAT_EXTERN(TEXT("Aggregator Evaluate"), STAT_AggregatorEvaluate, STATGROUP_AbilitySystem, );
+DECLARE_CYCLE_STAT_EXTERN(TEXT("Aggregator BroadcastOnDirty"), STAT_AggregatorBroadcastOnDirty, STATGROUP_AbilitySystem, );
 DECLARE_CYCLE_STAT_EXTERN(TEXT("Has Application Immunity To Spec"), STAT_HasApplicationImmunityToSpec, STATGROUP_AbilitySystem, );
 DECLARE_CYCLE_STAT_EXTERN(TEXT("Has Matching GameplayTag"), STAT_HasMatchingGameplayTag, STATGROUP_AbilitySystem, );
 DECLARE_CYCLE_STAT_EXTERN(TEXT("GameplayCueNotify Static"), STAT_HandleGameplayCueNotifyStatic, STATGROUP_AbilitySystem, );
diff --git a/Engine/Plugins/Runtime/GameplayAbilities/Source/GameplayAbilities/Public/GameplayEffectAggregator.h b/Engine/Plugins/Runtime/GameplayAbilities/Source/GameplayAbilities/Public/GameplayEffectAggregator.h
--- a/Engine/Plugins/Runtime/GameplayAbilities/Source/GameplayAbilities/Public/GameplayEffectAggregator.h
+++ b/Engine/Plugins/Runtime/GameplayAbilities/Source/GameplayAbilities/Public/GameplayEffectAggregator.h
@@ -416,3 +416,5 @@ struct GAMEPLAYABILITIES_API FScopedAggregatorOnDirtyBatch
 	static bool		GlobalFromNetworkUpdate;
 	static int32	NetUpdateID;
 };
+
+#define AGGREGATOR_BATCH_SCOPE() FScopedAggregatorOnDirtyBatch AggregatorOnDirtyBatcher;
