Made the gameplay cue editor use SAddNewGameplayTagWidget.
Added error checking for input strings to SAddNewGameplayTagWidget.

diff --git a/Engine/Plugins/Editor/GameplayTagsEditor/Source/GameplayTagsEditor/Private/SAddNewGameplayTagWidget.cpp b/Engine/Plugins/Editor/GameplayTagsEditor/Source/GameplayTagsEditor/Private/SAddNewGameplayTagWidget.cpp
--- a/Engine/Plugins/Editor/GameplayTagsEditor/Source/GameplayTagsEditor/Private/SAddNewGameplayTagWidget.cpp
+++ b/Engine/Plugins/Editor/GameplayTagsEditor/Source/GameplayTagsEditor/Private/SAddNewGameplayTagWidget.cpp
@@ -32,6 +32,7 @@ void SAddNewGameplayTagWidget::Construct(const FArguments& InArgs)
 	bShouldGetKeyboardFocus = false;
 
 	OnGameplayTagAdded = InArgs._OnGameplayTagAdded;
+	IsValidTag = InArgs._IsValidTag;
 	PopulateTagSources();
 
 	IGameplayTagsModule::OnTagSettingsChanged.AddRaw(this, &SAddNewGameplayTagWidget::PopulateTagSources);
@@ -244,7 +245,8 @@ void SAddNewGameplayTagWidget::CreateNewGameplayTag()
 		return;
 	}
 
-	FString TagName = TagNameTextBox->GetText().ToString();
+	FText TagNameAsText = TagNameTextBox->GetText();
+	FString TagName = TagNameAsText.ToString();
 	FString TagComment = TagCommentTextBox->GetText().ToString();
 	FName TagSource = *TagSourcesComboBox->GetSelectedItem().Get();
 
@@ -253,6 +255,18 @@ void SAddNewGameplayTagWidget::CreateNewGameplayTag()
 		return;
 	}
 
+	// check to see if this is a valid tag
+	// first check the base rules for all tags then look for any additional rules in the delegate
+	FText ErrorMsg;
+	if (!UGameplayTagsManager::Get().IsValidGameplayTagString(TagName, &ErrorMsg) || 
+		(IsValidTag.IsBound() && !IsValidTag.Execute(TagName, &ErrorMsg))
+		)
+	{
+		FText MessageTitle(LOCTEXT("InvalidTag", "Invalid Tag"));
+		FMessageDialog::Open(EAppMsgType::Ok, ErrorMsg, &MessageTitle);
+		return;
+	}
+
 	// set bIsAddingNewTag, this guards against the window closing when it loses focus due to source control checking out a file
 	TGuardValue<bool>	Guard(bAddingNewTag, true);
 
diff --git a/Engine/Plugins/Editor/GameplayTagsEditor/Source/GameplayTagsEditor/Private/SAddNewGameplayTagWidget.h b/Engine/Plugins/Editor/GameplayTagsEditor/Source/GameplayTagsEditor/Private/SAddNewGameplayTagWidget.h
--- a/Engine/Plugins/Editor/GameplayTagsEditor/Source/GameplayTagsEditor/Private/SAddNewGameplayTagWidget.h
+++ b/Engine/Plugins/Editor/GameplayTagsEditor/Source/GameplayTagsEditor/Private/SAddNewGameplayTagWidget.h
@@ -16,20 +16,23 @@ class SAddNewGameplayTagWidget : public SCompoundWidget
 {
 public:
 
-	DECLARE_DELEGATE_ThreeParams( FOnGameplayTagAdded, const FString& /*TagName*/, const FString& /*TagComment*/, const FName& /*TagSource*/);
+	DECLARE_DELEGATE_ThreeParams(FOnGameplayTagAdded, const FString& /*TagName*/, const FString& /*TagComment*/, const FName& /*TagSource*/);
+
+	DECLARE_DELEGATE_RetVal_TwoParams(bool, FIsValidTag, const FString& /*TagName*/, FText* /*OutError*/)
 
 	SLATE_BEGIN_ARGS(SAddNewGameplayTagWidget)
 		: _NewTagName(TEXT(""))
-		{}
-		SLATE_EVENT( FOnGameplayTagAdded, OnGameplayTagAdded )	// Callback for when a new tag is added	
-		SLATE_ARGUMENT( FString, NewTagName ) // String that will initially populate the New Tag Name field
+	{}
+		SLATE_EVENT(FOnGameplayTagAdded, OnGameplayTagAdded)	// Callback for when a new tag is added	
+		SLATE_EVENT(FIsValidTag, IsValidTag)
+		SLATE_ARGUMENT(FString, NewTagName) // String that will initially populate the New Tag Name field
 	SLATE_END_ARGS();
 
-	virtual ~SAddNewGameplayTagWidget();
+	GAMEPLAYTAGSEDITOR_API virtual ~SAddNewGameplayTagWidget();
 
-	virtual void Tick( const FGeometry& AllottedGeometry, const double InCurrentTime, const float InDeltaTime ) override;
+	GAMEPLAYTAGSEDITOR_API virtual void Tick(const FGeometry& AllottedGeometry, const double InCurrentTime, const float InDeltaTime) override;
 
-	void Construct( const FArguments& InArgs);
+	GAMEPLAYTAGSEDITOR_API void Construct(const FArguments& InArgs);
 
 	/** Returns true if we're currently attempting to add a new gameplay tag to an INI file */
 	bool IsAddingNewTag() const
@@ -86,6 +89,9 @@ private:
 	/** Callback for when a new gameplay tag has been added to the INI files */
 	FOnGameplayTagAdded OnGameplayTagAdded;
 
+	/** Callback to see if the gameplay tag is valid. This should be used for any specialized rules that are not covered by IsValidGameplayTagString */
+	FIsValidTag IsValidTag;
+
 	bool bAddingNewTag;
 
 	/** Tracks if this widget should get keyboard focus */
diff --git a/Engine/Plugins/Runtime/GameplayAbilities/Source/GameplayAbilitiesEditor/Private/SGameplayCueEditor.cpp b/Engine/Plugins/Runtime/GameplayAbilities/Source/GameplayAbilitiesEditor/Private/SGameplayCueEditor.cpp
--- a/Engine/Plugins/Runtime/GameplayAbilities/Source/GameplayAbilitiesEditor/Private/SGameplayCueEditor.cpp
+++ b/Engine/Plugins/Runtime/GameplayAbilities/Source/GameplayAbilitiesEditor/Private/SGameplayCueEditor.cpp
@@ -50,12 +50,14 @@
 #include "Editor/LevelEditor/Public/LevelEditor.h"
 #include "AssetRegistryModule.h"
 #include "Framework/MultiBox/MultiBoxBuilder.h"
+#include "GameplayTagsEditor/Private/SAddNewGameplayTagWidget.h"
 
 #include "SGameplayCueEditor_Picker.h"
 #define LOCTEXT_NAMESPACE "SGameplayCueEditor"
 
 static const FName CueTagColumnName("GameplayCueTags");
 static const FName CueHandlerColumnName("GameplayCueHandlers");
+static const FString CueTagPrefix(TEXT("GameplayCue."));
 
 // Whether to show the Hotreload button in the GC editor.
 #define GAMEPLAYCUEEDITOR_HOTRELOAD_BUTTON 0
@@ -110,10 +112,21 @@ public:
 
 	virtual void Construct(const FArguments& InArgs) override;
 
+	/** Helper function to determine the visibility of the Add New Tag widget */
+	EVisibility DetermineAddNewTagWidgetVisibility() const;
+
+	/** Helper functions to expand or collapse the add new tag widget */
+	EVisibility DetermineExpandableUIVisibility() const;
+	ECheckBoxState GetAddTagSectionExpansionState() const;
+	void OnAddTagSectionExpansionStateChanged(ECheckBoxState NewState);
+
 private:
 
 	// ---------------------------------------------------------------------
 
+	/* Expand the Add Tag section (if it's visible) */
+	bool bAddTagSectionExpanded;
+
 	/** Show all GC Tags, even ones without handlers */
 	bool bShowAll;
 
@@ -129,9 +142,6 @@ private:
 	/** Global flag suppress rebuilding cue tree view. Needed when doing operations that would rebuld it multiple times */
 	static bool bSuppressCueViewUpdate;
 
-	/** Text box for creating new GC tag */
-	TSharedPtr<SEditableTextBox> NewGameplayCueTextBox;
-
 	/** Main widgets that shows the gameplay cue tree */
 	TSharedPtr<SGameplayCueTreeView> GameplayCueTreeView;
 
@@ -168,20 +178,36 @@ private:
 	/** For tracking expanded tags in between recreation of GC view */
 	TSet<FName>	ExpandedTags;
 
+	/** The widget that controls how new gameplay tags are added to the config files */
+	TSharedPtr<class SAddNewGameplayTagWidget> AddNewTagWidget;
+
 public:
 
-	virtual void OnNewGameplayCueTagCommited(const FText& InText, ETextCommit::Type InCommitType) override
+	virtual void OnNewGameplayCueTagCommited(const FString& TagName, const FString& TagComment, const FName& TagSource) override
 	{
-		// Only support adding tags via ini file
-		if (UGameplayTagsManager::Get().ShouldImportTagsFromINI() == false)
-		{
-			return;
-		}
+		FScopedSlowTask SlowTask(0.f, LOCTEXT("AddingNewGameplaycue", "Adding new GameplayCue Tag"));
+		SlowTask.MakeDialog();
 
-		if (InCommitType == ETextCommit::OnEnter)
+		check(!TagName.IsEmpty());
+
+		TSharedPtr<FGameplayTagNode> TagNode = UGameplayTagsManager::Get().FindTagNode(FName(*TagName));
+		const FGameplayTag* NewTag = TagNode.IsValid() ? &TagNode->GetCompleteTag() : nullptr;
+
+		UpdateGameplayCueListItems(NewTag);
+	}
+
+	bool IsValidTag(const FString& TagName, FText* ErrorMsg) const
+	{
+		if (!TagName.StartsWith(CueTagPrefix))
 		{
-			CreateNewGameplayCueTag();
+			if (ErrorMsg)
+			{
+				*ErrorMsg = FText::Format(LOCTEXT("BadPrefixError", "Gameplay cue tags must start with {0}"), FText::FromString(CueTagPrefix));
+			}
+			return false;
 		}
+
+		return true;
 	}
 
 	virtual void OnSearchTagCommited(const FText& InText, ETextCommit::Type InCommitType) override
@@ -203,35 +229,6 @@ public:
 	}
 
 
-	virtual FReply OnNewGameplayCueButtonPressed() override
-	{
-		CreateNewGameplayCueTag();
-		return FReply::Handled();
-	}
-
-	
-	/** Checks out config file, adds new tag, repopulates widget cue list */
-	void CreateNewGameplayCueTag()
-	{
-		FScopedSlowTask SlowTask(0.f, LOCTEXT("AddingNewGameplaycue", "Adding new GameplayCue Tag"));
-		SlowTask.MakeDialog();
-
-		FString str = NewGameplayCueTextBox->GetText().ToString();
-		if (str.IsEmpty())
-		{
-			return;
-		}
-
-		SelectedTag = FName(*str);
-		SelectedUniqueID = 0;
-
-		IGameplayTagsEditorModule::Get().AddNewGameplayTagToINI(str);
-
-		UpdateGameplayCueListItems();
-
-		NewGameplayCueTextBox->SetText(FText::GetEmpty());
-	}
-
 	void OnFilterMenuOpenChanged(bool bOpen)
 	{
 		if (bOpen == false && bFilterIDsDirty)
@@ -462,8 +459,13 @@ public:
 		}
 	}
 
-	/** Builds content of the list in the GameplayCue Editor */
 	void UpdateGameplayCueListItems()
+	{
+		UpdateGameplayCueListItems(nullptr);
+	}
+
+	/** Builds content of the list in the GameplayCue Editor */
+	void UpdateGameplayCueListItems(const FGameplayTag* SelectedTagPostUpdate)
 	{
 		if (bSuppressCueViewUpdate)
 		{
@@ -586,8 +588,16 @@ public:
 			}
 
 			// ----------------------------------------------------------------
+			
+			// if we were told to specifically select an item we want to force it to show up
+			bool bForcedShow = false;
+			if (SelectedTagPostUpdate && SelectedTagPostUpdate->MatchesTagExact(NewItem->GameplayCueTag))
+			{
+				SelectedItem = NewItem;
+				bForcedShow = true;
+			}
 
-			if (!FilteredOut && (bShowAll || Handled))
+			if ((!FilteredOut && (bShowAll || Handled)) || bForcedShow)
 			{
 				GameplayCueListItems.Add(NewItem);
 			}
@@ -612,10 +622,6 @@ public:
 				GameplayCueTreeView->SetItemSelection(SelectedItem, true);
 				GameplayCueTreeView->RequestScrollIntoView(SelectedItem);
 			}
-
-			//UE_LOG(LogTemp, Display, TEXT("Accumulated FindGameplayCueNotifyObjTime: %.4f"), FindGameplayCueNotifyObjTime);
-			//UE_LOG(LogTemp, Display, TEXT("Accumulated FindGameplayCueNotifyObjTime: %.4f"), AddTranslationTagsTime);
-			//UE_LOG(LogTemp, Display, TEXT("Accumulated FindGameplayCueNotifyObjTime: %.4f"), AddEventsTime);
 		}
 	}
 
@@ -1160,23 +1166,38 @@ void SGameplayCueEditorImpl::Construct(const FArguments& InArgs)
 		[
 			SNew(SHorizontalBox)
 			+ SHorizontalBox::Slot()
-			.Padding(2.0f, 2.0f)
 			.AutoWidth()
 			[
-				SAssignNew(NewGameplayCueTextBox, SEditableTextBox)
-				.MinDesiredWidth(210.0f)
-				.HintText(LOCTEXT("GameplayCueXY", "GameplayCue.X.Y"))
-				.OnTextCommitted(this, &SGameplayCueEditorImpl::OnNewGameplayCueTagCommited)
+				SNew(SCheckBox)
+				.IsChecked(this, &SGameplayCueEditorImpl::GetAddTagSectionExpansionState)
+				.OnCheckStateChanged(this, &SGameplayCueEditorImpl::OnAddTagSectionExpansionStateChanged)
+				.CheckedImage(FEditorStyle::GetBrush("TreeArrow_Expanded"))
+				.CheckedHoveredImage(FEditorStyle::GetBrush("TreeArrow_Expanded_Hovered"))
+				.CheckedPressedImage(FEditorStyle::GetBrush("TreeArrow_Expanded"))
+				.UncheckedImage(FEditorStyle::GetBrush("TreeArrow_Collapsed"))
+				.UncheckedHoveredImage(FEditorStyle::GetBrush("TreeArrow_Collapsed_Hovered"))
+				.UncheckedPressedImage(FEditorStyle::GetBrush("TreeArrow_Collapsed"))
+				.Visibility(this, &SGameplayCueEditorImpl::DetermineExpandableUIVisibility)
+				[
+					SNew(STextBlock)
+					.Text(LOCTEXT("AddNewTag", "Add New Gameplay Cue Tag"))
+				]
 			]
+		]
 
+		+ SVerticalBox::Slot()
+		.AutoHeight()
+		.Padding(16.0f, 0.0f)
+		[
+			SNew(SHorizontalBox)
 			+ SHorizontalBox::Slot()
-			.Padding(2.0f, 2.0f)
 			.AutoWidth()
 			[
-				SNew(SButton)
-				.Text(LOCTEXT("AddNew", "Add New"))
-				.OnClicked(this, &SGameplayCueEditorImpl::OnNewGameplayCueButtonPressed)
-				.Visibility(CanAddFromINI ? EVisibility::Visible : EVisibility::Collapsed)
+				SAssignNew(AddNewTagWidget, SAddNewGameplayTagWidget)
+				.Visibility(this, &SGameplayCueEditorImpl::DetermineAddNewTagWidgetVisibility)
+				.OnGameplayTagAdded(this, &SGameplayCueEditorImpl::OnNewGameplayCueTagCommited)
+				.IsValidTag(this, &SGameplayCueEditorImpl::IsValidTag)
+				.NewTagName(CueTagPrefix)
 			]
 		]
 
@@ -1321,6 +1342,41 @@ void SGameplayCueEditorImpl::Construct(const FArguments& InArgs)
 
 END_SLATE_FUNCTION_BUILD_OPTIMIZATION
 
+EVisibility SGameplayCueEditorImpl::DetermineAddNewTagWidgetVisibility() const
+{
+	UGameplayTagsManager& Manager = UGameplayTagsManager::Get();
+
+	if (!Manager.ShouldImportTagsFromINI() || !bAddTagSectionExpanded)
+	{
+		// If we can't support adding tags from INI files, we should never see this widget
+		return EVisibility::Collapsed;
+	}
+
+	return EVisibility::Visible;
+}
+
+EVisibility SGameplayCueEditorImpl::DetermineExpandableUIVisibility() const
+{
+	UGameplayTagsManager& Manager = UGameplayTagsManager::Get();
+
+	if (!Manager.ShouldImportTagsFromINI())
+	{
+		// If we can't support adding tags from INI files, we should never see this widget
+		return EVisibility::Collapsed;
+	}
+
+	return EVisibility::Visible;
+}
+
+ECheckBoxState SGameplayCueEditorImpl::GetAddTagSectionExpansionState() const
+{
+	return bAddTagSectionExpanded ? ECheckBoxState::Checked : ECheckBoxState::Unchecked;
+}
+
+void SGameplayCueEditorImpl::OnAddTagSectionExpansionStateChanged(ECheckBoxState NewState)
+{
+	bAddTagSectionExpanded = NewState == ECheckBoxState::Checked;
+}
 
 // -----------------------------------------------------------------------------------------------------------
 //
diff --git a/Engine/Plugins/Runtime/GameplayAbilities/Source/GameplayAbilitiesEditor/Private/SGameplayCueEditor.h b/Engine/Plugins/Runtime/GameplayAbilities/Source/GameplayAbilitiesEditor/Private/SGameplayCueEditor.h
--- a/Engine/Plugins/Runtime/GameplayAbilities/Source/GameplayAbilitiesEditor/Private/SGameplayCueEditor.h
+++ b/Engine/Plugins/Runtime/GameplayAbilities/Source/GameplayAbilitiesEditor/Private/SGameplayCueEditor.h
@@ -7,7 +7,7 @@
 #include "Input/Reply.h"
 #include "Widgets/SUserWidget.h"
 
-/** Main CollisionAnalyzer UI widget */
+/** Main Gameplay Cue Editor UI widget */
 class SGameplayCueEditor : public SUserWidget
 {
 public:
@@ -16,12 +16,11 @@ public:
 	SLATE_END_ARGS()
 	virtual void Construct(const FArguments& InArgs) = 0;
 
-	virtual void OnNewGameplayCueTagCommited(const FText& InText, ETextCommit::Type InCommitType) = 0;
+	virtual void OnNewGameplayCueTagCommited(const FString& TagName, const FString& TagComment, const FName& TagSource) = 0;
+
 	virtual void OnSearchTagCommited(const FText& InText, ETextCommit::Type InCommitType) = 0;
 	virtual void HandleNotify_OpenAssetInEditor(FString AssetName, int AssetType) = 0;
 	virtual void HandleNotify_FindAssetInEditor(FString AssetName, int AssetType) = 0;
-
-	virtual FReply OnNewGameplayCueButtonPressed() = 0;
 	
 	static FString GetPathNameForGameplayCueTag(FString Tag);
 
