Removed outdated code that made sure attributes were POD. Added validation for keys.

diff --git a/Engine/Plugins/Runtime/GameplayAbilities/Source/GameplayAbilities/Private/AttributeSet.cpp b/Engine/Plugins/Runtime/GameplayAbilities/Source/GameplayAbilities/Private/AttributeSet.cpp
--- a/Engine/Plugins/Runtime/GameplayAbilities/Source/GameplayAbilities/Private/AttributeSet.cpp
+++ b/Engine/Plugins/Runtime/GameplayAbilities/Source/GameplayAbilities/Private/AttributeSet.cpp
@@ -564,6 +564,8 @@ TSubclassOf<UAttributeSet> FindBestAttributeClass(TArray<TSubclassOf<UAttributeS
  *	UCurveTable requires string parsing to map to GroupName/AttributeSet/Attribute
  *	Each curve in the table represents a *single attribute's values for all levels*.
  *	At runtime, we want *all attribute values at given level*.
+ *
+ *	This code assumes that your curve data starts with a key of 1 and increases by 1 with each key.
  */
 void FAttributeSetInitterDiscreteLevels::PreloadAttributeSetData(const TArray<UCurveTable*>& CurveData)
 {
@@ -583,19 +585,6 @@ void FAttributeSetInitterDiscreteLevels::PreloadAttributeSetData(const TArray<UC
 		if (TestClass->IsChildOf(UAttributeSet::StaticClass()))
 		{
 			ClassList.Add(TestClass);
-#if !(UE_BUILD_SHIPPING || UE_BUILD_TEST)
-			// This can only work right now on POD attribute sets. If we ever support FStrings or TArrays in AttributeSets
-			// we will need to update this code to not use memcpy etc.
-			for (TFieldIterator<UProperty> PropIt(TestClass, EFieldIteratorFlags::IncludeSuper); PropIt; ++PropIt)
-			{
-				if (!PropIt->HasAllPropertyFlags(CPF_IsPlainOldData))
-				{
-					ABILITY_LOG(Error, TEXT("FAttributeSetInitterDiscreteLevels::PreloadAttributeSetData Unable to Handle AttributeClass %s because it has a non POD property: %s"),
-						*TestClass->GetName(), *PropIt->GetName());
-					return;
-				}
-			}
-#endif
 		}
 	}
 
@@ -643,6 +632,35 @@ void FAttributeSetInitterDiscreteLevels::PreloadAttributeSetData(const TArray<UC
 			FName ClassFName = FName(*ClassName);
 			FAttributeSetDefaultsCollection& DefaultCollection = Defaults.FindOrAdd(ClassFName);
 
+			// Check our curve to make sure the keys match the expected format
+			int32 ExpectedLevel = 1;
+			bool bShouldSkip = false;
+			for (auto KeyIter = Curve->GetKeyHandleIterator(); KeyIter; ++KeyIter)
+			{
+				const FKeyHandle& KeyHandle = *KeyIter;
+				if (KeyHandle == FKeyHandle::Invalid())
+				{
+					ABILITY_LOG(Verbose, TEXT("FAttributeSetInitterDiscreteLevels::PreloadAttributeSetData Data contains an invalid key handle (row: %s)"), *RowName);
+					bShouldSkip = true;
+					break;
+				}
+
+				int32 Level = Curve->GetKeyTimeValuePair(KeyHandle).Key;
+				if (ExpectedLevel != Level)
+				{
+					ABILITY_LOG(Verbose, TEXT("FAttributeSetInitterDiscreteLevels::PreloadAttributeSetData Keys are expected to start at 1 and increase by 1 for every key (row: %s)"), *RowName);
+					bShouldSkip = true;
+					break;
+				}
+
+				++ExpectedLevel;
+			}
+
+			if (bShouldSkip)
+			{
+				continue;
+			}
+
 			int32 LastLevel = Curve->GetKeyTime(Curve->GetLastKeyHandle());
 			DefaultCollection.LevelData.SetNum(FMath::Max(LastLevel, DefaultCollection.LevelData.Num()));
 
