Manual merge of CL#4795280

Original description:
Added a function to AbilitySystemComponent allowing to stop a given Montage it it's currently running.

diff --git a/Engine/Plugins/Runtime/GameplayAbilities/Source/GameplayAbilities/Private/AbilitySystemComponent_Abilities.cpp b/Engine/Plugins/Runtime/GameplayAbilities/Source/GameplayAbilities/Private/AbilitySystemComponent_Abilities.cpp
--- a/Engine/Plugins/Runtime/GameplayAbilities/Source/GameplayAbilities/Private/AbilitySystemComponent_Abilities.cpp
+++ b/Engine/Plugins/Runtime/GameplayAbilities/Source/GameplayAbilities/Private/AbilitySystemComponent_Abilities.cpp
@@ -2659,6 +2659,14 @@ void UAbilitySystemComponent::CurrentMontageStop(float OverrideBlendOutTime)
 	}
 }
 
+void UAbilitySystemComponent::StopMontageIfCurrent(const UAnimMontage& Montage, float OverrideBlendOutTime)
+{
+	if (&Montage == LocalAnimMontageInfo.AnimMontage)
+	{
+		CurrentMontageStop(OverrideBlendOutTime);
+	}
+}
+
 void UAbilitySystemComponent::ClearAnimatingAbility(UGameplayAbility* Ability)
 {
 	if (LocalAnimMontageInfo.AnimatingAbility == Ability)
diff --git a/Engine/Plugins/Runtime/GameplayAbilities/Source/GameplayAbilities/Public/AbilitySystemComponent.h b/Engine/Plugins/Runtime/GameplayAbilities/Source/GameplayAbilities/Public/AbilitySystemComponent.h
--- a/Engine/Plugins/Runtime/GameplayAbilities/Source/GameplayAbilities/Public/AbilitySystemComponent.h
+++ b/Engine/Plugins/Runtime/GameplayAbilities/Source/GameplayAbilities/Public/AbilitySystemComponent.h
@@ -1206,6 +1206,9 @@ class GAMEPLAYABILITIES_API UAbilitySystemComponent : public UGameplayTasksCompo
 	/** Stops whatever montage is currently playing. Expectation is caller should only be stopping it if they are the current animating ability (or have good reason not to check) */
 	virtual void CurrentMontageStop(float OverrideBlendOutTime = -1.0f);
 
+	/** Stops current montage if it's the one given as the Montage param */
+	virtual void StopMontageIfCurrent(const UAnimMontage& Montage, float OverrideBlendOutTime = -1.0f);
+
 	/** Clear the animating ability that is passed in, if it's still currently animating */
 	virtual void ClearAnimatingAbility(UGameplayAbility* Ability);
 
