* Added method for accessing GameplayTag manager without allocating it.
* Added static method for getting "default" gameplay tag category filter (to be used by custom overides)
* Added more remappable gameplay tag categories for GameplayAbility tag properties.

#RNX

diff --git a/Engine/Plugins/Runtime/GameplayAbilities/Source/GameplayAbilities/Public/Abilities/GameplayAbility.h b/Engine/Plugins/Runtime/GameplayAbilities/Source/GameplayAbilities/Public/Abilities/GameplayAbility.h
--- a/Engine/Plugins/Runtime/GameplayAbilities/Source/GameplayAbilities/Public/Abilities/GameplayAbility.h
+++ b/Engine/Plugins/Runtime/GameplayAbilities/Source/GameplayAbilities/Public/Abilities/GameplayAbility.h
@@ -81,7 +81,7 @@ struct FAbilityTriggerData
 	{}
 
 	/** The tag to respond to */
-	UPROPERTY(EditAnywhere, Category=TriggerData)
+	UPROPERTY(EditAnywhere, Category=TriggerData, meta=(Categories="TriggerTagCategory"))
 	FGameplayTag TriggerTag;
 
 	/** The type of trigger to respond to */
@@ -460,7 +460,7 @@ public:
 	// --------------------------------------
 
 	/** This ability has these tags */
-	UPROPERTY(EditDefaultsOnly, Category = Tags)
+	UPROPERTY(EditDefaultsOnly, Category = Tags, meta=(Categories="AbilityTagCategory"))
 	FGameplayTagContainer AbilityTags;
 
 	/** If true, this ability will always replicate input press/release events to the server. */
@@ -689,43 +689,43 @@ protected:
 	// ----------------------------------------------------------------------------------------------------------------
 
 	/** Abilities matching query are cancelled when this ability is executed */
-	UPROPERTY(EditDefaultsOnly, Category = TagQueries)
+	UPROPERTY(EditDefaultsOnly, Category = TagQueries, meta=(Categories="AbilityTagCategory"))
 	FGameplayTagQuery CancelAbilitiesMatchingTagQuery;
 
 	/** Abilities with these tags are cancelled when this ability is executed */
-	UPROPERTY(EditDefaultsOnly, Category = Tags)
+	UPROPERTY(EditDefaultsOnly, Category = Tags, meta=(Categories="AbilityTagCategory"))
 	FGameplayTagContainer CancelAbilitiesWithTag;
 
 	/** Abilities with these tags are blocked while this ability is active */
-	UPROPERTY(EditDefaultsOnly, Category = Tags)
+	UPROPERTY(EditDefaultsOnly, Category = Tags, meta=(Categories="AbilityTagCategory"))
 	FGameplayTagContainer BlockAbilitiesWithTag;
 
 	/** Tags to apply to activating owner while this ability is active */
-	UPROPERTY(EditDefaultsOnly, Category = Tags)
+	UPROPERTY(EditDefaultsOnly, Category = Tags, meta=(Categories="OwnedTagsCategory"))
 	FGameplayTagContainer ActivationOwnedTags;
 
 	/** This ability can only be activated if the activating actor/component has all of these tags */
-	UPROPERTY(EditDefaultsOnly, Category = Tags)
+	UPROPERTY(EditDefaultsOnly, Category = Tags, meta=(Categories="OwnedTagsCategory"))
 	FGameplayTagContainer ActivationRequiredTags;
 
 	/** This ability is blocked if the activating actor/component has any of these tags */
-	UPROPERTY(EditDefaultsOnly, Category = Tags)
+	UPROPERTY(EditDefaultsOnly, Category = Tags, meta=(Categories="OwnedTagsCategory"))
 	FGameplayTagContainer ActivationBlockedTags;
 
 	/** This ability can only be activated if the source actor/component has all of these tags */
-	UPROPERTY(EditDefaultsOnly, Category = Tags)
+	UPROPERTY(EditDefaultsOnly, Category = Tags, meta=(Categories="SourceTagsCategory"))
 	FGameplayTagContainer SourceRequiredTags;
 
 	/** This ability is blocked if the source actor/component has any of these tags */
-	UPROPERTY(EditDefaultsOnly, Category = Tags)
+	UPROPERTY(EditDefaultsOnly, Category = Tags, meta=(Categories="SourceTagsCategory"))
 	FGameplayTagContainer SourceBlockedTags;
 
 	/** This ability can only be activated if the target actor/component has all of these tags */
-	UPROPERTY(EditDefaultsOnly, Category = Tags)
+	UPROPERTY(EditDefaultsOnly, Category = Tags, meta=(Categories="TargetTagsCategory"))
 	FGameplayTagContainer TargetRequiredTags;
 
 	/** This ability is blocked if the target actor/component has any of these tags */
-	UPROPERTY(EditDefaultsOnly, Category = Tags)
+	UPROPERTY(EditDefaultsOnly, Category = Tags, meta=(Categories="TargetTagsCategory"))
 	FGameplayTagContainer TargetBlockedTags;
 
 	// ----------------------------------------------------------------------------------------------------------------
diff --git a/Engine/Source/Runtime/GameplayTags/Classes/GameplayTagsManager.h b/Engine/Source/Runtime/GameplayTags/Classes/GameplayTagsManager.h
--- a/Engine/Source/Runtime/GameplayTags/Classes/GameplayTagsManager.h
+++ b/Engine/Source/Runtime/GameplayTags/Classes/GameplayTagsManager.h
@@ -285,6 +285,9 @@ class GAMEPLAYTAGS_API UGameplayTagsManager : public UObject
 		return *SingletonManager;
 	}
 
+	/** Returns possibly nullptr to the manager. Needed for some shutdown cases to avoid reallocating. */
+	FORCEINLINE static UGameplayTagsManager* GetIfAllocated() { return SingletonManager; }
+
 	/**
 	* Adds the gameplay tags corresponding to the strings in the array TagStrings to OutTagsContainer
 	*
@@ -534,6 +537,9 @@ class GAMEPLAYTAGS_API UGameplayTagsManager : public UObject
 	/** Returns "Categories" meta property from given handle, used for filtering by tag widget */
 	FString GetCategoriesMetaFromPropertyHandle(TSharedPtr<class IPropertyHandle> PropertyHandle) const;
 
+	/** Helper function, made to be called by custom OnGetCategoriesMetaFromPropertyHandle handlers  */
+	static FString StaticGetCategoriesMetaFromPropertyHandle(TSharedPtr<class IPropertyHandle> PropertyHandle);
+
 	/** Returns "Categories" meta property from given field, used for filtering by tag widget */
 	FString GetCategoriesMetaFromField(UField* Field) const;
 
diff --git a/Engine/Source/Runtime/GameplayTags/Private/GameplayTagsManager.cpp b/Engine/Source/Runtime/GameplayTags/Private/GameplayTagsManager.cpp
--- a/Engine/Source/Runtime/GameplayTags/Private/GameplayTagsManager.cpp
+++ b/Engine/Source/Runtime/GameplayTags/Private/GameplayTagsManager.cpp
@@ -1287,6 +1287,11 @@ FString UGameplayTagsManager::GetCategoriesMetaFromPropertyHandle(TSharedPtr<IPr
 		return DelegateOverrideString;
 	}
 
+	return StaticGetCategoriesMetaFromPropertyHandle(PropertyHandle);
+}
+
+FString UGameplayTagsManager::StaticGetCategoriesMetaFromPropertyHandle(TSharedPtr<class IPropertyHandle> PropertyHandle)
+{
 	FString Categories;
 
 	auto GetMetaData = ([&](UField* Field)
