It is unsafe to reference dynamic objects from a CDO so spawn an instance of the default class

diff --git a/Engine/Plugins/Runtime/GameplayAbilities/Source/GameplayAbilities/Private/AbilitySystemGlobals.cpp b/Engine/Plugins/Runtime/GameplayAbilities/Source/GameplayAbilities/Private/AbilitySystemGlobals.cpp
--- a/Engine/Plugins/Runtime/GameplayAbilities/Source/GameplayAbilities/Private/AbilitySystemGlobals.cpp
+++ b/Engine/Plugins/Runtime/GameplayAbilities/Source/GameplayAbilities/Private/AbilitySystemGlobals.cpp
@@ -384,8 +384,8 @@ UGameplayCueManager* UAbilitySystemGlobals::GetGameplayCueManager()
 
 		if ( GlobalGameplayCueManager == nullptr)
 		{
-			// Fallback to CDO
-			GlobalGameplayCueManager = UGameplayCueManager::StaticClass()->GetDefaultObject<UGameplayCueManager>();
+			// Fallback to base native class
+			GlobalGameplayCueManager = NewObject<UGameplayCueManager>(this, UGameplayCueManager::StaticClass(), NAME_None);
 		}
 
 		GlobalGameplayCueManager->OnCreated();
