[NO NON-PII COMMIT MESSAGE]

diff --git a/Engine/Plugins/Runtime/GameplayAbilities/Source/GameplayAbilities/Private/AbilitySystemBlueprintLibrary.cpp b/Engine/Plugins/Runtime/GameplayAbilities/Source/GameplayAbilities/Private/AbilitySystemBlueprintLibrary.cpp
--- a/Engine/Plugins/Runtime/GameplayAbilities/Source/GameplayAbilities/Private/AbilitySystemBlueprintLibrary.cpp
+++ b/Engine/Plugins/Runtime/GameplayAbilities/Source/GameplayAbilities/Private/AbilitySystemBlueprintLibrary.cpp
@@ -660,6 +660,48 @@ bool UAbilitySystemBlueprintLibrary::DoesGameplayCueMeetTagRequirements(FGamepla
 		&& TargetTagReqs.RequirementsMet(Parameters.AggregatedSourceTags);
 }
 
+FGameplayCueParameters UAbilitySystemBlueprintLibrary::MakeGameplayCueParameters(float NormalizedMagnitude, float RawMagnitude, FGameplayEffectContextHandle EffectContext, FGameplayTag MatchedTagName, FGameplayTag OriginalTag, FGameplayTagContainer AggregatedSourceTags, FGameplayTagContainer AggregatedTargetTags, FVector Location, FVector Normal, AActor* Instigator, AActor* EffectCauser, UObject* SourceObject, UPhysicalMaterial* PhysicalMaterial, int32 GameplayEffectLevel, int32 AbilityLevel, USceneComponent* TargetAttachComponent)
+{
+	FGameplayCueParameters Parameters;
+	Parameters.NormalizedMagnitude = NormalizedMagnitude;
+	Parameters.RawMagnitude = RawMagnitude;
+	Parameters.EffectContext = EffectContext;
+	Parameters.MatchedTagName = MatchedTagName;
+	Parameters.OriginalTag = OriginalTag;
+	Parameters.AggregatedSourceTags = AggregatedSourceTags;
+	Parameters.AggregatedTargetTags = AggregatedTargetTags;
+	Parameters.Location = Location;
+	Parameters.Normal = Normal;
+	Parameters.Instigator = Instigator;
+	Parameters.EffectCauser = EffectCauser;
+	Parameters.SourceObject = SourceObject;
+	Parameters.PhysicalMaterial = PhysicalMaterial;
+	Parameters.GameplayEffectLevel = GameplayEffectLevel;
+	Parameters.AbilityLevel = AbilityLevel;
+	Parameters.TargetAttachComponent = TargetAttachComponent;
+	return Parameters;
+}
+
+void UAbilitySystemBlueprintLibrary::BreakGameplayCueParameters(const struct FGameplayCueParameters& Parameters, float& NormalizedMagnitude, float& RawMagnitude, FGameplayEffectContextHandle& EffectContext, FGameplayTag& MatchedTagName, FGameplayTag& OriginalTag, FGameplayTagContainer& AggregatedSourceTags, FGameplayTagContainer& AggregatedTargetTags, FVector& Location, FVector& Normal, AActor* Instigator, AActor*& EffectCauser, UObject*& SourceObject, UPhysicalMaterial*& PhysicalMaterial, int32& GameplayEffectLevel, int32& AbilityLevel, USceneComponent*& TargetAttachComponent)
+{
+	NormalizedMagnitude = Parameters.NormalizedMagnitude;
+	RawMagnitude = Parameters.RawMagnitude;
+	EffectContext = Parameters.EffectContext;
+	MatchedTagName = Parameters.MatchedTagName;
+	OriginalTag = Parameters.OriginalTag;
+	AggregatedSourceTags = Parameters.AggregatedSourceTags;
+	AggregatedTargetTags = Parameters.AggregatedTargetTags;
+	Location = Parameters.Location;
+	Normal = Parameters.Normal;
+	Instigator = Parameters.Instigator.Get();
+	EffectCauser = Parameters.EffectCauser.Get();
+	SourceObject = const_cast<UObject*>(Parameters.SourceObject.Get());
+	PhysicalMaterial = const_cast<UPhysicalMaterial*>(Parameters.PhysicalMaterial.Get());
+	GameplayEffectLevel = GameplayEffectLevel;
+	AbilityLevel = AbilityLevel;
+	TargetAttachComponent = Parameters.TargetAttachComponent.Get();
+}
+
 // ---------------------------------------------------------------------------------------
 
 FGameplayEffectSpecHandle UAbilitySystemBlueprintLibrary::AssignSetByCallerMagnitude(FGameplayEffectSpecHandle SpecHandle, FName DataName, float Magnitude)
diff --git a/Engine/Plugins/Runtime/GameplayAbilities/Source/GameplayAbilities/Public/AbilitySystemBlueprintLibrary.h b/Engine/Plugins/Runtime/GameplayAbilities/Source/GameplayAbilities/Public/AbilitySystemBlueprintLibrary.h
--- a/Engine/Plugins/Runtime/GameplayAbilities/Source/GameplayAbilities/Public/AbilitySystemBlueprintLibrary.h
+++ b/Engine/Plugins/Runtime/GameplayAbilities/Source/GameplayAbilities/Public/AbilitySystemBlueprintLibrary.h
@@ -266,6 +266,14 @@ class GAMEPLAYABILITIES_API UAbilitySystemBlueprintLibrary : public UBlueprintFu
 	UFUNCTION(BlueprintPure, Category = "Ability|GameplayCue")
 	static bool DoesGameplayCueMeetTagRequirements(FGameplayCueParameters Parameters, const FGameplayTagRequirements& SourceTagReqs, const FGameplayTagRequirements& TargetTagReqs);
 
+	/** Native make, to avoid having to deal with quantized vector types */
+	UFUNCTION(BlueprintPure, Category = "Ability|GameplayCue", meta = (NativeMakeFunc, AdvancedDisplay=5, Location="0,0,0", Normal= "0,0,0", GameplayEffectLevel = "1", AbilityLevel = "1"))
+	static FGameplayCueParameters MakeGameplayCueParameters(float NormalizedMagnitude, float RawMagnitude, FGameplayEffectContextHandle EffectContext, FGameplayTag MatchedTagName, FGameplayTag OriginalTag, FGameplayTagContainer AggregatedSourceTags, FGameplayTagContainer AggregatedTargetTags, FVector Location, FVector Normal, AActor* Instigator, AActor* EffectCauser, UObject* SourceObject, UPhysicalMaterial* PhysicalMaterial, int32 GameplayEffectLevel, int32 AbilityLevel, USceneComponent* TargetAttachComponent);
+
+	/** Native break, to avoid having to deal with quantized vector types */
+	UFUNCTION(BlueprintPure, Category = "Ability|GameplayCue", meta = (NativeBreakFunc, AdvancedDisplay=6))
+	static void BreakGameplayCueParameters(const struct FGameplayCueParameters& Parameters, float& NormalizedMagnitude, float& RawMagnitude, FGameplayEffectContextHandle& EffectContext, FGameplayTag& MatchedTagName, FGameplayTag& OriginalTag, FGameplayTagContainer& AggregatedSourceTags, FGameplayTagContainer& AggregatedTargetTags, FVector& Location, FVector& Normal, AActor* Instigator, AActor*& EffectCauser, UObject*& SourceObject, UPhysicalMaterial*& PhysicalMaterial, int32& GameplayEffectLevel, int32& AbilityLevel, USceneComponent*& TargetAttachComponent);
+
 	// -------------------------------------------------------------------------------
 	//		GameplayEffectSpec
 	// -------------------------------------------------------------------------------
diff --git a/Engine/Plugins/Runtime/GameplayAbilities/Source/GameplayAbilities/Public/GameplayEffectTypes.h b/Engine/Plugins/Runtime/GameplayAbilities/Source/GameplayAbilities/Public/GameplayEffectTypes.h
--- a/Engine/Plugins/Runtime/GameplayAbilities/Source/GameplayAbilities/Public/GameplayEffectTypes.h
+++ b/Engine/Plugins/Runtime/GameplayAbilities/Source/GameplayAbilities/Public/GameplayEffectTypes.h
@@ -918,7 +918,7 @@ struct FGameplayEffectRemovalInfo
 
 
 /** Metadata about a gameplay cue execution */
-USTRUCT(BlueprintType)
+USTRUCT(BlueprintType, meta = (HasNativeBreak = "GameplayAbilities.AbilitySystemBlueprintLibrary.BreakGameplayCueParameters", HasNativeMake = "GameplayAbilities.AbilitySystemBlueprintLibrary.MakeGameplayCueParameters"))
 struct GAMEPLAYABILITIES_API FGameplayCueParameters
 {
 	GENERATED_USTRUCT_BODY()
