Fix typo with assignment

diff --git a/Engine/Plugins/Runtime/GameplayAbilities/Source/GameplayAbilities/Private/AbilitySystemBlueprintLibrary.cpp b/Engine/Plugins/Runtime/GameplayAbilities/Source/GameplayAbilities/Private/AbilitySystemBlueprintLibrary.cpp
--- a/Engine/Plugins/Runtime/GameplayAbilities/Source/GameplayAbilities/Private/AbilitySystemBlueprintLibrary.cpp
+++ b/Engine/Plugins/Runtime/GameplayAbilities/Source/GameplayAbilities/Private/AbilitySystemBlueprintLibrary.cpp
@@ -697,8 +697,8 @@ void UAbilitySystemBlueprintLibrary::BreakGameplayCueParameters(const struct FGa
 	EffectCauser = Parameters.EffectCauser.Get();
 	SourceObject = const_cast<UObject*>(Parameters.SourceObject.Get());
 	PhysicalMaterial = const_cast<UPhysicalMaterial*>(Parameters.PhysicalMaterial.Get());
-	GameplayEffectLevel = GameplayEffectLevel;
-	AbilityLevel = AbilityLevel;
+	GameplayEffectLevel = Parameters.GameplayEffectLevel;
+	AbilityLevel = Parameters.AbilityLevel;
 	TargetAttachComponent = Parameters.TargetAttachComponent.Get();
 }
 
