[NO NON-PII COMMIT MESSAGE]

diff --git a/Engine/Plugins/Runtime/GameplayAbilities/Source/GameplayAbilities/Private/AbilitySystemBlueprintLibrary.cpp b/Engine/Plugins/Runtime/GameplayAbilities/Source/GameplayAbilities/Private/AbilitySystemBlueprintLibrary.cpp
--- a/Engine/Plugins/Runtime/GameplayAbilities/Source/GameplayAbilities/Private/AbilitySystemBlueprintLibrary.cpp
+++ b/Engine/Plugins/Runtime/GameplayAbilities/Source/GameplayAbilities/Private/AbilitySystemBlueprintLibrary.cpp
@@ -682,7 +682,7 @@ FGameplayCueParameters UAbilitySystemBlueprintLibrary::MakeGameplayCueParameters
 	return Parameters;
 }
 
-void UAbilitySystemBlueprintLibrary::BreakGameplayCueParameters(const struct FGameplayCueParameters& Parameters, float& NormalizedMagnitude, float& RawMagnitude, FGameplayEffectContextHandle& EffectContext, FGameplayTag& MatchedTagName, FGameplayTag& OriginalTag, FGameplayTagContainer& AggregatedSourceTags, FGameplayTagContainer& AggregatedTargetTags, FVector& Location, FVector& Normal, AActor* Instigator, AActor*& EffectCauser, UObject*& SourceObject, UPhysicalMaterial*& PhysicalMaterial, int32& GameplayEffectLevel, int32& AbilityLevel, USceneComponent*& TargetAttachComponent)
+void UAbilitySystemBlueprintLibrary::BreakGameplayCueParameters(const struct FGameplayCueParameters& Parameters, float& NormalizedMagnitude, float& RawMagnitude, FGameplayEffectContextHandle& EffectContext, FGameplayTag& MatchedTagName, FGameplayTag& OriginalTag, FGameplayTagContainer& AggregatedSourceTags, FGameplayTagContainer& AggregatedTargetTags, FVector& Location, FVector& Normal, AActor*& Instigator, AActor*& EffectCauser, UObject*& SourceObject, UPhysicalMaterial*& PhysicalMaterial, int32& GameplayEffectLevel, int32& AbilityLevel, USceneComponent*& TargetAttachComponent)
 {
 	NormalizedMagnitude = Parameters.NormalizedMagnitude;
 	RawMagnitude = Parameters.RawMagnitude;
diff --git a/Engine/Plugins/Runtime/GameplayAbilities/Source/GameplayAbilities/Public/AbilitySystemBlueprintLibrary.h b/Engine/Plugins/Runtime/GameplayAbilities/Source/GameplayAbilities/Public/AbilitySystemBlueprintLibrary.h
--- a/Engine/Plugins/Runtime/GameplayAbilities/Source/GameplayAbilities/Public/AbilitySystemBlueprintLibrary.h
+++ b/Engine/Plugins/Runtime/GameplayAbilities/Source/GameplayAbilities/Public/AbilitySystemBlueprintLibrary.h
@@ -272,7 +272,7 @@ class GAMEPLAYABILITIES_API UAbilitySystemBlueprintLibrary : public UBlueprintFu
 
 	/** Native break, to avoid having to deal with quantized vector types */
 	UFUNCTION(BlueprintPure, Category = "Ability|GameplayCue", meta = (NativeBreakFunc, AdvancedDisplay=6))
-	static void BreakGameplayCueParameters(const struct FGameplayCueParameters& Parameters, float& NormalizedMagnitude, float& RawMagnitude, FGameplayEffectContextHandle& EffectContext, FGameplayTag& MatchedTagName, FGameplayTag& OriginalTag, FGameplayTagContainer& AggregatedSourceTags, FGameplayTagContainer& AggregatedTargetTags, FVector& Location, FVector& Normal, AActor* Instigator, AActor*& EffectCauser, UObject*& SourceObject, UPhysicalMaterial*& PhysicalMaterial, int32& GameplayEffectLevel, int32& AbilityLevel, USceneComponent*& TargetAttachComponent);
+	static void BreakGameplayCueParameters(const struct FGameplayCueParameters& Parameters, float& NormalizedMagnitude, float& RawMagnitude, FGameplayEffectContextHandle& EffectContext, FGameplayTag& MatchedTagName, FGameplayTag& OriginalTag, FGameplayTagContainer& AggregatedSourceTags, FGameplayTagContainer& AggregatedTargetTags, FVector& Location, FVector& Normal, AActor*& Instigator, AActor*& EffectCauser, UObject*& SourceObject, UPhysicalMaterial*& PhysicalMaterial, int32& GameplayEffectLevel, int32& AbilityLevel, USceneComponent*& TargetAttachComponent);
 
 	// -------------------------------------------------------------------------------
 	//		GameplayEffectSpec
