Remove [Unsupported] from ability system plugin description

diff --git a/Engine/Plugins/Runtime/GameplayAbilities/GameplayAbilities.uplugin b/Engine/Plugins/Runtime/GameplayAbilities/GameplayAbilities.uplugin
--- a/Engine/Plugins/Runtime/GameplayAbilities/GameplayAbilities.uplugin
+++ b/Engine/Plugins/Runtime/GameplayAbilities/GameplayAbilities.uplugin
@@ -3,7 +3,7 @@
 	"Version" : 1,
 	"VersionName" : "0.1",
 	"FriendlyName" : "Gameplay Abilities",
-	"Description" : "[UNSUPPORTED] Adds GameplayEffect and GameplayAbility classes to handle complicated gameplay interactions.",
+	"Description" : "Adds GameplayEffect and GameplayAbility classes to handle complicated gameplay interactions.",
 	"Category" : "Gameplay",
 	"CreatedBy" : "Epic Games, Inc.",
 	"CreatedByURL" : "http://epicgames.com",
