Fix for crash in ability system.

diff --git a/Engine/Plugins/Runtime/GameplayAbilities/Source/GameplayAbilities/Private/GameplayEffect.cpp b/Engine/Plugins/Runtime/GameplayAbilities/Source/GameplayAbilities/Private/GameplayEffect.cpp
--- a/Engine/Plugins/Runtime/GameplayAbilities/Source/GameplayAbilities/Private/GameplayEffect.cpp
+++ b/Engine/Plugins/Runtime/GameplayAbilities/Source/GameplayAbilities/Private/GameplayEffect.cpp
@@ -2547,11 +2547,13 @@ void FActiveGameplayEffectsContainer::InternalUpdateNumericalAttribute(FGameplay
 void FActiveGameplayEffectsContainer::SetAttributeBaseValue(FGameplayAttribute Attribute, float NewBaseValue)
 {
 	const UAttributeSet* Set = Owner->GetAttributeSubobject(Attribute.GetAttributeSetClass());
-	if (ensure(Set))
+	if (!ensureMsgf(Set, TEXT("FActiveGameplayEffectsContainer::SetAttributeBaseValue: Unable to get attribute set for attribute %s"), *Attribute.AttributeName))
 	{
-		Set->PreAttributeBaseChange(Attribute, NewBaseValue);
+		return;
 	}
 
+	Set->PreAttributeBaseChange(Attribute, NewBaseValue);
+
 	// if we're using the new attributes we should always update their base value
 	bool bIsGameplayAttributeDataProperty = FGameplayAttribute::IsGameplayAttributeDataProperty(Attribute.GetUProperty());
 	if (bIsGameplayAttributeDataProperty)
