Change FActiveGameplayEffect::bIsInhibited to not replicate as it was causing bugged behavior.

diff --git a/Engine/Plugins/Runtime/GameplayAbilities/Source/GameplayAbilities/Public/GameplayEffect.h b/Engine/Plugins/Runtime/GameplayAbilities/Source/GameplayAbilities/Public/GameplayEffect.h
--- a/Engine/Plugins/Runtime/GameplayAbilities/Source/GameplayAbilities/Public/GameplayEffect.h
+++ b/Engine/Plugins/Runtime/GameplayAbilities/Source/GameplayAbilities/Public/GameplayEffect.h
@@ -1271,7 +1271,7 @@ struct GAMEPLAYABILITIES_API FActiveGameplayEffect : public FFastArraySerializer
 	float StartWorldTime;
 
 	// Not sure if this should replicate or not. If replicated, we may have trouble where IsInhibited doesn't appear to change when we do tag checks (because it was previously inhibited, but replication made it inhibited).
-	UPROPERTY()
+	UPROPERTY(NotReplicated)
 	bool bIsInhibited;
 
 	/** When replicated down, we cue the GC events until the entire list of active gameplay effects has been received */
