- Fixed a crash in the GameplayCueManager by ensuring gameplay cue actors tie themselves to the main persistent level, instead of the streaming level their owning actors might be on.

#rnx

diff --git a/Engine/Plugins/Runtime/GameplayAbilities/Source/GameplayAbilities/Private/GameplayCueManager.cpp b/Engine/Plugins/Runtime/GameplayAbilities/Source/GameplayAbilities/Private/GameplayCueManager.cpp
--- a/Engine/Plugins/Runtime/GameplayAbilities/Source/GameplayAbilities/Private/GameplayCueManager.cpp
+++ b/Engine/Plugins/Runtime/GameplayAbilities/Source/GameplayAbilities/Private/GameplayCueManager.cpp
@@ -438,6 +438,7 @@ AGameplayCueNotify_Actor* UGameplayCueManager::GetInstancedCueActor(AActor* Targ
 				ABILITY_LOG(Warning, TEXT("Spawning GameplaycueActor: %s"), *CueClass->GetName());
 			}
 
+			SpawnParams.OverrideLevel = World->PersistentLevel;
 			SpawnedCue = World->SpawnActor<AGameplayCueNotify_Actor>(CueClass, TargetActor->GetActorLocation(), TargetActor->GetActorRotation(), SpawnParams);
 		}
 
