Cleaning up demo net driver includes

diff --git a/Engine/Plugins/Runtime/GameplayAbilities/Source/GameplayAbilities/Private/AbilitySystemComponent_Abilities.cpp b/Engine/Plugins/Runtime/GameplayAbilities/Source/GameplayAbilities/Private/AbilitySystemComponent_Abilities.cpp
--- a/Engine/Plugins/Runtime/GameplayAbilities/Source/GameplayAbilities/Private/AbilitySystemComponent_Abilities.cpp
+++ b/Engine/Plugins/Runtime/GameplayAbilities/Source/GameplayAbilities/Private/AbilitySystemComponent_Abilities.cpp
@@ -30,7 +30,6 @@
 #include "Abilities/GameplayAbilityTargetActor.h"
 #include "TickableAttributeSetInterface.h"
 #include "GameplayTagResponseTable.h"
-#include "Engine/DemoNetDriver.h"
 #include "ProfilingDebugging/CsvProfiler.h"
 
 #define LOCTEXT_NAMESPACE "AbilitySystemComponent"
@@ -2519,7 +2518,7 @@ void UAbilitySystemComponent::OnRep_ReplicatedAnimMontage()
 		RepAnimMontageInfo.PlayRate = 1.f;
 	}
 
-	const bool bIsPlayingReplay = World && World->DemoNetDriver && World->DemoNetDriver->IsPlaying();
+	const bool bIsPlayingReplay = World && World->IsPlayingReplay();
 
 	const float MONTAGE_REP_POS_ERR_THRESH = bIsPlayingReplay ? CVarReplayMontageErrorThreshold.GetValueOnGameThread() : 0.1f;
 
diff --git a/Engine/Source/Runtime/Engine/Private/DataReplication.cpp b/Engine/Source/Runtime/Engine/Private/DataReplication.cpp
--- a/Engine/Source/Runtime/Engine/Private/DataReplication.cpp
+++ b/Engine/Source/Runtime/Engine/Private/DataReplication.cpp
@@ -14,9 +14,9 @@
 #include "Engine/PackageMapClient.h"
 #include "Net/RepLayout.h"
 #include "Engine/ActorChannel.h"
-#include "Engine/DemoNetDriver.h"
 #include "ProfilingDebugging/CsvProfiler.h"
 #include "Engine/Engine.h"
+#include "Engine/NetConnection.h"
 
 DECLARE_CYCLE_STAT(TEXT("Custom Delta Property Rep Time"), STAT_NetReplicateCustomDeltaPropTime, STATGROUP_Game);
 DECLARE_CYCLE_STAT(TEXT("ReceiveRPC"), STAT_NetReceiveRPC, STATGROUP_Game);
