Merging //UE4/Dev-Main to Dev-Sequencer (//UE4/Dev-Sequencer)

diff --git a/Engine/Plugins/Runtime/GameplayAbilities/Source/GameplayAbilities/Private/GameplayCueManager.cpp b/Engine/Plugins/Runtime/GameplayAbilities/Source/GameplayAbilities/Private/GameplayCueManager.cpp
--- a/Engine/Plugins/Runtime/GameplayAbilities/Source/GameplayAbilities/Private/GameplayCueManager.cpp
+++ b/Engine/Plugins/Runtime/GameplayAbilities/Source/GameplayAbilities/Private/GameplayCueManager.cpp
@@ -438,6 +438,7 @@ AGameplayCueNotify_Actor* UGameplayCueManager::GetInstancedCueActor(AActor* Targ
 				ABILITY_LOG(Warning, TEXT("Spawning GameplaycueActor: %s"), *CueClass->GetName());
 			}
 
+			SpawnParams.OverrideLevel = World->PersistentLevel;
 			SpawnedCue = World->SpawnActor<AGameplayCueNotify_Actor>(CueClass, TargetActor->GetActorLocation(), TargetActor->GetActorRotation(), SpawnParams);
 		}
 
diff --git a/Engine/Source/Programs/UnrealBuildTool/Platform/TVOS/UEBuildTVOS.cs b/Engine/Source/Programs/UnrealBuildTool/Platform/TVOS/UEBuildTVOS.cs
--- a/Engine/Source/Programs/UnrealBuildTool/Platform/TVOS/UEBuildTVOS.cs
+++ b/Engine/Source/Programs/UnrealBuildTool/Platform/TVOS/UEBuildTVOS.cs
@@ -17,7 +17,7 @@ namespace UnrealBuildTool
 		/// </summary>
 		public override string RuntimeVersion
 		{
-			get { return "9.0"; }
+			get { return "10.0"; }
 		}
 
 		/// <summary>
diff --git a/Engine/Source/Runtime/Core/Private/Logging/LogMacros.cpp b/Engine/Source/Runtime/Core/Private/Logging/LogMacros.cpp
--- a/Engine/Source/Runtime/Core/Private/Logging/LogMacros.cpp
+++ b/Engine/Source/Runtime/Core/Private/Logging/LogMacros.cpp
@@ -72,7 +72,6 @@ void FMsg::Logf_InternalImpl(const ANSICHAR* File, int32 Line, const FName& Cate
 {
 #if !NO_LOGGING
 	QUICK_SCOPE_CYCLE_COUNTER(STAT_FMsgLogf);
-	CSV_SCOPED_TIMING_STAT(FMsgLogf, FMsgLogfTime);
 	CSV_CUSTOM_STAT(FMsgLogf, FMsgLogfCount, 1, ECsvCustomStatOp::Accumulate);
 
 	if (Verbosity != ELogVerbosity::Fatal)
diff --git a/Engine/Source/Runtime/Slate/Private/Widgets/LayerManager/STooltipPresenter.cpp b/Engine/Source/Runtime/Slate/Private/Widgets/LayerManager/STooltipPresenter.cpp
--- a/Engine/Source/Runtime/Slate/Private/Widgets/LayerManager/STooltipPresenter.cpp
+++ b/Engine/Source/Runtime/Slate/Private/Widgets/LayerManager/STooltipPresenter.cpp
@@ -40,7 +40,14 @@ void STooltipPresenter::OnArrangeChildren(const FGeometry& AllottedGeometry, FAr
 
 FVector2D STooltipPresenter::ComputeDesiredSize( float ) const
 {
-	return ChildSlot.GetWidget()->GetDesiredSize();
+	if (ChildSlot.GetWidget() != SNullWidget::NullWidget)
+	{
+		return ChildSlot.GetWidget()->GetDesiredSize();
+	}
+	else
+	{
+		return FVector2D(0, 0);
+	}
 }
 
 FChildren* STooltipPresenter::GetChildren()
