Improved ability system debug display of ability system by adding drop shadows under all text and making text moves around less.

#test: pie

diff --git a/Engine/Plugins/Runtime/GameplayAbilities/Source/GameplayAbilities/Private/AbilitySystemComponent.cpp b/Engine/Plugins/Runtime/GameplayAbilities/Source/GameplayAbilities/Private/AbilitySystemComponent.cpp
--- a/Engine/Plugins/Runtime/GameplayAbilities/Source/GameplayAbilities/Private/AbilitySystemComponent.cpp
+++ b/Engine/Plugins/Runtime/GameplayAbilities/Source/GameplayAbilities/Private/AbilitySystemComponent.cpp
@@ -1784,12 +1784,22 @@ void UAbilitySystemComponent::AccumulateScreenPos(FAbilitySystemComponentDebugIn
 	Info.YPos = NewY;
 }
 
-void UAbilitySystemComponent::DebugLine(FAbilitySystemComponentDebugInfo& Info, FString Str, float XOffset, float YOffset)
+void UAbilitySystemComponent::DebugLine(FAbilitySystemComponentDebugInfo& Info, FString Str, float XOffset, float YOffset, int32 MinTextRowsToAdvance)
 {
 	if (Info.Canvas)
 	{
-		Info.YL = Info.Canvas->DrawText(GEngine->GetTinyFont(), Str, Info.XPos + XOffset, Info.YPos );
-		AccumulateScreenPos(Info);
+		FFontRenderInfo RenderInfo = FFontRenderInfo();
+		RenderInfo.bEnableShadow = true;
+		if (const UFont* Font = GEngine->GetTinyFont())
+		{
+			float ScaleY = 1.f;
+			Info.YL = Info.Canvas->DrawText(Font, Str, Info.XPos + XOffset, Info.YPos, 1.f, ScaleY, RenderInfo);
+			if (Info.YL < MinTextRowsToAdvance * (Font->GetMaxCharHeight() * ScaleY))
+			{
+				Info.YL = MinTextRowsToAdvance * (Font->GetMaxCharHeight() * ScaleY);
+			}
+			AccumulateScreenPos(Info);
+		}
 	}
 
 	if (Info.bPrintToLog)
@@ -2094,7 +2104,9 @@ void UAbilitySystemComponent::Debug_Internal(FAbilitySystemComponentDebugInfo& I
 		if (Info.Canvas)
 		{
 			Info.Canvas->SetDrawColor(FColor::White);
-			Info.Canvas->DrawText(GEngine->GetLargeFont(), DebugTitle, Info.XPos + 4.f, 10.f, 1.5f, 1.5f);
+			FFontRenderInfo RenderInfo = FFontRenderInfo();
+			RenderInfo.bEnableShadow = true;
+			Info.Canvas->DrawText(GEngine->GetLargeFont(), DebugTitle, Info.XPos + 4.f, 10.f, 1.5f, 1.5f, RenderInfo);
 		}
 		else
 		{
@@ -2110,12 +2122,16 @@ void UAbilitySystemComponent::Debug_Internal(FAbilitySystemComponentDebugInfo& I
 		Info.Canvas->SetDrawColor(FColor::White);
 	}
 
-	DebugLine(Info, FString::Printf(TEXT("Owned Tags: %s"), *OwnerTags.ToStringSimple()), 4.f, 0.f);
+	DebugLine(Info, FString::Printf(TEXT("Owned Tags: %s"), *OwnerTags.ToStringSimple()), 4.f, 0.f, 4);
 
 	if (BlockedAbilityTags.GetExplicitGameplayTags().Num() > 0)
 	{
 		DebugLine(Info, FString::Printf(TEXT("BlockedAbilityTags: %s"), *BlockedAbilityTags.GetExplicitGameplayTags().ToStringSimple()), 4.f, 0.f);
 	}
+	else
+	{
+		DebugLine(Info, FString::Printf(TEXT("BlockedAbilityTags: ")), 4.f, 0.f);
+	}
 
 	TSet<FGameplayAttribute> DrawAttributes;
 
diff --git a/Engine/Plugins/Runtime/GameplayAbilities/Source/GameplayAbilities/Public/AbilitySystemComponent.h b/Engine/Plugins/Runtime/GameplayAbilities/Source/GameplayAbilities/Public/AbilitySystemComponent.h
--- a/Engine/Plugins/Runtime/GameplayAbilities/Source/GameplayAbilities/Public/AbilitySystemComponent.h
+++ b/Engine/Plugins/Runtime/GameplayAbilities/Source/GameplayAbilities/Public/AbilitySystemComponent.h
@@ -1021,7 +1021,7 @@ class GAMEPLAYABILITIES_API UAbilitySystemComponent : public UGameplayTasksCompo
 
 	void AccumulateScreenPos(FAbilitySystemComponentDebugInfo& Info);
 	virtual void Debug_Internal(struct FAbilitySystemComponentDebugInfo& Info);
-	void DebugLine(struct FAbilitySystemComponentDebugInfo& Info, FString Str, float XOffset, float YOffset);
+	void DebugLine(struct FAbilitySystemComponentDebugInfo& Info, FString Str, float XOffset, float YOffset, int32 MinTextRowsToAdvance = 0);
 	FString CleanupName(FString Str);
 
 	/** Print a debug list of all gameplay effects */
