-Preventing server crash by changing an ensure to a log. Still unsure why this is triggering.

diff --git a/Engine/Plugins/Runtime/GameplayAbilities/Source/GameplayAbilities/Private/AbilitySystemBlueprintLibrary.cpp b/Engine/Plugins/Runtime/GameplayAbilities/Source/GameplayAbilities/Private/AbilitySystemBlueprintLibrary.cpp
--- a/Engine/Plugins/Runtime/GameplayAbilities/Source/GameplayAbilities/Private/AbilitySystemBlueprintLibrary.cpp
+++ b/Engine/Plugins/Runtime/GameplayAbilities/Source/GameplayAbilities/Private/AbilitySystemBlueprintLibrary.cpp
@@ -25,13 +25,15 @@ void UAbilitySystemBlueprintLibrary::SendGameplayEventToActor(AActor* Actor, FGa
 		if (AbilitySystemInterface != nullptr)
 		{
 			UAbilitySystemComponent* AbilitySystemComponent = AbilitySystemInterface->GetAbilitySystemComponent();
-			if (AbilitySystemComponent != nullptr)
+			if (AbilitySystemComponent != nullptr && !AbilitySystemComponent->IsPendingKill())
 			{
-				TWeakObjectPtr<UAbilitySystemComponent> ASC = AbilitySystemComponent;
-				ensureMsgf(ASC.IsValid(), TEXT("UAbilitySystemBlueprintLibrary::SendGameplayEventToActor: Invalid ability system component retrieved from Actor %s. EventTag was %s"), *Actor->GetName(), *EventTag.ToString());
 				FScopedPredictionWindow NewScopedWindow(AbilitySystemComponent, true);
 				AbilitySystemComponent->HandleGameplayEvent(EventTag, &Payload);
 			}
+			else
+			{
+				ABILITY_LOG(Error, TEXT("UAbilitySystemBlueprintLibrary::SendGameplayEventToActor: Invalid ability system component retrieved from Actor %s. EventTag was %s"), *Actor->GetName(), *EventTag.ToString());
+			}
 		}
 	}
 }
