Adding AsBool and AsInteger helper to Attribute Set

#tests PIE

diff --git a/Engine/Plugins/Runtime/GameplayAbilities/Source/GameplayAbilities/Private/AttributeSet.cpp b/Engine/Plugins/Runtime/GameplayAbilities/Source/GameplayAbilities/Private/AttributeSet.cpp
--- a/Engine/Plugins/Runtime/GameplayAbilities/Source/GameplayAbilities/Private/AttributeSet.cpp
+++ b/Engine/Plugins/Runtime/GameplayAbilities/Source/GameplayAbilities/Private/AttributeSet.cpp
@@ -449,6 +449,16 @@ float FScalableFloat::GetValue(const FString* ContextString /*= nullptr*/) const
 	return GetValueAtLevel(0, ContextString);
 }
 
+bool FScalableFloat::AsBool(float Level, const FString* ContextString) const
+{
+	return GetValueAtLevel(Level, ContextString) > 0.0f;
+}
+
+int32 FScalableFloat::AsInteger(float Level, const FString* ContextString) const
+{
+	return (int32)GetValueAtLevel(Level, ContextString);
+}
+
 void FScalableFloat::SetValue(float NewValue)
 {
 	Value = NewValue;
diff --git a/Engine/Plugins/Runtime/GameplayAbilities/Source/GameplayAbilities/Public/AttributeSet.h b/Engine/Plugins/Runtime/GameplayAbilities/Source/GameplayAbilities/Public/AttributeSet.h
--- a/Engine/Plugins/Runtime/GameplayAbilities/Source/GameplayAbilities/Public/AttributeSet.h
+++ b/Engine/Plugins/Runtime/GameplayAbilities/Source/GameplayAbilities/Public/AttributeSet.h
@@ -289,6 +289,12 @@ public:
 	/** Returns the scaled value at level 0 */
 	float GetValue(const FString* ContextString = nullptr) const;
 
+	/** Used when using a scalable float as a bool */
+	bool AsBool(float Level = 0, const FString* ContextString = nullptr) const;
+
+	/** Returns the value as an int32 */
+	int32 AsInteger(float Level = 0, const FString* ContextString = nullptr) const;
+
 	/** True if there is no curve lookup */
 	bool IsStatic() const
 	{
