Engine module replicated properties cleanup

Explicitely disabling replicated properties previously not registered in GetLifetimeRuntimeProps.

diff --git a/Engine/Plugins/Runtime/GameplayAbilities/Source/GameplayAbilities/Private/Abilities/GameplayAbilityTargetActor.cpp b/Engine/Plugins/Runtime/GameplayAbilities/Source/GameplayAbilities/Private/Abilities/GameplayAbilityTargetActor.cpp
--- a/Engine/Plugins/Runtime/GameplayAbilities/Source/GameplayAbilities/Private/Abilities/GameplayAbilityTargetActor.cpp
+++ b/Engine/Plugins/Runtime/GameplayAbilities/Source/GameplayAbilities/Private/Abilities/GameplayAbilityTargetActor.cpp
@@ -54,6 +54,8 @@ void AGameplayAbilityTargetActor::GetLifetimeReplicatedProps(TArray< FLifetimePr
 	DOREPLIFETIME(AGameplayAbilityTargetActor, SourceActor);
 	DOREPLIFETIME(AGameplayAbilityTargetActor, bDebug);
 	DOREPLIFETIME(AGameplayAbilityTargetActor, bDestroyOnConfirmation);
+
+	DISABLE_REPLICATED_PROPERTY(AGameplayAbilityTargetActor, Filter);
 }
 
 void AGameplayAbilityTargetActor::StartTargeting(UGameplayAbility* Ability)
diff --git a/Engine/Source/Runtime/Engine/Private/Character.cpp b/Engine/Source/Runtime/Engine/Private/Character.cpp
--- a/Engine/Source/Runtime/Engine/Private/Character.cpp
+++ b/Engine/Source/Runtime/Engine/Private/Character.cpp
@@ -1544,6 +1544,9 @@ void ACharacter::GetLifetimeReplicatedProps( TArray< FLifetimeProperty > & OutLi
 {
 	Super::GetLifetimeReplicatedProps( OutLifetimeProps );
 
+	DISABLE_REPLICATED_PROPERTY(ACharacter, JumpMaxHoldTime);
+	DISABLE_REPLICATED_PROPERTY(ACharacter, JumpMaxCount);
+
 	DOREPLIFETIME_CONDITION( ACharacter, RepRootMotion,						COND_SimulatedOnly );
 	DOREPLIFETIME_CONDITION( ACharacter, ReplicatedBasedMovement,			COND_SimulatedOnly );
 	DOREPLIFETIME_CONDITION( ACharacter, ReplicatedServerLastTransformUpdateTimeStamp, COND_SimulatedOnlyNoReplay );
diff --git a/Engine/Source/Runtime/Engine/Private/Engine/ServerStatReplicator.cpp b/Engine/Source/Runtime/Engine/Private/Engine/ServerStatReplicator.cpp
--- a/Engine/Source/Runtime/Engine/Private/Engine/ServerStatReplicator.cpp
+++ b/Engine/Source/Runtime/Engine/Private/Engine/ServerStatReplicator.cpp
@@ -27,6 +27,9 @@ AServerStatReplicator::AServerStatReplicator(const FObjectInitializer& ObjectIni
 
 void AServerStatReplicator::GetLifetimeReplicatedProps(TArray<FLifetimeProperty>& OutLifetimeProps) const
 {
+	// No need to replicate any inherited properties
+	DISABLE_ALL_CLASS_REPLICATED_PROPERTIES(Super, EFieldIteratorFlags::IncludeSuper);
+
 	DOREPLIFETIME(AServerStatReplicator, Channels);
 	DOREPLIFETIME(AServerStatReplicator, InRate);
 	DOREPLIFETIME(AServerStatReplicator, OutRate);
diff --git a/Engine/Source/Runtime/Engine/Private/Replication.cpp b/Engine/Source/Runtime/Engine/Private/Replication.cpp
--- a/Engine/Source/Runtime/Engine/Private/Replication.cpp
+++ b/Engine/Source/Runtime/Engine/Private/Replication.cpp
@@ -18,6 +18,8 @@ void AMatineeActor::GetLifetimeReplicatedProps( TArray< FLifetimeProperty > & Ou
 {
 	Super::GetLifetimeReplicatedProps( OutLifetimeProps );
 
+	DISABLE_REPLICATED_PROPERTY( AMatineeActor, bPendingStop );
+
 	DOREPLIFETIME_CONDITION( AMatineeActor, MatineeData,		COND_InitialOnly );
 	DOREPLIFETIME_CONDITION( AMatineeActor, GroupActorInfos,	COND_InitialOnly );
 
diff --git a/Engine/Source/Runtime/GameplayTasks/Private/GameplayTasksComponent.cpp b/Engine/Source/Runtime/GameplayTasks/Private/GameplayTasksComponent.cpp
--- a/Engine/Source/Runtime/GameplayTasks/Private/GameplayTasksComponent.cpp
+++ b/Engine/Source/Runtime/GameplayTasks/Private/GameplayTasksComponent.cpp
@@ -150,6 +150,8 @@ void UGameplayTasksComponent::OnTaskEnded(UGameplayTask& Task)
 void UGameplayTasksComponent::GetLifetimeReplicatedProps(TArray< FLifetimeProperty >& OutLifetimeProps) const
 {
 	// Intentionally not calling super: We do not want to replicate bActive which controls ticking. We sometimes need to tick on client predictively.
+	DISABLE_ALL_CLASS_REPLICATED_PROPERTIES(Super, EFieldIteratorFlags::IncludeSuper);
+
 	DOREPLIFETIME_CONDITION(UGameplayTasksComponent, SimulatedTasks, COND_SkipOwner);
 }
 
