Added validity-checked accessors for the AbilitySystemComponent in UGameplayAbility.
Added calls to said accessors where appropriate (i.e. BP_RemoveGameplayEffectFromOwnerWith* functions). I can scale these down to just the called out functions, if adding them anywhere else feels risky.

diff --git a/Engine/Plugins/Runtime/GameplayAbilities/Source/GameplayAbilities/Private/Abilities/GameplayAbility.cpp b/Engine/Plugins/Runtime/GameplayAbilities/Source/GameplayAbilities/Private/Abilities/GameplayAbility.cpp
--- a/Engine/Plugins/Runtime/GameplayAbilities/Source/GameplayAbilities/Private/Abilities/GameplayAbility.cpp
+++ b/Engine/Plugins/Runtime/GameplayAbilities/Source/GameplayAbilities/Private/Abilities/GameplayAbility.cpp
@@ -124,8 +124,8 @@ bool UGameplayAbility::CallRemoteFunction(UFunction* Function, void* Parameters,
 
 void UGameplayAbility::SendGameplayEvent(FGameplayTag EventTag, FGameplayEventData Payload)
 {
-	UAbilitySystemComponent* AbilitySystemComponent = CurrentActorInfo->AbilitySystemComponent.Get();
-	if (ensure(AbilitySystemComponent))
+	UAbilitySystemComponent* const AbilitySystemComponent = GetAbilitySystemComponentFromActorInfo_Unchecked();
+	if (AbilitySystemComponent)
 	{
 		FScopedPredictionWindow NewScopedWindow(AbilitySystemComponent, true);
 		AbilitySystemComponent->HandleGameplayEvent(EventTag, &Payload);
@@ -992,6 +992,22 @@ UAbilitySystemComponent* UGameplayAbility::GetAbilitySystemComponentFromActorInf
 	return CurrentActorInfo->AbilitySystemComponent.Get();
 }
 
+UAbilitySystemComponent* UGameplayAbility::GetAbilitySystemComponentFromActorInfo_Checked() const
+{
+	UAbilitySystemComponent* AbilitySystemComponent = CurrentActorInfo ? CurrentActorInfo->AbilitySystemComponent.Get() : nullptr;
+	check(AbilitySystemComponent);
+
+	return AbilitySystemComponent;
+}
+
+UAbilitySystemComponent* UGameplayAbility::GetAbilitySystemComponentFromActorInfo_Unchecked() const
+{
+	UAbilitySystemComponent* AbilitySystemComponent = CurrentActorInfo ? CurrentActorInfo->AbilitySystemComponent.Get() : nullptr;
+	ensure(AbilitySystemComponent);
+
+	return AbilitySystemComponent;
+}
+
 FGameplayEffectSpecHandle UGameplayAbility::MakeOutgoingGameplayEffectSpec(TSubclassOf<UGameplayEffect> GameplayEffectClass, float Level) const
 {
 	check(CurrentActorInfo && CurrentActorInfo->AbilitySystemComponent.IsValid());
@@ -1063,7 +1079,8 @@ bool UGameplayAbility::K2_CommitAbilityCooldown(bool BroadcastCommitEvent, bool
 	check(CurrentActorInfo);
 	if (BroadcastCommitEvent)
 	{
-		CurrentActorInfo->AbilitySystemComponent->NotifyAbilityCommit(this);
+		UAbilitySystemComponent* const AbilitySystemComponent = GetAbilitySystemComponentFromActorInfo_Checked();
+		AbilitySystemComponent->NotifyAbilityCommit(this);
 	}
 	return CommitAbilityCooldown(CurrentSpecHandle, CurrentActorInfo, CurrentActivationInfo, ForceCooldown);
 }
@@ -1073,7 +1090,8 @@ bool UGameplayAbility::K2_CommitAbilityCost(bool BroadcastCommitEvent)
 	check(CurrentActorInfo);
 	if (BroadcastCommitEvent)
 	{
-		CurrentActorInfo->AbilitySystemComponent->NotifyAbilityCommit(this);
+		UAbilitySystemComponent* const AbilitySystemComponent = GetAbilitySystemComponentFromActorInfo_Checked();
+		AbilitySystemComponent->NotifyAbilityCommit(this);
 	}
 	return CommitAbilityCost(CurrentSpecHandle, CurrentActorInfo, CurrentActivationInfo);
 }
@@ -1103,7 +1121,7 @@ void UGameplayAbility::MontageJumpToSection(FName SectionName)
 {
 	check(CurrentActorInfo);
 
-	UAbilitySystemComponent* const AbilitySystemComponent = CurrentActorInfo->AbilitySystemComponent.Get();
+	UAbilitySystemComponent* const AbilitySystemComponent = GetAbilitySystemComponentFromActorInfo_Checked();
 	if (AbilitySystemComponent->IsAnimatingAbility(this))
 	{
 		AbilitySystemComponent->CurrentMontageJumpToSection(SectionName);
@@ -1114,7 +1132,7 @@ void UGameplayAbility::MontageSetNextSectionName(FName FromSectionName, FName To
 {
 	check(CurrentActorInfo);
 
-	UAbilitySystemComponent* const AbilitySystemComponent = CurrentActorInfo->AbilitySystemComponent.Get();
+	UAbilitySystemComponent* const AbilitySystemComponent = GetAbilitySystemComponentFromActorInfo_Checked();
 	if (AbilitySystemComponent->IsAnimatingAbility(this))
 	{
 		AbilitySystemComponent->CurrentMontageSetNextSectionName(FromSectionName, ToSectionName);
@@ -1336,14 +1354,18 @@ void UGameplayAbility::AddAbilityTaskDebugMessage(UGameplayTask* AbilityTask, FS
 void UGameplayAbility::K2_ExecuteGameplayCue(FGameplayTag GameplayCueTag, FGameplayEffectContextHandle Context)
 {
 	check(CurrentActorInfo);
-	CurrentActorInfo->AbilitySystemComponent->ExecuteGameplayCue(GameplayCueTag, Context);
+
+	UAbilitySystemComponent* const AbilitySystemComponent = GetAbilitySystemComponentFromActorInfo_Checked();
+	AbilitySystemComponent->ExecuteGameplayCue(GameplayCueTag, Context);
 }
 
 void UGameplayAbility::K2_ExecuteGameplayCueWithParams(FGameplayTag GameplayCueTag, const FGameplayCueParameters& GameplayCueParameters)
 {
 	check(CurrentActorInfo);
 	const_cast<FGameplayCueParameters&>(GameplayCueParameters).AbilityLevel = GetAbilityLevel();
-	CurrentActorInfo->AbilitySystemComponent->ExecuteGameplayCue(GameplayCueTag, GameplayCueParameters);
+
+	UAbilitySystemComponent* const AbilitySystemComponent = GetAbilitySystemComponentFromActorInfo_Checked();
+	AbilitySystemComponent->ExecuteGameplayCue(GameplayCueTag, GameplayCueParameters);
 }
 
 void UGameplayAbility::K2_AddGameplayCue(FGameplayTag GameplayCueTag, FGameplayEffectContextHandle Context, bool bRemoveOnAbilityEnd)
@@ -1358,7 +1380,8 @@ void UGameplayAbility::K2_AddGameplayCue(FGameplayTag GameplayCueTag, FGameplayE
 
 	Context.SetAbility(this);
 
-	CurrentActorInfo->AbilitySystemComponent->AddGameplayCue(GameplayCueTag, Context);
+	UAbilitySystemComponent* const AbilitySystemComponent = GetAbilitySystemComponentFromActorInfo_Checked();
+	AbilitySystemComponent->AddGameplayCue(GameplayCueTag, Context);
 
 	if (bRemoveOnAbilityEnd)
 	{
@@ -1370,7 +1393,8 @@ void UGameplayAbility::K2_AddGameplayCueWithParams(FGameplayTag GameplayCueTag,
 {
 	check(CurrentActorInfo);
 
-	CurrentActorInfo->AbilitySystemComponent->AddGameplayCue(GameplayCueTag, GameplayCueParameter);
+	UAbilitySystemComponent* const AbilitySystemComponent = GetAbilitySystemComponentFromActorInfo_Checked();
+	AbilitySystemComponent->AddGameplayCue(GameplayCueTag, GameplayCueParameter);
 
 	if (bRemoveOnAbilityEnd)
 	{
@@ -1382,7 +1406,9 @@ void UGameplayAbility::K2_AddGameplayCueWithParams(FGameplayTag GameplayCueTag,
 void UGameplayAbility::K2_RemoveGameplayCue(FGameplayTag GameplayCueTag)
 {
 	check(CurrentActorInfo);
-	CurrentActorInfo->AbilitySystemComponent->RemoveGameplayCue(GameplayCueTag);
+
+	UAbilitySystemComponent* const AbilitySystemComponent = GetAbilitySystemComponentFromActorInfo_Checked();
+	AbilitySystemComponent->RemoveGameplayCue(GameplayCueTag);
 
 	TrackedGameplayCues.Remove(GameplayCueTag);
 }
@@ -1432,7 +1458,9 @@ FGameplayAbilitySpec* UGameplayAbility::GetCurrentAbilitySpec() const
 {
 	check(IsInstantiated()); // You should not call this on non instanced abilities.
 	check(CurrentActorInfo);
-	return CurrentActorInfo->AbilitySystemComponent->FindAbilitySpecFromHandle(CurrentSpecHandle);
+
+	UAbilitySystemComponent* const AbilitySystemComponent = GetAbilitySystemComponentFromActorInfo_Checked();
+	return AbilitySystemComponent->FindAbilitySpecFromHandle(CurrentSpecHandle);
 }
 
 FGameplayEffectContextHandle UGameplayAbility::GetGrantedByEffectContext() const
@@ -1441,7 +1469,7 @@ FGameplayEffectContextHandle UGameplayAbility::GetGrantedByEffectContext() const
 	check(CurrentActorInfo);
 	if (CurrentActorInfo)
 	{
-		UAbilitySystemComponent* const AbilitySystemComponent = CurrentActorInfo->AbilitySystemComponent.Get();
+		UAbilitySystemComponent* const AbilitySystemComponent = GetAbilitySystemComponentFromActorInfo_Checked();
 		FActiveGameplayEffectHandle ActiveHandle = AbilitySystemComponent->FindActiveGameplayEffectHandle(GetCurrentAbilitySpecHandle());
 		if (ActiveHandle.IsValid())
 		{
@@ -1458,7 +1486,7 @@ void UGameplayAbility::RemoveGrantedByEffect()
 	check(CurrentActorInfo);
 	if (CurrentActorInfo)
 	{
-		UAbilitySystemComponent* const AbilitySystemComponent = CurrentActorInfo->AbilitySystemComponent.Get();
+		UAbilitySystemComponent* const AbilitySystemComponent = GetAbilitySystemComponentFromActorInfo_Checked();
 		FActiveGameplayEffectHandle ActiveHandle = AbilitySystemComponent->FindActiveGameplayEffectHandle(GetCurrentAbilitySpecHandle());
 		if (ActiveHandle.IsValid())
 		{
@@ -1768,10 +1796,11 @@ void UGameplayAbility::BP_RemoveGameplayEffectFromOwnerWithAssetTags(FGameplayTa
 		return;
 	}
 
-	if (CurrentActorInfo)
+	UAbilitySystemComponent* const AbilitySystemComponent = GetAbilitySystemComponentFromActorInfo_Unchecked();
+	if (AbilitySystemComponent)
 	{
 		FGameplayEffectQuery const Query = FGameplayEffectQuery::MakeQuery_MatchAnyEffectTags(WithTags);
-		CurrentActorInfo->AbilitySystemComponent->RemoveActiveEffects(Query, StacksToRemove);
+		AbilitySystemComponent->RemoveActiveEffects(Query, StacksToRemove);
 	}
 }
 
@@ -1782,10 +1811,11 @@ void UGameplayAbility::BP_RemoveGameplayEffectFromOwnerWithGrantedTags(FGameplay
 		return;
 	}
 
-	if (CurrentActorInfo)
+	UAbilitySystemComponent* const AbilitySystemComponent = GetAbilitySystemComponentFromActorInfo_Unchecked();
+	if (AbilitySystemComponent)
 	{
 		FGameplayEffectQuery const Query = FGameplayEffectQuery::MakeQuery_MatchAnyOwningTags(WithGrantedTags);
-		CurrentActorInfo->AbilitySystemComponent->RemoveActiveEffects(Query, StacksToRemove);
+		AbilitySystemComponent->RemoveActiveEffects(Query, StacksToRemove);
 	}
 }
 
@@ -1796,9 +1826,10 @@ void UGameplayAbility::BP_RemoveGameplayEffectFromOwnerWithHandle(FActiveGamepla
 		return;
 	}
 
-	if (CurrentActorInfo)
+	UAbilitySystemComponent* const AbilitySystemComponent = GetAbilitySystemComponentFromActorInfo_Unchecked();
+	if (AbilitySystemComponent)
 	{
-		CurrentActorInfo->AbilitySystemComponent->RemoveActiveGameplayEffect(Handle, StacksToRemove);
+		AbilitySystemComponent->RemoveActiveGameplayEffect(Handle, StacksToRemove);
 	}
 }
 
@@ -1868,8 +1899,7 @@ void UGameplayAbility::NotifyAbilityTaskWaitingOnPlayerData(class UAbilityTask*
 {
 	// This should never happen since it will only be called from actively running ability tasks
 	check(CurrentActorInfo);
-	UAbilitySystemComponent* const AbilitySystemComponent = CurrentActorInfo->AbilitySystemComponent.Get();
-	check(AbilitySystemComponent);
+	UAbilitySystemComponent* const AbilitySystemComponent = GetAbilitySystemComponentFromActorInfo_Checked();
 
 	if (RemoteInstanceEnded)
 	{
@@ -1883,6 +1913,8 @@ void UGameplayAbility::NotifyAbilityTaskWaitingOnAvatar(class UAbilityTask* Abil
 	if (CurrentActorInfo && CurrentActorInfo->AvatarActor.IsValid() == false)
 	{
 		ABILITY_LOG(Log, TEXT("Ability %s is force cancelling because Task %s has started while there is no valid AvatarActor"), *GetName(), *AbilityTask->GetDebugString());
-		CurrentActorInfo->AbilitySystemComponent->ForceCancelAbilityDueToReplication(this);
+
+		UAbilitySystemComponent* const AbilitySystemComponent = GetAbilitySystemComponentFromActorInfo_Checked();
+		AbilitySystemComponent->ForceCancelAbilityDueToReplication(this);
 	}
 }
diff --git a/Engine/Plugins/Runtime/GameplayAbilities/Source/GameplayAbilities/Public/Abilities/GameplayAbility.h b/Engine/Plugins/Runtime/GameplayAbilities/Source/GameplayAbilities/Public/Abilities/GameplayAbility.h
--- a/Engine/Plugins/Runtime/GameplayAbilities/Source/GameplayAbilities/Public/Abilities/GameplayAbility.h
+++ b/Engine/Plugins/Runtime/GameplayAbilities/Source/GameplayAbilities/Public/Abilities/GameplayAbility.h
@@ -163,6 +163,8 @@ public:
 	/** Returns the AbilitySystemComponent that is activating this ability */
 	UFUNCTION(BlueprintCallable, Category = Ability)
 	UAbilitySystemComponent* GetAbilitySystemComponentFromActorInfo() const;
+	UAbilitySystemComponent* GetAbilitySystemComponentFromActorInfo_Checked() const;
+	UAbilitySystemComponent* GetAbilitySystemComponentFromActorInfo_Unchecked() const;
 
 	/** Gets the current actor info bound to this ability - can only be called on instanced abilities. */
 	const FGameplayAbilityActorInfo* GetCurrentActorInfo() const
