Renamed GetAbilitySystemComponentFromActorInfo_Unchecked to GetAbilitySystemComponentFromActorInfo_Ensured.

diff --git a/Engine/Plugins/Runtime/GameplayAbilities/Source/GameplayAbilities/Private/Abilities/GameplayAbility.cpp b/Engine/Plugins/Runtime/GameplayAbilities/Source/GameplayAbilities/Private/Abilities/GameplayAbility.cpp
--- a/Engine/Plugins/Runtime/GameplayAbilities/Source/GameplayAbilities/Private/Abilities/GameplayAbility.cpp
+++ b/Engine/Plugins/Runtime/GameplayAbilities/Source/GameplayAbilities/Private/Abilities/GameplayAbility.cpp
@@ -124,7 +124,7 @@ bool UGameplayAbility::CallRemoteFunction(UFunction* Function, void* Parameters,
 
 void UGameplayAbility::SendGameplayEvent(FGameplayTag EventTag, FGameplayEventData Payload)
 {
-	UAbilitySystemComponent* const AbilitySystemComponent = GetAbilitySystemComponentFromActorInfo_Unchecked();
+	UAbilitySystemComponent* const AbilitySystemComponent = GetAbilitySystemComponentFromActorInfo_Ensured();
 	if (AbilitySystemComponent)
 	{
 		FScopedPredictionWindow NewScopedWindow(AbilitySystemComponent, true);
@@ -1000,7 +1000,7 @@ UAbilitySystemComponent* UGameplayAbility::GetAbilitySystemComponentFromActorInf
 	return AbilitySystemComponent;
 }
 
-UAbilitySystemComponent* UGameplayAbility::GetAbilitySystemComponentFromActorInfo_Unchecked() const
+UAbilitySystemComponent* UGameplayAbility::GetAbilitySystemComponentFromActorInfo_Ensured() const
 {
 	UAbilitySystemComponent* AbilitySystemComponent = CurrentActorInfo ? CurrentActorInfo->AbilitySystemComponent.Get() : nullptr;
 	ensure(AbilitySystemComponent);
@@ -1796,7 +1796,7 @@ void UGameplayAbility::BP_RemoveGameplayEffectFromOwnerWithAssetTags(FGameplayTa
 		return;
 	}
 
-	UAbilitySystemComponent* const AbilitySystemComponent = GetAbilitySystemComponentFromActorInfo_Unchecked();
+	UAbilitySystemComponent* const AbilitySystemComponent = GetAbilitySystemComponentFromActorInfo_Ensured();
 	if (AbilitySystemComponent)
 	{
 		FGameplayEffectQuery const Query = FGameplayEffectQuery::MakeQuery_MatchAnyEffectTags(WithTags);
@@ -1811,7 +1811,7 @@ void UGameplayAbility::BP_RemoveGameplayEffectFromOwnerWithGrantedTags(FGameplay
 		return;
 	}
 
-	UAbilitySystemComponent* const AbilitySystemComponent = GetAbilitySystemComponentFromActorInfo_Unchecked();
+	UAbilitySystemComponent* const AbilitySystemComponent = GetAbilitySystemComponentFromActorInfo_Ensured();
 	if (AbilitySystemComponent)
 	{
 		FGameplayEffectQuery const Query = FGameplayEffectQuery::MakeQuery_MatchAnyOwningTags(WithGrantedTags);
@@ -1826,7 +1826,7 @@ void UGameplayAbility::BP_RemoveGameplayEffectFromOwnerWithHandle(FActiveGamepla
 		return;
 	}
 
-	UAbilitySystemComponent* const AbilitySystemComponent = GetAbilitySystemComponentFromActorInfo_Unchecked();
+	UAbilitySystemComponent* const AbilitySystemComponent = GetAbilitySystemComponentFromActorInfo_Ensured();
 	if (AbilitySystemComponent)
 	{
 		AbilitySystemComponent->RemoveActiveGameplayEffect(Handle, StacksToRemove);
diff --git a/Engine/Plugins/Runtime/GameplayAbilities/Source/GameplayAbilities/Public/Abilities/GameplayAbility.h b/Engine/Plugins/Runtime/GameplayAbilities/Source/GameplayAbilities/Public/Abilities/GameplayAbility.h
--- a/Engine/Plugins/Runtime/GameplayAbilities/Source/GameplayAbilities/Public/Abilities/GameplayAbility.h
+++ b/Engine/Plugins/Runtime/GameplayAbilities/Source/GameplayAbilities/Public/Abilities/GameplayAbility.h
@@ -164,7 +164,7 @@ public:
 	UFUNCTION(BlueprintCallable, Category = Ability)
 	UAbilitySystemComponent* GetAbilitySystemComponentFromActorInfo() const;
 	UAbilitySystemComponent* GetAbilitySystemComponentFromActorInfo_Checked() const;
-	UAbilitySystemComponent* GetAbilitySystemComponentFromActorInfo_Unchecked() const;
+	UAbilitySystemComponent* GetAbilitySystemComponentFromActorInfo_Ensured() const;
 
 	/** Gets the current actor info bound to this ability - can only be called on instanced abilities. */
 	const FGameplayAbilityActorInfo* GetCurrentActorInfo() const
