Gameplay tag blueprint properties;

* Added a generic way of allowing blueprints to map gameplay tag counts to blueprint variables by adding a FGameplayTagBlueprintPropertyMap to the Blueprintable class.

* Added UnregisterGameplayTagEvent() to UAbilitySystemComponent.  This allows you to remove gameplay tag events that have been registered.

* Updated RegisterAndCallGameplayTagEvent() on UAbilitySystemComponent to return the delegate handle.  This is needed so the gameplay tag event can be removed later.

diff --git a/Engine/Plugins/Runtime/GameplayAbilities/Source/GameplayAbilities/Private/AbilitySystemComponent.cpp b/Engine/Plugins/Runtime/GameplayAbilities/Source/GameplayAbilities/Private/AbilitySystemComponent.cpp
--- a/Engine/Plugins/Runtime/GameplayAbilities/Source/GameplayAbilities/Private/AbilitySystemComponent.cpp
+++ b/Engine/Plugins/Runtime/GameplayAbilities/Source/GameplayAbilities/Private/AbilitySystemComponent.cpp
@@ -468,15 +468,22 @@ FOnGameplayEffectTagCountChanged& UAbilitySystemComponent::RegisterGameplayTagEv
 	return GameplayTagCountContainer.RegisterGameplayTagEvent(Tag, EventType);
 }
 
-void UAbilitySystemComponent::RegisterAndCallGameplayTagEvent(FGameplayTag Tag, FOnGameplayEffectTagCountChanged::FDelegate Delegate, EGameplayTagEventType::Type EventType)
+void UAbilitySystemComponent::UnregisterGameplayTagEvent(FDelegateHandle DelegateHandle, FGameplayTag Tag, EGameplayTagEventType::Type EventType)
 {
-	GameplayTagCountContainer.RegisterGameplayTagEvent(Tag, EventType).Add(Delegate);
+	GameplayTagCountContainer.RegisterGameplayTagEvent(Tag, EventType).Remove(DelegateHandle);
+}
+
+FDelegateHandle UAbilitySystemComponent::RegisterAndCallGameplayTagEvent(FGameplayTag Tag, FOnGameplayEffectTagCountChanged::FDelegate Delegate, EGameplayTagEventType::Type EventType)
+{
+	FDelegateHandle DelegateHandle = GameplayTagCountContainer.RegisterGameplayTagEvent(Tag, EventType).Add(Delegate);
 
 	const int32 TagCount = GetTagCount(Tag);
 	if (TagCount > 0)
 	{
 		Delegate.Execute(Tag, TagCount);
 	}
+
+	return DelegateHandle;
 }
 
 FOnGameplayEffectTagCountChanged& UAbilitySystemComponent::RegisterGenericGameplayTagEvent()
diff --git a/Engine/Plugins/Runtime/GameplayAbilities/Source/GameplayAbilities/Private/GameplayEffectTypes.cpp b/Engine/Plugins/Runtime/GameplayAbilities/Source/GameplayAbilities/Private/GameplayEffectTypes.cpp
--- a/Engine/Plugins/Runtime/GameplayAbilities/Source/GameplayAbilities/Private/GameplayEffectTypes.cpp
+++ b/Engine/Plugins/Runtime/GameplayAbilities/Source/GameplayAbilities/Private/GameplayEffectTypes.cpp
@@ -10,6 +10,10 @@
 #include "AbilitySystemComponent.h"
 #include "AbilitySystemInterface.h"
 
+
+#define LOCTEXT_NAMESPACE "GameplayEffectTypes"
+
+
 #if WITH_EDITORONLY_DATA
 const FName FGameplayModEvaluationChannelSettings::ForceHideMetadataKey(TEXT("ForceHideEvaluationChannel"));
 const FString FGameplayModEvaluationChannelSettings::ForceHideMetadataEnabledValue(TEXT("True"));
@@ -492,6 +496,151 @@ bool FGameplayTagCountContainer::UpdateTagMap_Internal(const FGameplayTag& Tag,
 	return CreatedSignificantChange;
 }
 
+FGameplayTagBlueprintPropertyMap::FGameplayTagBlueprintPropertyMap()
+{
+}
+
+FGameplayTagBlueprintPropertyMap::~FGameplayTagBlueprintPropertyMap()
+{
+	if (UAbilitySystemComponent* ASC = CachedASC.Get())
+	{
+		for (FGameplayTagBlueprintPropertyMapping& Mapping : PropertyMappings)
+		{
+			if (Mapping.PropertyToEdit && Mapping.TagToMap.IsValid())
+			{
+				ASC->UnregisterGameplayTagEvent(Mapping.DelegateHandle, Mapping.TagToMap, GetGameplayTagEventType(Mapping.PropertyToEdit));
+			}
+		}
+	}
+}
+
+#if WITH_EDITOR
+EDataValidationResult FGameplayTagBlueprintPropertyMap::IsDataValid(UObject* Owner, TArray<FText>& ValidationErrors)
+{
+	UClass* OwnerClass = ((Owner != nullptr) ? Owner->GetClass() : nullptr);
+	if (!OwnerClass)
+	{
+		ABILITY_LOG(Error, TEXT("FGameplayTagBlueprintPropertyMap: IsDataValid() called with an invalid Owner."));
+		return EDataValidationResult::Invalid;
+	}
+
+	for (const FGameplayTagBlueprintPropertyMapping& Mapping : PropertyMappings)
+	{
+		if (!Mapping.TagToMap.IsValid())
+		{
+			ValidationErrors.Add(FText::Format(LOCTEXT("GameplayTagBlueprintPropertyMap_BadTag", "The gameplay tag [{0}] for property [{1}] is empty or invalid."),
+				FText::AsCultureInvariant(Mapping.TagToMap.ToString()),
+				FText::FromName(Mapping.PropertyName)));
+		}
+
+		if (UProperty* Property = OwnerClass->FindPropertyByName(Mapping.PropertyName))
+		{
+			if (!IsPropertyTypeValid(Property))
+			{
+				ValidationErrors.Add(FText::Format(LOCTEXT("GameplayTagBlueprintPropertyMap_BadType", "The property [{0}] for gameplay tag [{1}] is not a supported type.  Supported types are: integer, float, and boolean."),
+					FText::FromName(Mapping.PropertyName),
+					FText::AsCultureInvariant(Mapping.TagToMap.ToString())));
+			}
+		}
+		else
+		{
+			ValidationErrors.Add(FText::Format(LOCTEXT("GameplayTagBlueprintPropertyMap_MissingProperty", "The property [{0}] for gameplay tag [{1}] could not be found."),
+				FText::FromName(Mapping.PropertyName),
+				FText::AsCultureInvariant(Mapping.TagToMap.ToString())));
+		}
+	}
+
+	return ((ValidationErrors.Num() > 0) ? EDataValidationResult::Invalid : EDataValidationResult::Valid);
+}
+#endif // #if WITH_EDITOR
+
+void FGameplayTagBlueprintPropertyMap::Initialize(UObject* Owner, UAbilitySystemComponent* ASC)
+{
+	UClass* OwnerClass = (Owner ? Owner->GetClass() : nullptr);
+	if (!OwnerClass)
+	{
+		ABILITY_LOG(Error, TEXT("FGameplayTagBlueprintPropertyMap: Initialize() called with an invalid Owner."));
+		return;
+	}
+
+	if (!ASC)
+	{
+		ABILITY_LOG(Error, TEXT("FGameplayTagBlueprintPropertyMap: Initialize() called with an invalid AbilitySystemComponent."));
+		return;
+	}
+
+	CachedOwner = Owner;
+	CachedASC = ASC;
+
+	FOnGameplayEffectTagCountChanged::FDelegate Delegate = FOnGameplayEffectTagCountChanged::FDelegate::CreateRaw(this, &FGameplayTagBlueprintPropertyMap::GameplayTagEventCallback);
+
+	// Process array starting at the end so we can remove invalid entries.
+	for (int32 MappingIndex = (PropertyMappings.Num() - 1); MappingIndex >= 0; --MappingIndex)
+	{
+		FGameplayTagBlueprintPropertyMapping& Mapping = PropertyMappings[MappingIndex];
+
+		if (Mapping.TagToMap.IsValid())
+		{
+			UProperty* Property = OwnerClass->FindPropertyByName(Mapping.PropertyName);
+			if (Property && IsPropertyTypeValid(Property))
+			{
+				Mapping.PropertyToEdit = Property;
+				Mapping.DelegateHandle = ASC->RegisterAndCallGameplayTagEvent(Mapping.TagToMap, Delegate, GetGameplayTagEventType(Property));
+				continue;
+			}
+		}
+
+		// Entry was invalid.  Remove it from the array.
+		ABILITY_LOG(Error, TEXT("FGameplayTagBlueprintPropertyMap: Removing invalid GameplayTagBlueprintPropertyMapping [Index: %d, Tag:%s, Property:%s] for [%s]."),
+			MappingIndex, *Mapping.TagToMap.ToString(), *Mapping.PropertyName.ToString(), *GetNameSafe(Owner));
+
+		PropertyMappings.RemoveAtSwap(MappingIndex, 1, false);
+	}
+}
+
+void FGameplayTagBlueprintPropertyMap::GameplayTagEventCallback(const FGameplayTag Tag, int32 NewCount)
+{
+	UObject* Owner = CachedOwner.Get();
+	if (!Owner)
+	{
+		ABILITY_LOG(Warning, TEXT("FGameplayTagBlueprintPropertyMap::GameplayTagEventCallback has an invalid Owner."));
+		return;
+	}
+
+	FGameplayTagBlueprintPropertyMapping* Mapping = PropertyMappings.FindByPredicate([Tag](const FGameplayTagBlueprintPropertyMapping& Test)
+	{
+		return (Tag == Test.TagToMap);
+	});
+
+	if (Mapping && Mapping->PropertyToEdit)
+	{
+		if (const UBoolProperty* BoolProperty = Cast<const UBoolProperty>(Mapping->PropertyToEdit))
+		{
+			BoolProperty->SetPropertyValue_InContainer(Owner, NewCount > 0);
+		}
+		else if (const UIntProperty* IntProperty = Cast<const UIntProperty>(Mapping->PropertyToEdit))
+		{
+			IntProperty->SetPropertyValue_InContainer(Owner, NewCount);
+		}
+		else if (const UFloatProperty* FloatProperty = Cast<const UFloatProperty>(Mapping->PropertyToEdit))
+		{
+			FloatProperty->SetPropertyValue_InContainer(Owner, (float)NewCount);
+		}
+	}
+}
+
+bool FGameplayTagBlueprintPropertyMap::IsPropertyTypeValid(const UProperty* Property) const
+{
+	check(Property);
+	return (Property->IsA<UBoolProperty>() || Property->IsA<UIntProperty>() || Property->IsA<UFloatProperty>());
+}
+
+EGameplayTagEventType::Type FGameplayTagBlueprintPropertyMap::GetGameplayTagEventType(const UProperty* Property) const
+{
+	check(Property);
+	return (Property->IsA(UBoolProperty::StaticClass()) ? EGameplayTagEventType::NewOrRemoved : EGameplayTagEventType::AnyCountChange);
+}
+
 bool FGameplayTagRequirements::RequirementsMet(const FGameplayTagContainer& Container) const
 {
 	bool HasRequired = Container.HasAll(RequireTags);
diff --git a/Engine/Plugins/Runtime/GameplayAbilities/Source/GameplayAbilities/Public/AbilitySystemComponent.h b/Engine/Plugins/Runtime/GameplayAbilities/Source/GameplayAbilities/Public/AbilitySystemComponent.h
--- a/Engine/Plugins/Runtime/GameplayAbilities/Source/GameplayAbilities/Public/AbilitySystemComponent.h
+++ b/Engine/Plugins/Runtime/GameplayAbilities/Source/GameplayAbilities/Public/AbilitySystemComponent.h
@@ -615,8 +615,11 @@ class GAMEPLAYABILITIES_API UAbilitySystemComponent : public UGameplayTasksCompo
 	/** Allow events to be registered for specific gameplay tags being added or removed */
 	FOnGameplayEffectTagCountChanged& RegisterGameplayTagEvent(FGameplayTag Tag, EGameplayTagEventType::Type EventType=EGameplayTagEventType::NewOrRemoved);
 
+	/** Unregister previously added events */
+	void UnregisterGameplayTagEvent(FDelegateHandle DelegateHandle, FGameplayTag Tag, EGameplayTagEventType::Type EventType=EGameplayTagEventType::NewOrRemoved);
+
 	/** Register a tag event and immediately call it */
-	void RegisterAndCallGameplayTagEvent(FGameplayTag Tag, FOnGameplayEffectTagCountChanged::FDelegate Delegate, EGameplayTagEventType::Type EventType=EGameplayTagEventType::NewOrRemoved);
+	FDelegateHandle RegisterAndCallGameplayTagEvent(FGameplayTag Tag, FOnGameplayEffectTagCountChanged::FDelegate Delegate, EGameplayTagEventType::Type EventType=EGameplayTagEventType::NewOrRemoved);
 
 	/** Returns multicast delegate that is invoked whenever a tag is added or removed (but not if just count is increased. Only for 'new' and 'removed' events) */
 	FOnGameplayEffectTagCountChanged& RegisterGenericGameplayTagEvent();
diff --git a/Engine/Plugins/Runtime/GameplayAbilities/Source/GameplayAbilities/Public/GameplayEffectTypes.h b/Engine/Plugins/Runtime/GameplayAbilities/Source/GameplayAbilities/Public/GameplayEffectTypes.h
--- a/Engine/Plugins/Runtime/GameplayAbilities/Source/GameplayAbilities/Public/GameplayEffectTypes.h
+++ b/Engine/Plugins/Runtime/GameplayAbilities/Source/GameplayAbilities/Public/GameplayEffectTypes.h
@@ -1324,6 +1324,81 @@ private:
 };
 
 
+/**
+ * Struct used to update a blueprint property with a gameplay tag count.
+ * The property is automatically updated as the gameplay tag count changes.
+ * It only supports boolean, integer, and float properties.
+ */
+USTRUCT()
+struct GAMEPLAYABILITIES_API FGameplayTagBlueprintPropertyMapping
+{
+	GENERATED_BODY()
+
+public:
+
+	/** Gameplay tag being counted. */
+	UPROPERTY(EditAnywhere, Category = GameplayTagBlueprintProperty)
+	FGameplayTag TagToMap;
+
+	/** Property to update with the gameplay tag count. */
+	UPROPERTY(VisibleAnywhere, Category = GameplayTagBlueprintProperty)
+	UProperty* PropertyToEdit;
+
+	/** Name of property being edited. */
+	UPROPERTY(VisibleAnywhere, Category = GameplayTagBlueprintProperty)
+	FName PropertyName;
+
+	/** Guid of property being edited. */
+	UPROPERTY(VisibleAnywhere, Category = GameplayTagBlueprintProperty)
+	FGuid PropertyGuid;
+
+	/** Handle to delegate bound on the ability system component. */
+	FDelegateHandle DelegateHandle;
+};
+
+
+/**
+ * Struct used to manage gameplay tag blueprint property mappings.
+ * It registers the properties with delegates on an ability system component.
+ * This struct should not be used in containers (such as TArray) since it uses a raw pointer
+ * to bind the delegate and it's address could change causing an invalid binding.
+ */
+USTRUCT()
+struct GAMEPLAYABILITIES_API FGameplayTagBlueprintPropertyMap
+{
+	GENERATED_BODY()
+
+public:
+
+	FGameplayTagBlueprintPropertyMap();
+	~FGameplayTagBlueprintPropertyMap();
+
+	/** Call this to initialize and bind the properties with the ability system component. */
+	void Initialize(UObject* Owner, class UAbilitySystemComponent* ASC);
+
+#if WITH_EDITOR
+	/** This can optionally be called in the owner's IsDataValid() for data validation. */
+	EDataValidationResult IsDataValid(UObject* ContainingAsset, TArray<FText>& ValidationErrors);
+#endif // #if WITH_EDITOR
+
+protected:
+
+	void GameplayTagEventCallback(const FGameplayTag Tag, int32 NewCount);
+
+	bool IsPropertyTypeValid(const UProperty* Property) const;
+
+	EGameplayTagEventType::Type GetGameplayTagEventType(const UProperty* Property) const;
+
+protected:
+
+	TWeakObjectPtr<UObject> CachedOwner;
+	TWeakObjectPtr<UAbilitySystemComponent> CachedASC;
+
+	UPROPERTY(EditAnywhere, Category = GameplayTagBlueprintProperty)
+	TArray<FGameplayTagBlueprintPropertyMapping> PropertyMappings;
+};
+
+
 /** Encapsulate require and ignore tags */
 USTRUCT(BlueprintType)
 struct GAMEPLAYABILITIES_API FGameplayTagRequirements
diff --git a/Engine/Plugins/Runtime/GameplayAbilities/Source/GameplayAbilitiesEditor/Private/GameplayAbilitiesEditorModule.cpp b/Engine/Plugins/Runtime/GameplayAbilities/Source/GameplayAbilitiesEditor/Private/GameplayAbilitiesEditorModule.cpp
--- a/Engine/Plugins/Runtime/GameplayAbilities/Source/GameplayAbilitiesEditor/Private/GameplayAbilitiesEditorModule.cpp
+++ b/Engine/Plugins/Runtime/GameplayAbilities/Source/GameplayAbilitiesEditor/Private/GameplayAbilitiesEditorModule.cpp
@@ -22,6 +22,7 @@
 #include "AbilitySystemGlobals.h"
 #include "GameplayEffectDetails.h"
 #include "GameplayEffectExecutionScopedModifierInfoDetails.h"
+#include "GameplayTagBlueprintPropertyMappingDetails.h"
 #include "GameplayEffectExecutionDefinitionDetails.h"
 #include "GameplayEffectModifierMagnitudeDetails.h"
 #include "GameplayModEvaluationChannelSettingsDetails.h"
@@ -128,6 +129,7 @@ void FGameplayAbilitiesEditorModule::StartupModule()
 	PropertyModule.RegisterCustomPropertyTypeLayout( "GameplayAttribute", FOnGetPropertyTypeCustomizationInstance::CreateStatic( &FAttributePropertyDetails::MakeInstance ) );
 	PropertyModule.RegisterCustomPropertyTypeLayout( "ScalableFloat", FOnGetPropertyTypeCustomizationInstance::CreateStatic( &FScalableFloatDetails::MakeInstance ) );
 	PropertyModule.RegisterCustomPropertyTypeLayout( "GameplayEffectExecutionScopedModifierInfo", FOnGetPropertyTypeCustomizationInstance::CreateStatic( &FGameplayEffectExecutionScopedModifierInfoDetails::MakeInstance ) );
+	PropertyModule.RegisterCustomPropertyTypeLayout( "GameplayTagBlueprintPropertyMapping", FOnGetPropertyTypeCustomizationInstance::CreateStatic( &FGameplayTagBlueprintPropertyMappingDetails::MakeInstance ) );
 	PropertyModule.RegisterCustomPropertyTypeLayout( "GameplayEffectExecutionDefinition", FOnGetPropertyTypeCustomizationInstance::CreateStatic( &FGameplayEffectExecutionDefinitionDetails::MakeInstance ) );
 	PropertyModule.RegisterCustomPropertyTypeLayout( "GameplayEffectModifierMagnitude", FOnGetPropertyTypeCustomizationInstance::CreateStatic( &FGameplayEffectModifierMagnitudeDetails::MakeInstance ) );
 	PropertyModule.RegisterCustomPropertyTypeLayout( "GameplayCueTag", FOnGetPropertyTypeCustomizationInstance::CreateStatic( &FGameplayCueTagDetails::MakeInstance ) );
@@ -279,6 +281,7 @@ void FGameplayAbilitiesEditorModule::ShutdownModule()
 		PropertyModule.UnregisterCustomPropertyTypeLayout("GameplayEffectModifierMagnitude");
 		PropertyModule.UnregisterCustomPropertyTypeLayout("GameplayEffectExecutionDefinition");
 		PropertyModule.UnregisterCustomPropertyTypeLayout("GameplayEffectExecutionScopedModifierInfo");
+		PropertyModule.UnregisterCustomPropertyTypeLayout("GameplayTagBlueprintPropertyMapping");
 		PropertyModule.UnregisterCustomPropertyTypeLayout("ScalableFloat");
 		PropertyModule.UnregisterCustomPropertyTypeLayout("GameplayAttribute");
 	}
diff --git a/Engine/Plugins/Runtime/GameplayAbilities/Source/GameplayAbilitiesEditor/Private/GameplayTagBlueprintPropertyMappingDetails.cpp b/Engine/Plugins/Runtime/GameplayAbilities/Source/GameplayAbilitiesEditor/Private/GameplayTagBlueprintPropertyMappingDetails.cpp
new file mode 100644
--- /dev/null
+++ b/Engine/Plugins/Runtime/GameplayAbilities/Source/GameplayAbilitiesEditor/Private/GameplayTagBlueprintPropertyMappingDetails.cpp
@@ -0,0 +1,186 @@
+// Copyright 1998-2019 Epic Games, Inc. All Rights Reserved.
+
+#include "GameplayTagBlueprintPropertyMappingDetails.h"
+#include "GameplayEffectTypes.h"
+#include "PropertyHandle.h"
+#include "DetailWidgetRow.h"
+#include "IDetailChildrenBuilder.h"
+#include "IDetailPropertyRow.h"
+#include "Engine/Blueprint.h"
+#include "Widgets/Input/SComboBox.h"
+
+
+TSharedRef<IPropertyTypeCustomization> FGameplayTagBlueprintPropertyMappingDetails::MakeInstance()
+{
+	return MakeShareable(new FGameplayTagBlueprintPropertyMappingDetails());
+}
+
+void FGameplayTagBlueprintPropertyMappingDetails::CustomizeHeader(TSharedRef<IPropertyHandle> StructPropertyHandle, FDetailWidgetRow& HeaderRow, IPropertyTypeCustomizationUtils& StructCustomizationUtils)
+{
+	NamePropertyHandle = StructPropertyHandle->GetChildHandle(GET_MEMBER_NAME_CHECKED(FGameplayTagBlueprintPropertyMapping, PropertyName));
+	GuidPropertyHandle = StructPropertyHandle->GetChildHandle(GET_MEMBER_NAME_CHECKED(FGameplayTagBlueprintPropertyMapping, PropertyGuid));
+
+	FString SelectedPropertyName;
+	NamePropertyHandle->GetValue(SelectedPropertyName);
+
+	TArray<void*> RawData;
+	GuidPropertyHandle->AccessRawData(RawData);
+
+	FGuid SelectedPropertyGuid;
+	if (RawData.Num() > 0)
+	{
+		SelectedPropertyGuid = *static_cast<FGuid*>(RawData[0]);
+	}
+
+	UProperty* FoundProperty = nullptr;
+
+	TArray<UObject*> OuterObjects;
+	NamePropertyHandle->GetOuterObjects(OuterObjects);
+
+	for (UObject* ParentObject : OuterObjects)
+	{
+		if (!ParentObject)
+		{
+			continue;
+		}
+
+		for (TFieldIterator<UProperty> PropertyIt(ParentObject->GetClass()); PropertyIt; ++PropertyIt)
+		{
+			UProperty* Property = *PropertyIt;
+			if (!Property)
+			{
+				continue;
+			}
+
+			// Only support booleans, floats, and integers.
+			const bool bIsValidType = Property->IsA(UBoolProperty::StaticClass()) || Property->IsA(UIntProperty::StaticClass()) || Property->IsA(UFloatProperty::StaticClass());
+			if (!bIsValidType)
+			{
+				continue;
+			}
+
+			// Only accept properties from a blueprint.
+			if (!UBlueprint::GetBlueprintFromClass(Property->GetOwnerClass()))
+			{
+				continue;
+			}
+
+			// Ignore properties that don't have a GUID since we rely on it to handle property name changes.
+			FGuid PropertyToTestGuid = GetPropertyGuid(Property);
+			if (!PropertyToTestGuid.IsValid())
+			{
+				continue;
+			}
+
+			// Add the property to the combo box.
+			PropertyOptions.AddUnique(Property);
+
+			// Find our current selected property in the list.
+			if (SelectedPropertyGuid == PropertyToTestGuid)
+			{
+				FoundProperty = Property;
+			}
+		}
+	}
+
+	// Sort the options list alphabetically.
+	PropertyOptions.StableSort([](const UProperty& A, const UProperty& B) { return (A.GetName() < B.GetName()); });
+
+	if ((FoundProperty == nullptr) || (FoundProperty != SelectedProperty) || (FoundProperty->GetName() != SelectedPropertyName))
+	{
+		// The selected property needs to be updated.
+		OnChangeProperty(FoundProperty, ESelectInfo::Direct);
+	}
+
+	HeaderRow
+		.NameContent()
+		[
+			StructPropertyHandle->CreatePropertyNameWidget()
+		];
+}
+
+void FGameplayTagBlueprintPropertyMappingDetails::CustomizeChildren(TSharedRef<IPropertyHandle> StructPropertyHandle, IDetailChildrenBuilder& StructBuilder, IPropertyTypeCustomizationUtils& StructCustomizationUtils)
+{
+	TSharedPtr<IPropertyHandle> TagToMapHandle = StructPropertyHandle->GetChildHandle(GET_MEMBER_NAME_CHECKED(FGameplayTagBlueprintPropertyMapping, TagToMap));
+	TSharedPtr<IPropertyHandle> PropertyToEditHandle = StructPropertyHandle->GetChildHandle(GET_MEMBER_NAME_CHECKED(FGameplayTagBlueprintPropertyMapping, PropertyToEdit));
+
+	// Add the FGameplay Tag first.
+	StructBuilder.AddProperty(TagToMapHandle.ToSharedRef());
+
+	// Add the combo box next.
+	IDetailPropertyRow& PropertyRow = StructBuilder.AddProperty(StructPropertyHandle);
+	PropertyRow.CustomWidget()
+		.NameContent()
+		[
+			PropertyToEditHandle->CreatePropertyNameWidget()
+		]
+
+		.ValueContent()
+		.HAlign(HAlign_Left)
+		.MinDesiredWidth(250.0f)
+		[
+			SNew(SHorizontalBox)
+			+SHorizontalBox::Slot()
+			.HAlign(HAlign_Fill)
+			.Padding(0.0f, 0.0f, 2.0f, 0.0f)
+			[
+				SNew(SComboBox<UProperty*>)
+				.OptionsSource(&PropertyOptions)
+				.OnGenerateWidget(this, &FGameplayTagBlueprintPropertyMappingDetails::GeneratePropertyWidget)
+				.OnSelectionChanged(this, &FGameplayTagBlueprintPropertyMappingDetails::OnChangeProperty)
+				.ContentPadding(FMargin(2.0f, 2.0f))
+				.ToolTipText(this, &FGameplayTagBlueprintPropertyMappingDetails::GetSelectedValueText)
+				.InitiallySelectedItem(SelectedProperty)
+				[
+					SNew(STextBlock)
+					.Text(this, &FGameplayTagBlueprintPropertyMappingDetails::GetSelectedValueText)
+				]
+			]
+		];
+}
+
+void FGameplayTagBlueprintPropertyMappingDetails::OnChangeProperty(UProperty* ItemSelected, ESelectInfo::Type SelectInfo)
+{
+	if (NamePropertyHandle.IsValid() && GuidPropertyHandle.IsValid())
+	{
+		SelectedProperty = ItemSelected;
+
+		NamePropertyHandle->SetValue(GetPropertyName(ItemSelected));
+
+		TArray<void*> RawData;
+		GuidPropertyHandle->AccessRawData(RawData);
+
+		if (RawData.Num() > 0)
+		{
+			FGuid* RawGuid = static_cast<FGuid*>(RawData[0]);
+			*RawGuid = GetPropertyGuid(ItemSelected);
+		}
+	}
+}
+
+FGuid FGameplayTagBlueprintPropertyMappingDetails::GetPropertyGuid(UProperty* Property) const
+{
+	FGuid Guid;
+
+	if (Property != nullptr)
+	{
+		UBlueprint::GetGuidFromClassByFieldName<UProperty>(Property->GetOwnerClass(), Property->GetFName(), Guid);
+	}
+
+	return Guid;
+}
+
+FString FGameplayTagBlueprintPropertyMappingDetails::GetPropertyName(UProperty* Property) const
+{
+	return (Property ? Property->GetName() : TEXT("None"));
+}
+
+TSharedRef<SWidget> FGameplayTagBlueprintPropertyMappingDetails::GeneratePropertyWidget(UProperty* Property)
+{
+	return SNew(STextBlock).Text(FText::FromString(GetPropertyName(Property)));
+}
+
+FText FGameplayTagBlueprintPropertyMappingDetails::GetSelectedValueText() const
+{
+	return FText::FromString(GetPropertyName(SelectedProperty));
+}
diff --git a/Engine/Plugins/Runtime/GameplayAbilities/Source/GameplayAbilitiesEditor/Private/GameplayTagBlueprintPropertyMappingDetails.h b/Engine/Plugins/Runtime/GameplayAbilities/Source/GameplayAbilitiesEditor/Private/GameplayTagBlueprintPropertyMappingDetails.h
new file mode 100644
--- /dev/null
+++ b/Engine/Plugins/Runtime/GameplayAbilities/Source/GameplayAbilitiesEditor/Private/GameplayTagBlueprintPropertyMappingDetails.h
@@ -0,0 +1,42 @@
+// Copyright 1998-2019 Epic Games, Inc. All Rights Reserved.
+
+#pragma once
+
+#include "IPropertyTypeCustomization.h"
+
+
+class IPropertyHandle;
+class FDetailWidgetRow;
+class IDetailChildrenBuilder;
+class SWidget;
+
+
+class FGameplayTagBlueprintPropertyMappingDetails : public IPropertyTypeCustomization
+{
+public:
+
+	static TSharedRef<IPropertyTypeCustomization> MakeInstance();
+
+	virtual void CustomizeHeader(TSharedRef<IPropertyHandle> StructPropertyHandle, FDetailWidgetRow& HeaderRow, IPropertyTypeCustomizationUtils& StructCustomizationUtils) override;
+	virtual void CustomizeChildren(TSharedRef<IPropertyHandle> StructPropertyHandle, IDetailChildrenBuilder& StructBuilder, IPropertyTypeCustomizationUtils& StructCustomizationUtils) override;
+
+protected:
+
+	void OnChangeProperty(UProperty* ItemSelected, ESelectInfo::Type SelectInfo);
+
+	FGuid GetPropertyGuid(UProperty* Property) const;
+	FString GetPropertyName(UProperty* Property) const;
+
+	TSharedRef<SWidget> GeneratePropertyWidget(UProperty* Property);
+
+	FText GetSelectedValueText() const;
+
+protected:
+
+	TSharedPtr<IPropertyHandle> NamePropertyHandle;
+	TSharedPtr<IPropertyHandle> GuidPropertyHandle;
+
+	TArray<UProperty*> PropertyOptions;
+
+	UProperty* SelectedProperty;
+};
