Remove the UGameplayEffectTemplate class and related details customization code. This class has been deprecated for a long time and shouldn't be used by anyone. Removing it to avoid confusion by people that stumble across it while learning the ability system.

diff --git a/Engine/Plugins/Runtime/GameplayAbilities/Source/GameplayAbilities/Private/GameplayEffect.cpp b/Engine/Plugins/Runtime/GameplayAbilities/Source/GameplayAbilities/Private/GameplayEffect.cpp
--- a/Engine/Plugins/Runtime/GameplayAbilities/Source/GameplayAbilities/Private/GameplayEffect.cpp
+++ b/Engine/Plugins/Runtime/GameplayAbilities/Source/GameplayAbilities/Private/GameplayEffect.cpp
@@ -58,12 +58,6 @@ UGameplayEffect::UGameplayEffect(const FObjectInitializer& ObjectInitializer)
 	StackDurationRefreshPolicy = EGameplayEffectStackingDurationPolicy::RefreshOnSuccessfulApplication;
 	StackPeriodResetPolicy = EGameplayEffectStackingPeriodPolicy::ResetOnSuccessfulApplication;
 	bRequireModifierSuccessToTriggerCues = true;
-
-#if WITH_EDITORONLY_DATA
-	ShowAllProperties = true;
-	Template = nullptr;
-#endif
-
 }
 
 void UGameplayEffect::GetOwnedGameplayTags(FGameplayTagContainer& TagContainer) const
diff --git a/Engine/Plugins/Runtime/GameplayAbilities/Source/GameplayAbilities/Private/GameplayEffectTemplate.cpp b/Engine/Plugins/Runtime/GameplayAbilities/Source/GameplayAbilities/Private/GameplayEffectTemplate.cpp
deleted file mode 100644
diff --git a/Engine/Plugins/Runtime/GameplayAbilities/Source/GameplayAbilities/Public/GameplayEffect.h b/Engine/Plugins/Runtime/GameplayAbilities/Source/GameplayAbilities/Public/GameplayEffect.h
--- a/Engine/Plugins/Runtime/GameplayAbilities/Source/GameplayAbilities/Public/GameplayEffect.h
+++ b/Engine/Plugins/Runtime/GameplayAbilities/Source/GameplayAbilities/Public/GameplayEffect.h
@@ -1880,16 +1880,6 @@ public:
 
 	// ---------------------------------------------------------------------------------------------------------------------------------
 
-#if WITH_EDITORONLY_DATA
-	/** Template to derive starting values and editing customization from */
-	UPROPERTY()
-	UGameplayEffectTemplate*	Template;
-
-	/** When false, show a limited set of properties for editing, based on the template we are derived from */
-	UPROPERTY()
-	bool ShowAllProperties;
-#endif
-
 	/** Policy for the duration of this effect */
 	UPROPERTY(EditDefaultsOnly, Category=GameplayEffect)
 	EGameplayEffectDurationType DurationPolicy;
diff --git a/Engine/Plugins/Runtime/GameplayAbilities/Source/GameplayAbilities/Public/GameplayEffectTemplate.h b/Engine/Plugins/Runtime/GameplayAbilities/Source/GameplayAbilities/Public/GameplayEffectTemplate.h
deleted file mode 100644
diff --git a/Engine/Plugins/Runtime/GameplayAbilities/Source/GameplayAbilitiesEditor/Private/GameplayEffectDetails.cpp b/Engine/Plugins/Runtime/GameplayAbilities/Source/GameplayAbilitiesEditor/Private/GameplayEffectDetails.cpp
--- a/Engine/Plugins/Runtime/GameplayAbilities/Source/GameplayAbilitiesEditor/Private/GameplayEffectDetails.cpp
+++ b/Engine/Plugins/Runtime/GameplayAbilities/Source/GameplayAbilitiesEditor/Private/GameplayEffectDetails.cpp
@@ -5,7 +5,6 @@
 #include "DetailLayoutBuilder.h"
 #include "GameplayEffectTypes.h"
 #include "GameplayEffect.h"
-#include "GameplayEffectTemplate.h"
 
 #define LOCTEXT_NAMESPACE "GameplayEffectDetailsCustomization"
 
@@ -30,15 +29,6 @@ void FGameplayEffectDetails::CustomizeDetails(IDetailLayoutBuilder& DetailLayout
 		return;
 	}
 
-	TemplateProperty = DetailLayout.GetProperty(FName("Template"), UGameplayEffect::StaticClass());
-	ShowAllProperty = DetailLayout.GetProperty(FName("ShowAllProperties"), UGameplayEffect::StaticClass());
-
-	FSimpleDelegate UpdateShowAllDelegate = FSimpleDelegate::CreateSP(this, &FGameplayEffectDetails::OnShowAllChange);
-	ShowAllProperty->SetOnPropertyValueChanged(UpdateShowAllDelegate);
-
-	FSimpleDelegate UpdatTemplateDelegate = FSimpleDelegate::CreateSP(this, &FGameplayEffectDetails::OnTemplateChange);
-	TemplateProperty->SetOnPropertyValueChanged(UpdatTemplateDelegate);
-
 	TSharedPtr<IPropertyHandle> DurationPolicyProperty = DetailLayout.GetProperty(GET_MEMBER_NAME_CHECKED(UGameplayEffect, DurationPolicy), UGameplayEffect::StaticClass());
 	FSimpleDelegate UpdateDurationPolicyDelegate = FSimpleDelegate::CreateSP(this, &FGameplayEffectDetails::OnDurationPolicyChange);
 	DurationPolicyProperty->SetOnPropertyValueChanged(UpdateDurationPolicyDelegate);
@@ -47,23 +37,6 @@ void FGameplayEffectDetails::CustomizeDetails(IDetailLayoutBuilder& DetailLayout
 	UGameplayEffect* Obj = Cast<UGameplayEffect>(Objects[0].Get());
 	if (Obj)
 	{
-		if ( !Obj->ShowAllProperties )
-		{
-			UGameplayEffectTemplate* Template = Obj->Template;
-
-			if (Template)
-			{
-				UGameplayEffect* DefObj = Obj->Template->GetClass()->GetDefaultObject<UGameplayEffect>();
-
-				for (TFieldIterator<UProperty> PropIt(UGameplayEffect::StaticClass(), EFieldIteratorFlags::ExcludeSuper); PropIt; ++PropIt)
-				{
-					UProperty* Property = *PropIt;
-					TSharedPtr<IPropertyHandle> PropHandle = DetailLayout.GetProperty(Property->GetFName());
-					HideProperties(DetailLayout, PropHandle, Template);
-				}
-			}
-		}
-
 		if (Obj->DurationPolicy != EGameplayEffectDurationType::HasDuration)
 		{
 			TSharedPtr<IPropertyHandle> DurationMagnitudeProperty = DetailLayout.GetProperty(GET_MEMBER_NAME_CHECKED(UGameplayEffect, DurationMagnitude), UGameplayEffect::StaticClass());
@@ -94,88 +67,6 @@ void FGameplayEffectDetails::CustomizeDetails(IDetailLayoutBuilder& DetailLayout
 	}
 }
 
-/** Recursively hide properties that are not default editable according to the Template */
-bool FGameplayEffectDetails::HideProperties(IDetailLayoutBuilder& DetailLayout, TSharedPtr<IPropertyHandle> PropHandle, UGameplayEffectTemplate* Template)
-{
-	UProperty* UProp = PropHandle->GetProperty();
-
-	// Never hide the Template or ShowAllProperties properties
-	if (TemplateProperty->GetProperty() == UProp || ShowAllProperty->GetProperty() == UProp)
-	{
-		return false;
-	}
-
-	// Don't hide the EditableProperties
-	for (FString MatchStr : Template->EditableProperties)
-	{
-		if (MatchStr.Equals(UProp->GetName(), ESearchCase::IgnoreCase))
-		{			
-			return false;
-		}
-	}
-	
-	// Recurse into children - if they are all hidden then we are hidden.
-	uint32 NumChildren = 0;
-	PropHandle->GetNumChildren(NumChildren);
-
-	bool AllChildrenHidden = true;
-	for (uint32 ChildIdx = 0; ChildIdx < NumChildren; ++ChildIdx)
-	{
-		TSharedPtr<IPropertyHandle> ChildHandle = PropHandle->GetChildHandle(ChildIdx);
-		AllChildrenHidden &= HideProperties(DetailLayout, ChildHandle, Template);
-	}
-
-	if (AllChildrenHidden)
-	{
-		DetailLayout.HideProperty(PropHandle);
-		return true;
-	}
-
-	return false;
-}
-
-void FGameplayEffectDetails::OnShowAllChange()
-{
-	MyDetailLayout->ForceRefreshDetails();
-}
-
-void FGameplayEffectDetails::OnTemplateChange()
-{
-	TArray< TWeakObjectPtr<UObject> > Objects;
-	MyDetailLayout->GetObjectsBeingCustomized(Objects);
-	if (Objects.Num() != 1)
-	{
-		return;
-	}
-
-	UGameplayEffect* Obj = Cast<UGameplayEffect>(Objects[0].Get());
-	UGameplayEffect* Template = Obj->Template;
-
-	if (Template)
-	{
-		// Copy any non-default properties from the template into the current editable object
-		UGameplayEffect* DefObj = Template->GetClass()->GetDefaultObject<UGameplayEffect>();
-		for (TFieldIterator<UProperty> PropIt(UGameplayEffect::StaticClass(), EFieldIteratorFlags::ExcludeSuper); PropIt; ++PropIt)
-		{
-			UProperty* Property = *PropIt;
-			// don't overwrite the template property
-			if (Property->GetFName() == GET_MEMBER_NAME_CHECKED(UGameplayEffect, Template))
-			{
-				continue;
-			}
-			if (!Property->Identical_InContainer(Template, DefObj))
-			{
-				Property->CopyCompleteValue_InContainer(Obj, Template);
-			}
-
-			// Default to only showing template properties after changing template type
-			Obj->ShowAllProperties = false;
-		}
-	}
-
-	MyDetailLayout->ForceRefreshDetails();
-}
-
 void FGameplayEffectDetails::OnDurationPolicyChange()
 {
 	MyDetailLayout->ForceRefreshDetails();
diff --git a/Engine/Plugins/Runtime/GameplayAbilities/Source/GameplayAbilitiesEditor/Private/GameplayEffectDetails.h b/Engine/Plugins/Runtime/GameplayAbilities/Source/GameplayAbilitiesEditor/Private/GameplayEffectDetails.h
--- a/Engine/Plugins/Runtime/GameplayAbilities/Source/GameplayAbilitiesEditor/Private/GameplayEffectDetails.h
+++ b/Engine/Plugins/Runtime/GameplayAbilities/Source/GameplayAbilitiesEditor/Private/GameplayEffectDetails.h
@@ -7,12 +7,9 @@
 
 class IDetailLayoutBuilder;
 class IPropertyHandle;
-class UGameplayEffectTemplate;
 
 DECLARE_LOG_CATEGORY_EXTERN(LogGameplayEffectDetails, Log, All);
 
-class UGameplayEffectTemplate;
-
 class FGameplayEffectDetails : public IDetailCustomization
 {
 public:
@@ -31,13 +28,6 @@ private:
 
 	IDetailLayoutBuilder* MyDetailLayout;
 
-	TSharedPtr<IPropertyHandle> TemplateProperty;
-	TSharedPtr<IPropertyHandle> ShowAllProperty;
-
-	void OnShowAllChange();
-	void OnTemplateChange();
 	void OnDurationPolicyChange();
-
-	bool HideProperties(IDetailLayoutBuilder& DetailLayout, TSharedPtr<IPropertyHandle> PropHandle, UGameplayEffectTemplate* Template);
 };
 
