#UE4 Avoid calling Printf in FGameplayEffectSpec::SetLevel since it is called frequently

diff --git a/Engine/Plugins/Runtime/GameplayAbilities/Source/GameplayAbilities/Private/GameplayEffect.cpp b/Engine/Plugins/Runtime/GameplayAbilities/Source/GameplayAbilities/Private/GameplayEffect.cpp
--- a/Engine/Plugins/Runtime/GameplayAbilities/Source/GameplayAbilities/Private/GameplayEffect.cpp
+++ b/Engine/Plugins/Runtime/GameplayAbilities/Source/GameplayAbilities/Private/GameplayEffect.cpp
@@ -997,7 +997,7 @@ void FGameplayEffectSpec::SetLevel(float InLevel)
 			SetDuration(DefCalcDuration, false);
 		}
 
-		FString ContextString = FString::Printf(TEXT("FGameplayEffectSpec::SetLevel from effect %s"), *Def->GetName());
+		FString ContextString = Def->GetName();
 		Period = Def->Period.GetValueAtLevel(InLevel, &ContextString);
 		ChanceToApplyToTarget = Def->ChanceToApplyToTarget.GetValueAtLevel(InLevel, &ContextString);
 	}
