PR #5629: Typos in GameplayAbilitySystem (Contributed by [REDACTED_PERSON])

diff --git a/Engine/Plugins/Runtime/GameplayAbilities/Source/GameplayAbilities/Private/AttributeSet.cpp b/Engine/Plugins/Runtime/GameplayAbilities/Source/GameplayAbilities/Private/AttributeSet.cpp
--- a/Engine/Plugins/Runtime/GameplayAbilities/Source/GameplayAbilities/Private/AttributeSet.cpp
+++ b/Engine/Plugins/Runtime/GameplayAbilities/Source/GameplayAbilities/Private/AttributeSet.cpp
@@ -682,7 +682,7 @@ void FAttributeSetInitterDiscreteLevels::InitAttributeSetDefaults(UAbilitySystem
 	const FAttributeSetDefaultsCollection* Collection = Defaults.Find(GroupName);
 	if (!Collection)
 	{
-		ABILITY_LOG(Warning, TEXT("Unable to find DefaultAttributeSet Group %s. Failing back to Defaults"), *GroupName.ToString());
+		ABILITY_LOG(Warning, TEXT("Unable to find DefaultAttributeSet Group %s. Falling back to Defaults"), *GroupName.ToString());
 		Collection = Defaults.Find(FName(TEXT("Default")));
 		if (!Collection)
 		{
@@ -729,7 +729,7 @@ void FAttributeSetInitterDiscreteLevels::ApplyAttributeDefault(UAbilitySystemCom
 	const FAttributeSetDefaultsCollection* Collection = Defaults.Find(GroupName);
 	if (!Collection)
 	{
-		ABILITY_LOG(Warning, TEXT("Unable to find DefaultAttributeSet Group %s. Failing back to Defaults"), *GroupName.ToString());
+		ABILITY_LOG(Warning, TEXT("Unable to find DefaultAttributeSet Group %s. Falling back to Defaults"), *GroupName.ToString());
 		Collection = Defaults.Find(FName(TEXT("Default")));
 		if (!Collection)
 		{
diff --git a/Engine/Plugins/Runtime/GameplayAbilities/Source/GameplayAbilities/Private/GameplayEffectAggregatorLibrary.cpp b/Engine/Plugins/Runtime/GameplayAbilities/Source/GameplayAbilities/Private/GameplayEffectAggregatorLibrary.cpp
--- a/Engine/Plugins/Runtime/GameplayAbilities/Source/GameplayAbilities/Private/GameplayEffectAggregatorLibrary.cpp
+++ b/Engine/Plugins/Runtime/GameplayAbilities/Source/GameplayAbilities/Private/GameplayEffectAggregatorLibrary.cpp
@@ -6,7 +6,7 @@
 #include "AbilitySystemComponent.h"
 
 /** Custom functions. The idea here is that we may want to mix and match these (if FAggregatorEvaluateMetaData starts to hold more than just the qualifier functions) */
-void QualifierFunc_MostNegativeMod_AllPostiiveMods(const FAggregatorEvaluateParameters& EvalParameters, const FAggregator* Aggregator)
+void QualifierFunc_MostNegativeMod_AllPositiveMods(const FAggregatorEvaluateParameters& EvalParameters, const FAggregator* Aggregator)
 {
 	// We want to inhibit all qualified negative effects except for the most severe. We want to leave positive modifiers alone
 	const FAggregatorMod* MostNegativeMod = nullptr;
@@ -60,4 +60,4 @@ void QualifierFunc_MostNegativeMod_AllPostiiveMods(const FAggregatorEvaluatePara
 }
 
 /** static FAggregatorEvaluateMetaDatas that use the above functions */
-FAggregatorEvaluateMetaData FAggregatorEvaluateMetaDataLibrary::MostNegativeMod_AllPostiiveMods(QualifierFunc_MostNegativeMod_AllPostiiveMods);
\ No newline at end of file
+FAggregatorEvaluateMetaData FAggregatorEvaluateMetaDataLibrary::MostNegativeMod_AllPositiveMods(QualifierFunc_MostNegativeMod_AllPositiveMods);
\ No newline at end of file
diff --git a/Engine/Plugins/Runtime/GameplayAbilities/Source/GameplayAbilities/Public/Abilities/Tasks/AbilityTask_MoveToLocation.h b/Engine/Plugins/Runtime/GameplayAbilities/Source/GameplayAbilities/Public/Abilities/Tasks/AbilityTask_MoveToLocation.h
--- a/Engine/Plugins/Runtime/GameplayAbilities/Source/GameplayAbilities/Public/Abilities/Tasks/AbilityTask_MoveToLocation.h
+++ b/Engine/Plugins/Runtime/GameplayAbilities/Source/GameplayAbilities/Public/Abilities/Tasks/AbilityTask_MoveToLocation.h
@@ -18,7 +18,7 @@ class UGameplayTasksComponent;
 
 /**
  *	TODO:
- *		-Implement replicated time so that this can work as a simulated task for Join In Prgorss clients.
+ *		-Implement replicated time so that this can work as a simulated task for Join In Progress clients.
  */
 
 
diff --git a/Engine/Plugins/Runtime/GameplayAbilities/Source/GameplayAbilities/Public/GameplayEffect.h b/Engine/Plugins/Runtime/GameplayAbilities/Source/GameplayAbilities/Public/GameplayEffect.h
--- a/Engine/Plugins/Runtime/GameplayAbilities/Source/GameplayAbilities/Public/GameplayEffect.h
+++ b/Engine/Plugins/Runtime/GameplayAbilities/Source/GameplayAbilities/Public/GameplayEffect.h
@@ -624,7 +624,7 @@ enum class EGameplayEffectStackingExpirationPolicy : uint8
 	/** The current stack count will be decremented by 1 and the duration refreshed. The GE is not "reapplied", just continues to exist with one less stacks. */
 	RemoveSingleStackAndRefreshDuration,
 
-	/** The duration of the gameplay effect is refreshed. This essentially makes the effect infinite in duration. This can be used to manually handle stack decrements via XXX callback */
+	/** The duration of the gameplay effect is refreshed. This essentially makes the effect infinite in duration. This can be used to manually handle stack decrements via OnStackCountChange callback */
 	RefreshDuration,
 };
 
diff --git a/Engine/Plugins/Runtime/GameplayAbilities/Source/GameplayAbilities/Public/GameplayEffectAggregatorLibrary.h b/Engine/Plugins/Runtime/GameplayAbilities/Source/GameplayAbilities/Public/GameplayEffectAggregatorLibrary.h
--- a/Engine/Plugins/Runtime/GameplayAbilities/Source/GameplayAbilities/Public/GameplayEffectAggregatorLibrary.h
+++ b/Engine/Plugins/Runtime/GameplayAbilities/Source/GameplayAbilities/Public/GameplayEffectAggregatorLibrary.h
@@ -6,5 +6,5 @@
 
 struct GAMEPLAYABILITIES_API FAggregatorEvaluateMetaDataLibrary
 {
-	static FAggregatorEvaluateMetaData MostNegativeMod_AllPostiiveMods;
+	static FAggregatorEvaluateMetaData MostNegativeMod_AllPositiveMods;
 };
\ No newline at end of file
