Add ensure message to UAbilitySystemComponent::OnRemoveAbility to try to catch GameplayEventTriggeredAbilities being out of sync with ActivatableAbilities

diff --git a/Engine/Plugins/Runtime/GameplayAbilities/Source/GameplayAbilities/Private/AbilitySystemComponent_Abilities.cpp b/Engine/Plugins/Runtime/GameplayAbilities/Source/GameplayAbilities/Private/AbilitySystemComponent_Abilities.cpp
--- a/Engine/Plugins/Runtime/GameplayAbilities/Source/GameplayAbilities/Private/AbilitySystemComponent_Abilities.cpp
+++ b/Engine/Plugins/Runtime/GameplayAbilities/Source/GameplayAbilities/Private/AbilitySystemComponent_Abilities.cpp
@@ -443,7 +443,8 @@ void UAbilitySystemComponent::OnRemoveAbility(FGameplayAbilitySpec& Spec)
 
 		auto& TriggeredAbilityMap = (TriggerData.TriggerSource == EGameplayAbilityTriggerSource::GameplayEvent) ? GameplayEventTriggeredAbilities : OwnedTagTriggeredAbilities;
 
-		if (TriggeredAbilityMap.Contains(EventTag))
+		if (ensureMsgf(TriggeredAbilityMap.Contains(EventTag), 
+			TEXT("%s::%s not found in TriggeredAbilityMap while removing, TriggerSource: %d"), *Spec.Ability->GetName(), *EventTag.ToString(), (int32)TriggerData.TriggerSource))
 		{
 			TriggeredAbilityMap[EventTag].Remove(Spec.Handle);
 			if (TriggeredAbilityMap[EventTag].Num() == 0)
