Add some server only scope cycle counters in the ability system for better tracking in profile captures.

diff --git a/Engine/Plugins/Runtime/GameplayAbilities/Source/GameplayAbilities/Private/AbilitySystemComponent.cpp b/Engine/Plugins/Runtime/GameplayAbilities/Source/GameplayAbilities/Private/AbilitySystemComponent.cpp
--- a/Engine/Plugins/Runtime/GameplayAbilities/Source/GameplayAbilities/Private/AbilitySystemComponent.cpp
+++ b/Engine/Plugins/Runtime/GameplayAbilities/Source/GameplayAbilities/Private/AbilitySystemComponent.cpp
@@ -19,6 +19,12 @@
 DEFINE_LOG_CATEGORY(LogAbilitySystemComponent);
 
 DECLARE_CYCLE_STAT(TEXT("AbilitySystemComp ApplyGameplayEffectSpecToTarget"), STAT_AbilitySystemComp_ApplyGameplayEffectSpecToTarget, STATGROUP_AbilitySystem);
+DECLARE_CYCLE_STAT(TEXT("AbilitySystemComp ApplyGameplayEffectSpecToSelf"), STAT_AbilitySystemComp_ApplyGameplayEffectSpecToSelf, STATGROUP_AbilitySystem);
+DECLARE_CYCLE_STAT(TEXT("AbilitySystemComp OnImmunityBlockGameplayEffect"), STAT_AbilitySystemComp_OnImmunityBlockGameplayEffect, STATGROUP_AbilitySystem);
+DECLARE_CYCLE_STAT(TEXT("AbilitySystemComp InvokeGameplayCueEvent"), STAT_AbilitySystemComp_InvokeGameplayCueEvent, STATGROUP_AbilitySystem);
+DECLARE_CYCLE_STAT(TEXT("AbilitySystemComp OnGameplayEffectAppliedToSelf"), STAT_AbilitySystemComp_OnGameplayEffectAppliedToSelf, STATGROUP_AbilitySystem);
+DECLARE_CYCLE_STAT(TEXT("AbilitySystemComp OnGameplayEffectAppliedToTarget"), STAT_AbilitySystemComp_OnGameplayEffectAppliedToTarget, STATGROUP_AbilitySystem);
+DECLARE_CYCLE_STAT(TEXT("AbilitySystemComp ExecuteGameplayEffect"), STAT_AbilitySystemComp_ExecuteGameplayEffect, STATGROUP_AbilitySystem);
 
 #define LOCTEXT_NAMESPACE "AbilitySystemComponent"
 
@@ -568,6 +574,10 @@ FActiveGameplayEffectHandle UAbilitySystemComponent::ApplyGameplayEffectSpecToTa
 
 FActiveGameplayEffectHandle UAbilitySystemComponent::ApplyGameplayEffectSpecToSelf(const FGameplayEffectSpec &Spec, FPredictionKey PredictionKey)
 {
+#if WITH_SERVER_CODE
+	SCOPE_CYCLE_COUNTER(STAT_AbilitySystemComp_ApplyGameplayEffectSpecToSelf);
+#endif
+
 	// Scope lock the container after the addition has taken place to prevent the new effect from potentially getting mangled during the remainder
 	// of the add operation
 	FScopedActiveGameplayEffectLock ScopeLock(ActiveGameplayEffects);
@@ -842,6 +852,10 @@ void UAbilitySystemComponent::ExecutePeriodicEffect(FActiveGameplayEffectHandle
 
 void UAbilitySystemComponent::ExecuteGameplayEffect(FGameplayEffectSpec &Spec, FPredictionKey PredictionKey)
 {
+#if WITH_SERVER_CODE
+	SCOPE_CYCLE_COUNTER(STAT_AbilitySystemComp_ExecuteGameplayEffect);
+#endif
+
 	// Should only ever execute effects that are instant application or periodic application
 	// Effects with no period and that aren't instant application should never be executed
 	check( (Spec.GetDuration() == UGameplayEffect::INSTANT_APPLICATION || Spec.GetPeriod() != UGameplayEffect::NO_PERIOD) );
@@ -1000,6 +1014,9 @@ FActiveGameplayEffectHandle UAbilitySystemComponent::FindActiveGameplayEffectHan
 
 void UAbilitySystemComponent::OnImmunityBlockGameplayEffect(const FGameplayEffectSpec& Spec, const FActiveGameplayEffect* ImmunityGE)
 {
+#if WITH_SERVER_CODE
+	SCOPE_CYCLE_COUNTER(STAT_AbilitySystemComp_OnImmunityBlockGameplayEffect);
+#endif
 	OnImmunityBlockGameplayEffectDelegate.Broadcast(Spec, ImmunityGE);
 }
 
@@ -1022,6 +1039,10 @@ bool UAbilitySystemComponent::IsReadyForGameplayCues()
 
 void UAbilitySystemComponent::InvokeGameplayCueEvent(const FGameplayEffectSpecForRPC &Spec, EGameplayCueEvent::Type EventType)
 {
+#if WITH_SERVER_CODE
+	SCOPE_CYCLE_COUNTER(STAT_AbilitySystemComp_InvokeGameplayCueEvent);
+#endif
+
 	AActor* ActorAvatar = AbilityActorInfo->AvatarActor.Get();
 	if (ActorAvatar == nullptr || bSuppressGameplayCues)
 	{
@@ -1578,12 +1599,18 @@ void UAbilitySystemComponent::OnGameplayEffectDurationChange(struct FActiveGamep
 
 void UAbilitySystemComponent::OnGameplayEffectAppliedToTarget(UAbilitySystemComponent* Target, const FGameplayEffectSpec& SpecApplied, FActiveGameplayEffectHandle ActiveHandle)
 {
+#if WITH_SERVER_CODE
+	SCOPE_CYCLE_COUNTER(STAT_AbilitySystemComp_OnGameplayEffectAppliedToTarget);
+#endif
 	OnGameplayEffectAppliedDelegateToTarget.Broadcast(Target, SpecApplied, ActiveHandle);
 	ActiveGameplayEffects.ApplyStackingLogicPostApplyAsSource(Target, SpecApplied, ActiveHandle);
 }
 
 void UAbilitySystemComponent::OnGameplayEffectAppliedToSelf(UAbilitySystemComponent* Source, const FGameplayEffectSpec& SpecApplied, FActiveGameplayEffectHandle ActiveHandle)
 {
+#if WITH_SERVER_CODE
+	SCOPE_CYCLE_COUNTER(STAT_AbilitySystemComp_OnGameplayEffectAppliedToSelf);
+#endif
 	OnGameplayEffectAppliedDelegateToSelf.Broadcast(Source, SpecApplied, ActiveHandle);
 }
 
diff --git a/Engine/Plugins/Runtime/GameplayAbilities/Source/GameplayAbilities/Private/AbilitySystemStats.cpp b/Engine/Plugins/Runtime/GameplayAbilities/Source/GameplayAbilities/Private/AbilitySystemStats.cpp
--- a/Engine/Plugins/Runtime/GameplayAbilities/Source/GameplayAbilities/Private/AbilitySystemStats.cpp
+++ b/Engine/Plugins/Runtime/GameplayAbilities/Source/GameplayAbilities/Private/AbilitySystemStats.cpp
@@ -29,7 +29,10 @@ DEFINE_STAT(STAT_HasMatchingGameplayTag);
 DEFINE_STAT(STAT_HandleGameplayCueNotifyStatic);
 DEFINE_STAT(STAT_HandleGameplayCueNotifyActor);
 DEFINE_STAT(STAT_ApplyGameplayEffectToTarget);
+DEFINE_STAT(STAT_ExecuteActiveEffectsFrom);
 DEFINE_STAT(STAT_OnActiveGameplayEffectAdded);
 DEFINE_STAT(STAT_OnActiveGameplayEffectRemoved);
 DEFINE_STAT(STAT_GameplayCueInterface_HandleGameplayCue);
 DEFINE_STAT(STAT_AbilitySystem_TaskCount);
+DEFINE_STAT(STAT_InternalExecuteMod);
+DEFINE_STAT(STAT_PostGameplayEffectExecute);
diff --git a/Engine/Plugins/Runtime/GameplayAbilities/Source/GameplayAbilities/Private/GameplayCueManager.cpp b/Engine/Plugins/Runtime/GameplayAbilities/Source/GameplayAbilities/Private/GameplayCueManager.cpp
--- a/Engine/Plugins/Runtime/GameplayAbilities/Source/GameplayAbilities/Private/GameplayCueManager.cpp
+++ b/Engine/Plugins/Runtime/GameplayAbilities/Source/GameplayAbilities/Private/GameplayCueManager.cpp
@@ -129,6 +129,10 @@ void UGameplayCueManager::HandleGameplayCues(AActor* TargetActor, const FGamepla
 
 void UGameplayCueManager::HandleGameplayCue(AActor* TargetActor, FGameplayTag GameplayCueTag, EGameplayCueEvent::Type EventType, const FGameplayCueParameters& Parameters, EGameplayCueExecutionOptions Options)
 {
+#if WITH_SERVER_CODE
+	QUICK_SCOPE_CYCLE_COUNTER(STAT_GameplayCueManager_HandleGameplayCue);
+#endif
+
 #if WITH_EDITOR
 	if (GIsEditor && TargetActor == nullptr && UGameplayCueManager::PreviewComponent)
 	{
@@ -171,6 +175,10 @@ bool UGameplayCueManager::ShouldSuppressGameplayCues(AActor* TargetActor)
 
 void UGameplayCueManager::RouteGameplayCue(AActor* TargetActor, FGameplayTag GameplayCueTag, EGameplayCueEvent::Type EventType, const FGameplayCueParameters& Parameters, EGameplayCueExecutionOptions Options)
 {
+#if WITH_SERVER_CODE
+	QUICK_SCOPE_CYCLE_COUNTER(STAT_GameplayCueManager_RouteGameplayCue);
+#endif
+
 	// If we want to ignore interfaces, set the pointer to null
 	IGameplayCueInterface* GameplayCueInterface = !(Options & EGameplayCueExecutionOptions::IgnoreInterfaces) ? Cast<IGameplayCueInterface>(TargetActor) : nullptr;
 	bool bAcceptsCue = true;
@@ -224,6 +232,10 @@ void UGameplayCueManager::RouteGameplayCue(AActor* TargetActor, FGameplayTag Gam
 
 void UGameplayCueManager::TranslateGameplayCue(FGameplayTag& Tag, AActor* TargetActor, const FGameplayCueParameters& Parameters)
 {
+#if WITH_SERVER_CODE
+	QUICK_SCOPE_CYCLE_COUNTER(STAT_GameplayCueManager_TranslateGameplayCue);
+#endif
+
 	TranslationManager.TranslateTag(Tag, TargetActor, Parameters);
 }
 
diff --git a/Engine/Plugins/Runtime/GameplayAbilities/Source/GameplayAbilities/Private/GameplayCueSet.cpp b/Engine/Plugins/Runtime/GameplayAbilities/Source/GameplayAbilities/Private/GameplayCueSet.cpp
--- a/Engine/Plugins/Runtime/GameplayAbilities/Source/GameplayAbilities/Private/GameplayCueSet.cpp
+++ b/Engine/Plugins/Runtime/GameplayAbilities/Source/GameplayAbilities/Private/GameplayCueSet.cpp
@@ -23,6 +23,10 @@ UGameplayCueSet::UGameplayCueSet(const FObjectInitializer& ObjectInitializer)
 
 bool UGameplayCueSet::HandleGameplayCue(AActor* TargetActor, FGameplayTag GameplayCueTag, EGameplayCueEvent::Type EventType, const FGameplayCueParameters& Parameters)
 {
+#if WITH_SERVER_CODE
+	QUICK_SCOPE_CYCLE_COUNTER(STAT_GameplayCueSet_HandleGameplayCue);
+#endif
+
 	// GameplayCueTags could have been removed from the dictionary but not content. When the content is resaved the old tag will be cleaned up, but it could still come through here
 	// at runtime. Since we only populate the map with dictionary gameplaycue tags, we may not find it here.
 	int32* Ptr = GameplayCueDataMap.Find(GameplayCueTag);
diff --git a/Engine/Plugins/Runtime/GameplayAbilities/Source/GameplayAbilities/Private/GameplayEffect.cpp b/Engine/Plugins/Runtime/GameplayAbilities/Source/GameplayAbilities/Private/GameplayEffect.cpp
--- a/Engine/Plugins/Runtime/GameplayAbilities/Source/GameplayAbilities/Private/GameplayEffect.cpp
+++ b/Engine/Plugins/Runtime/GameplayAbilities/Source/GameplayAbilities/Private/GameplayEffect.cpp
@@ -1905,6 +1905,10 @@ void FActiveGameplayEffectsContainer::RegisterWithOwner(UAbilitySystemComponent*
 /** This is the main function that executes a GameplayEffect on Attributes and ActiveGameplayEffects */
 void FActiveGameplayEffectsContainer::ExecuteActiveEffectsFrom(FGameplayEffectSpec &Spec, FPredictionKey PredictionKey)
 {
+#if WITH_SERVER_CODE
+	SCOPE_CYCLE_COUNTER(STAT_ExecuteActiveEffectsFrom);
+#endif
+
 	FGameplayEffectSpec& SpecToUse = Spec;
 
 	// Capture our own tags.
@@ -2627,6 +2631,8 @@ float FActiveGameplayEffectsContainer::GetEffectContribution(const FAggregatorEv
 
 bool FActiveGameplayEffectsContainer::InternalExecuteMod(FGameplayEffectSpec& Spec, FGameplayModifierEvaluatedData& ModEvalData)
 {
+	SCOPE_CYCLE_COUNTER(STAT_InternalExecuteMod);
+
 	check(Owner);
 
 	bool bExecuted = false;
@@ -2659,8 +2665,11 @@ bool FActiveGameplayEffectsContainer::InternalExecuteMod(FGameplayEffectSpec& Sp
 			}
 			ModifiedAttribute->TotalMagnitude += ModEvalData.Magnitude;
 
-			/** This should apply 'gamewide' rules. Such as clamping Health to MaxHealth or granting +3 health for every point of strength, etc */
-			AttributeSet->PostGameplayEffectExecute(ExecuteData);
+			{
+				SCOPE_CYCLE_COUNTER(STAT_PostGameplayEffectExecute);
+				/** This should apply 'gamewide' rules. Such as clamping Health to MaxHealth or granting +3 health for every point of strength, etc */
+				AttributeSet->PostGameplayEffectExecute(ExecuteData);
+			}
 
 #if ENABLE_VISUAL_LOG
 			DebugExecutedGameplayEffectData DebugData;
diff --git a/Engine/Plugins/Runtime/GameplayAbilities/Source/GameplayAbilities/Public/AbilitySystemStats.h b/Engine/Plugins/Runtime/GameplayAbilities/Source/GameplayAbilities/Public/AbilitySystemStats.h
--- a/Engine/Plugins/Runtime/GameplayAbilities/Source/GameplayAbilities/Public/AbilitySystemStats.h
+++ b/Engine/Plugins/Runtime/GameplayAbilities/Source/GameplayAbilities/Public/AbilitySystemStats.h
@@ -34,9 +34,11 @@ DECLARE_CYCLE_STAT_EXTERN(TEXT("Has Matching GameplayTag"), STAT_HasMatchingGame
 DECLARE_CYCLE_STAT_EXTERN(TEXT("GameplayCueNotify Static"), STAT_HandleGameplayCueNotifyStatic, STATGROUP_AbilitySystem, );
 DECLARE_CYCLE_STAT_EXTERN(TEXT("GameplayCueNotify Actor"), STAT_HandleGameplayCueNotifyActor, STATGROUP_AbilitySystem, );
 DECLARE_CYCLE_STAT_EXTERN(TEXT("ApplyGameplayEffectToTarget"), STAT_ApplyGameplayEffectToTarget, STATGROUP_AbilitySystem, );
+DECLARE_CYCLE_STAT_EXTERN(TEXT("ExecuteActiveEffectsFrom"), STAT_ExecuteActiveEffectsFrom, STATGROUP_AbilitySystem, );
+DECLARE_CYCLE_STAT_EXTERN(TEXT("InternalExecuteMod"), STAT_InternalExecuteMod, STATGROUP_AbilitySystem, );
+DECLARE_CYCLE_STAT_EXTERN(TEXT("PostGameplayEffectExecute"), STAT_PostGameplayEffectExecute, STATGROUP_AbilitySystem, );
 DECLARE_CYCLE_STAT_EXTERN(TEXT("ActiveGameplayEffect Added"), STAT_OnActiveGameplayEffectAdded, STATGROUP_AbilitySystem, );
 DECLARE_CYCLE_STAT_EXTERN(TEXT("ActiveGameplayEffect Removed"), STAT_OnActiveGameplayEffectRemoved, STATGROUP_AbilitySystem, );
 DECLARE_CYCLE_STAT_EXTERN(TEXT("GameplayCueInterface HandleGameplayCue"), STAT_GameplayCueInterface_HandleGameplayCue, STATGROUP_AbilitySystem, );
 
-
 DECLARE_DWORD_ACCUMULATOR_STAT_EXTERN(TEXT("AbilityTask Count"), STAT_AbilitySystem_TaskCount, STATGROUP_AbilitySystem, );
\ No newline at end of file
