Fix for FActiveGameplayEffectIterator no longer being comparable.

diff --git a/Engine/Plugins/Runtime/GameplayAbilities/Source/GameplayAbilities/Public/ActiveGameplayEffectIterator.h b/Engine/Plugins/Runtime/GameplayAbilities/Source/GameplayAbilities/Public/ActiveGameplayEffectIterator.h
--- a/Engine/Plugins/Runtime/GameplayAbilities/Source/GameplayAbilities/Public/ActiveGameplayEffectIterator.h
+++ b/Engine/Plugins/Runtime/GameplayAbilities/Source/GameplayAbilities/Public/ActiveGameplayEffectIterator.h
@@ -65,6 +65,16 @@ public:
 		return (Current != nullptr);
 	}
 
+	friend bool operator==(const FActiveGameplayEffectIterator& Lhs, const FActiveGameplayEffectIterator& Rhs)
+	{
+		return Lhs.Current == Rhs.Current;
+	}
+
+	friend bool operator!=(const FActiveGameplayEffectIterator& Lhs, const FActiveGameplayEffectIterator& Rhs)
+	{
+		return Lhs.Current != Rhs.Current;
+	}
+
 private:
 
 	FORCEINLINE ElementType* AdvancePending(ElementType** Next)
