Remove HideFromLevelInfos from attribute sets, feature was deleted years ago

diff --git a/Engine/Plugins/Runtime/GameplayAbilities/Source/GameplayAbilities/Public/AbilitySystemTestAttributeSet.h b/Engine/Plugins/Runtime/GameplayAbilities/Source/GameplayAbilities/Public/AbilitySystemTestAttributeSet.h
--- a/Engine/Plugins/Runtime/GameplayAbilities/Source/GameplayAbilities/Public/AbilitySystemTestAttributeSet.h
+++ b/Engine/Plugins/Runtime/GameplayAbilities/Source/GameplayAbilities/Public/AbilitySystemTestAttributeSet.h
@@ -34,7 +34,7 @@ class GAMEPLAYABILITIES_API UAbilitySystemTestAttributeSet : public UAttributeSe
 	mutable float	MaxMana;
 
 	/** This Damage is just used for applying negative health mods. Its not a 'persistent' attribute. */
-	UPROPERTY(EditAnywhere, BlueprintReadWrite, Category = "AttributeTest", meta = (HideFromLevelInfos))		// You can't make a GameplayEffect 'powered' by Damage (Its transient)
+	UPROPERTY(EditAnywhere, BlueprintReadWrite, Category = "AttributeTest")		// You can't make a GameplayEffect 'powered' by Damage (Its transient)
 	mutable float	Damage;
 
 	/** This Attribute is the actual spell damage for this actor. It will power spell based GameplayEffects */
