Switching some comparisons of the gameplay effect period to use > instead of !=. This makes the code more consistent and prevents some potential bugs.

PR #5859: Fix comparing float numbers (Contributed by [REDACTED_PERSON])

diff --git a/Engine/Plugins/Runtime/GameplayAbilities/Source/GameplayAbilities/Private/GameplayEffect.cpp b/Engine/Plugins/Runtime/GameplayAbilities/Source/GameplayAbilities/Private/GameplayEffect.cpp
--- a/Engine/Plugins/Runtime/GameplayAbilities/Source/GameplayAbilities/Private/GameplayEffect.cpp
+++ b/Engine/Plugins/Runtime/GameplayAbilities/Source/GameplayAbilities/Private/GameplayEffect.cpp
@@ -2875,7 +2875,7 @@ FActiveGameplayEffect* FActiveGameplayEffectsContainer::ApplyGameplayEffectSpec(
 		const bool HasModifiedAttributes = AppliedEffectSpec.ModifiedAttributes.Num() > 0;
 		const bool HasDurationAndNoPeriod = AppliedEffectSpec.Def->DurationPolicy == EGameplayEffectDurationType::HasDuration && AppliedEffectSpec.GetPeriod() == UGameplayEffect::NO_PERIOD;
 		const bool HasPeriodAndNoDuration = AppliedEffectSpec.Def->DurationPolicy == EGameplayEffectDurationType::Instant &&
-											AppliedEffectSpec.GetPeriod() != UGameplayEffect::NO_PERIOD;
+											AppliedEffectSpec.GetPeriod() > UGameplayEffect::NO_PERIOD;
 		const bool ShouldBuildModifiedAttributeList = !HasModifiedAttributes && (HasDurationAndNoPeriod || HasPeriodAndNoDuration);
 		if (ShouldBuildModifiedAttributeList)
 		{
@@ -2954,7 +2954,7 @@ FActiveGameplayEffect* FActiveGameplayEffectsContainer::ApplyGameplayEffectSpec(
 	}
 	
 	// Register period callbacks with the timer manager
-	if (bSetPeriod && Owner && (AppliedEffectSpec.GetPeriod() != UGameplayEffect::NO_PERIOD))
+	if (bSetPeriod && Owner && (AppliedEffectSpec.GetPeriod() > UGameplayEffect::NO_PERIOD))
 	{
 		FTimerManager& TimerManager = Owner->GetWorld()->GetTimerManager();
 		FTimerDelegate Delegate = FTimerDelegate::CreateUObject(Owner, &UAbilitySystemComponent::ExecutePeriodicEffect, AppliedActiveGE->Handle);
