Remove unused function that would have done the wrong thing if it had been used.

PR #5247: Removed `AbilitySystemDebugNextCategory` (wrong and unused) (Contributed by [REDACTED_PERSON])

diff --git a/Engine/Plugins/Runtime/GameplayAbilities/Source/GameplayAbilities/Private/AbilitySystemComponent.cpp b/Engine/Plugins/Runtime/GameplayAbilities/Source/GameplayAbilities/Private/AbilitySystemComponent.cpp
--- a/Engine/Plugins/Runtime/GameplayAbilities/Source/GameplayAbilities/Private/AbilitySystemComponent.cpp
+++ b/Engine/Plugins/Runtime/GameplayAbilities/Source/GameplayAbilities/Private/AbilitySystemComponent.cpp
@@ -1987,11 +1987,6 @@ FAutoConsoleCommandWithWorld AbilitySystemPrevDebugTargetCmd(
 	FConsoleCommandWithWorldDelegate::CreateStatic(AbilitySystemCycleDebugTarget, false)
 	);
 
-static void	AbilitySystemDebugNextCategory(UWorld* InWorld, bool Next)
-{
-	CycleDebugTarget( GetDebugTargetInfo(InWorld), Next );
-}
-
 FAutoConsoleCommandWithWorld AbilitySystemDebugNextCategoryCmd(
 	TEXT("AbilitySystem.Debug.NextCategory"),
 	TEXT("Targets previous AbilitySystemComponent in ShowDebug AbilitySystem"),
