Added ability name to ensure to improve debugging
#rnx

diff --git a/Engine/Plugins/Runtime/GameplayAbilities/Source/GameplayAbilities/Private/AbilitySystemComponent_Abilities.cpp b/Engine/Plugins/Runtime/GameplayAbilities/Source/GameplayAbilities/Private/AbilitySystemComponent_Abilities.cpp
--- a/Engine/Plugins/Runtime/GameplayAbilities/Source/GameplayAbilities/Private/AbilitySystemComponent_Abilities.cpp
+++ b/Engine/Plugins/Runtime/GameplayAbilities/Source/GameplayAbilities/Private/AbilitySystemComponent_Abilities.cpp
@@ -714,14 +714,10 @@ void UAbilitySystemComponent::NotifyAbilityEnded(FGameplayAbilitySpecHandle Hand
 	}
 
 	// check to make sure we do not cause a roll over to uint8 by decrementing when it is 0
-	if (ensure(Spec->ActiveCount > 0))
+	if (ensureMsgf(Spec->ActiveCount > 0, TEXT("NotifyAbilityEnded called when the Spec->ActiveCount <= 0 for ability %s"), *Ability->GetName()))
 	{
 		Spec->ActiveCount--;
 	}
-	else
-	{
-		ABILITY_LOG(Warning, TEXT("NotifyAbilityEnded called when the Spec->ActiveCount <= 0"));
-	}
 
 	// Broadcast that the ability ended
 	AbilityEndedCallbacks.Broadcast(Ability);
