Added minor debugging info to the AbilitySystemComponent Debug Printing to show the tag count for Owned Tags.

#rnx

diff --git a/Engine/Plugins/Runtime/GameplayAbilities/Source/GameplayAbilities/Private/AbilitySystemComponent.cpp b/Engine/Plugins/Runtime/GameplayAbilities/Source/GameplayAbilities/Private/AbilitySystemComponent.cpp
--- a/Engine/Plugins/Runtime/GameplayAbilities/Source/GameplayAbilities/Private/AbilitySystemComponent.cpp
+++ b/Engine/Plugins/Runtime/GameplayAbilities/Source/GameplayAbilities/Private/AbilitySystemComponent.cpp
@@ -2151,7 +2151,19 @@ void UAbilitySystemComponent::Debug_Internal(FAbilitySystemComponentDebugInfo& I
 		Info.Canvas->SetDrawColor(FColor::White);
 	}
 
-	DebugLine(Info, FString::Printf(TEXT("Owned Tags: %s"), *OwnerTags.ToStringSimple()), 4.f, 0.f, 4);
+	FString TagsStrings;
+	int32 TagCount = 1;
+	const int32 NumTags = OwnerTags.Num();
+	for (FGameplayTag Tag : OwnerTags)
+	{
+		TagsStrings.Append(FString::Printf(TEXT("%s (%d)"), *Tag.ToString(), GetTagCount(Tag)));
+
+		if (TagCount++ < NumTags)
+		{
+			TagsStrings += TEXT(", ");
+		}
+	}
+	DebugLine(Info, FString::Printf(TEXT("Owned Tags: %s"), *TagsStrings), 4.f, 0.f, 4);
 
 	if (BlockedAbilityTags.GetExplicitGameplayTags().Num() > 0)
 	{
