Fix for minimal replication tag proxy that games can control.

* Adds option to suppress tagmap updating if the ASC owner is (effectively) still replicating (via checking net connection)
* This is opt in behavior is set by game code. Ideally we could buld this into replication conditions but that is not really possible.

#rnx

diff --git a/Engine/Plugins/Runtime/GameplayAbilities/Source/GameplayAbilities/Private/GameplayCueInterface.cpp b/Engine/Plugins/Runtime/GameplayAbilities/Source/GameplayAbilities/Private/GameplayCueInterface.cpp
--- a/Engine/Plugins/Runtime/GameplayAbilities/Source/GameplayAbilities/Private/GameplayCueInterface.cpp
+++ b/Engine/Plugins/Runtime/GameplayAbilities/Source/GameplayAbilities/Private/GameplayCueInterface.cpp
@@ -5,6 +5,7 @@
 #include "GameplayTagsModule.h"
 #include "AbilitySystemComponent.h"
 #include "GameplayCueSet.h"
+#include "Engine/PackageMapClient.h"
 
 
 namespace GameplayCueInterfacePrivate
@@ -391,6 +392,24 @@ bool FMinimalGameplayCueReplicationProxy::NetSerialize(FArchive& Ar, class UPack
 	}
 	else
 	{
+		// Only actually update the owner's tag map if we not the 
+		bool UpdateOwnerTagMap = Owner != nullptr;
+		if (bRequireNonOwningNetConnection && Owner)
+		{
+			if (AActor* OwningActor = Owner->GetOwner())
+			{
+				// Note we deliberately only want to do this if the NetConnection is not null
+				if (UNetConnection* OwnerNetConnection = OwningActor->GetNetConnection()) 
+				{
+					if (OwnerNetConnection == CastChecked<UPackageMapClient>(Map)->GetConnection())
+					{
+						UpdateOwnerTagMap = false;
+					}
+				}
+			}
+		}
+
+
 		NumElements = 0;
 		Ar.SerializeBits(&NumElements, NumBits);
 
@@ -415,16 +434,18 @@ bool FMinimalGameplayCueReplicationProxy::NetSerialize(FArchive& Ar, class UPack
 				// This tag already existed and is accounted for
 				LocalBitMask[LocalIdx] = false;
 			}
-			else if (Owner)
+			else if (UpdateOwnerTagMap)
 			{
+				bCachedModifiedOwnerTags = true;
 				// This is a new tag, we need to invoke the WhileActive gameplaycue event
 				Owner->SetTagMapCount(ReplicatedTag, 1);
 				Owner->InvokeGameplayCueEvent(ReplicatedTag, EGameplayCueEvent::WhileActive, Parameters);
 			}
 		}
 
-		if (Owner)
+		if (UpdateOwnerTagMap)
 		{
+			bCachedModifiedOwnerTags = true;
 			for (TConstSetBitIterator<TInlineAllocator<NumInlineTags>> It(LocalBitMask); It; ++It)
 			{
 				FGameplayTag& RemovedTag = LocalTags[It.GetIndex()];
@@ -441,7 +462,7 @@ bool FMinimalGameplayCueReplicationProxy::NetSerialize(FArchive& Ar, class UPack
 
 void FMinimalGameplayCueReplicationProxy::RemoveAllCues()
 {
-	if (!Owner)
+	if (!Owner || !bCachedModifiedOwnerTags)
 	{
 		return;
 	}
diff --git a/Engine/Plugins/Runtime/GameplayAbilities/Source/GameplayAbilities/Private/GameplayEffectTypes.cpp b/Engine/Plugins/Runtime/GameplayAbilities/Source/GameplayAbilities/Private/GameplayEffectTypes.cpp
--- a/Engine/Plugins/Runtime/GameplayAbilities/Source/GameplayAbilities/Private/GameplayEffectTypes.cpp
+++ b/Engine/Plugins/Runtime/GameplayAbilities/Source/GameplayAbilities/Private/GameplayEffectTypes.cpp
@@ -9,6 +9,7 @@
 #include "AbilitySystemGlobals.h"
 #include "AbilitySystemComponent.h"
 #include "AbilitySystemInterface.h"
+#include "Engine/PackageMapClient.h"
 
 
 #define LOCTEXT_NAMESPACE "GameplayEffectTypes"
@@ -1078,16 +1079,35 @@ bool FMinimalReplicationTagCountMap::NetSerialize(FArchive& Ar, class UPackageMa
 
 		if (Owner)
 		{
-			// Update our tags with owner tags
-			for(auto It = TagMap.CreateIterator(); It; ++It)
+			bool UpdateOwnerTagMap = true;
+			if (bRequireNonOwningNetConnection)
 			{
-				Owner->SetTagMapCount(It->Key, It->Value);
+				if (AActor* OwningActor = Owner->GetOwner())
+				{
+					// Note we deliberately only want to do this if the NetConnection is not null
+					if (UNetConnection* OwnerNetConnection = OwningActor->GetNetConnection())
+					{
+						if (OwnerNetConnection == CastChecked<UPackageMapClient>(Map)->GetConnection())
+						{
+							UpdateOwnerTagMap = false;
+						}
+					}
+				}
+			}
 
-				// Remove tags with a count of zero from the map. This prevents them
-				// from being replicated incorrectly when recording client replays.
-				if (It->Value == 0)
+			if (UpdateOwnerTagMap)
+			{
+				// Update our tags with owner tags
+				for(auto It = TagMap.CreateIterator(); It; ++It)
 				{
-					It.RemoveCurrent();
+					Owner->SetTagMapCount(It->Key, It->Value);
+
+					// Remove tags with a count of zero from the map. This prevents them
+					// from being replicated incorrectly when recording client replays.
+					if (It->Value == 0)
+					{
+						It.RemoveCurrent();
+					}
 				}
 			}
 		}
diff --git a/Engine/Plugins/Runtime/GameplayAbilities/Source/GameplayAbilities/Public/GameplayCueInterface.h b/Engine/Plugins/Runtime/GameplayAbilities/Source/GameplayAbilities/Public/GameplayCueInterface.h
--- a/Engine/Plugins/Runtime/GameplayAbilities/Source/GameplayAbilities/Public/GameplayCueInterface.h
+++ b/Engine/Plugins/Runtime/GameplayAbilities/Source/GameplayAbilities/Public/GameplayCueInterface.h
@@ -190,6 +190,9 @@ struct GAMEPLAYABILITIES_API FMinimalGameplayCueReplicationProxy
 	/** Will broadcast the OnRemove event for all currently active cues */
 	void RemoveAllCues();
 
+	/** If true, we will skip updating the Owner ASC if we replicate on a connection owned by the ASC */
+	void SetRequireNonOwningNetConnection(bool b) { bRequireNonOwningNetConnection = b; }
+
 	/** Called to init parameters */
 	TFunction<void(FGameplayCueParameters&, UAbilitySystemComponent*)> InitGameplayCueParametersFunc;
 
@@ -208,6 +211,9 @@ private:
 	UAbilitySystemComponent* Owner = nullptr;
 
 	int32 LastSourceArrayReplicationKey = -1;
+
+	bool bRequireNonOwningNetConnection = false;
+	bool bCachedModifiedOwnerTags = false;
 };
 
 template<>
diff --git a/Engine/Plugins/Runtime/GameplayAbilities/Source/GameplayAbilities/Public/GameplayEffectTypes.h b/Engine/Plugins/Runtime/GameplayAbilities/Source/GameplayAbilities/Public/GameplayEffectTypes.h
--- a/Engine/Plugins/Runtime/GameplayAbilities/Source/GameplayAbilities/Public/GameplayEffectTypes.h
+++ b/Engine/Plugins/Runtime/GameplayAbilities/Source/GameplayAbilities/Public/GameplayEffectTypes.h
@@ -1639,6 +1639,13 @@ struct GAMEPLAYABILITIES_API FMinimalReplicationTagCountMap
 	}
 
 	int32 MapID;
+
+	/** If true, we will skip updating the Owner ASC if we replicate on a connection owned by the ASC */
+	void SetRequireNonOwningNetConnection(bool b) { bRequireNonOwningNetConnection = b; }
+
+private:
+
+	bool bRequireNonOwningNetConnection = false;
 };
 
 template<>
