Make CurrentMontageStop use the passed in OverrideBlendTime variable if it is greater than or equal to 0.

diff --git a/Engine/Plugins/Runtime/GameplayAbilities/Source/GameplayAbilities/Private/AbilitySystemComponent_Abilities.cpp b/Engine/Plugins/Runtime/GameplayAbilities/Source/GameplayAbilities/Private/AbilitySystemComponent_Abilities.cpp
--- a/Engine/Plugins/Runtime/GameplayAbilities/Source/GameplayAbilities/Private/AbilitySystemComponent_Abilities.cpp
+++ b/Engine/Plugins/Runtime/GameplayAbilities/Source/GameplayAbilities/Private/AbilitySystemComponent_Abilities.cpp
@@ -2650,7 +2650,7 @@ void UAbilitySystemComponent::CurrentMontageStop(float OverrideBlendOutTime)
 	{
 		const float BlendOutTime = (OverrideBlendOutTime >= 0.0f ? OverrideBlendOutTime : MontageToStop->BlendOut.GetBlendTime());
 
-		AnimInstance->Montage_Stop(MontageToStop->BlendOut.GetBlendTime(), MontageToStop);
+		AnimInstance->Montage_Stop(BlendOutTime, MontageToStop);
 
 		if (IsOwnerActorAuthoritative())
 		{
