- Add the ability for gameplay effect execution calculations to support scoped modifiers against "temporary variables" that aren't required to be backed by an attribute capture
- Implementation basically allows for gameplay tag-identified aggregators to be created as a means for an execution to expose a temporary value to be manipulated with scoped modifiers; Allows for formulas that want manipulatable values that don't need to be captured from a source or target
- To use, an execution has to add a tag to the new member variable ValidTransientAggregatorIdentifiers; Those tags will show up in the calculation modifier array of scoped mods at the bottom, marked as temporary variables
- Update details customizations accordingly to support feature

diff --git a/Engine/Plugins/Runtime/GameplayAbilities/Source/GameplayAbilities/Private/GameplayEffectExecutionCalculation.cpp b/Engine/Plugins/Runtime/GameplayAbilities/Source/GameplayAbilities/Private/GameplayEffectExecutionCalculation.cpp
--- a/Engine/Plugins/Runtime/GameplayAbilities/Source/GameplayAbilities/Private/GameplayEffectExecutionCalculation.cpp
+++ b/Engine/Plugins/Runtime/GameplayAbilities/Source/GameplayAbilities/Private/GameplayEffectExecutionCalculation.cpp
@@ -22,17 +22,28 @@ FGameplayEffectCustomExecutionParameters::FGameplayEffectCustomExecutionParamete
 
 	for (const FGameplayEffectExecutionScopedModifierInfo& CurScopedMod : InScopedMods)
 	{
-		FAggregator* ScopedAggregator = ScopedModifierAggregators.Find(CurScopedMod.CapturedAttribute);
-		if (!ScopedAggregator)
-		{
-			const FGameplayEffectAttributeCaptureSpec* CaptureSpec = InOwningSpec.CapturedRelevantAttributes.FindCaptureSpecByDefinition(CurScopedMod.CapturedAttribute, true);
+		FAggregator* ScopedAggregator = nullptr;
 
-			FAggregator SnapshotAgg;
-			if (CaptureSpec && CaptureSpec->AttemptGetAttributeAggregatorSnapshot(SnapshotAgg))
+		// Scoped mod is a traditional, attribute-backed aggregator; Try to find an existing one or fetch one from a capture spec
+		if (CurScopedMod.AggregatorType == EGameplayEffectScopedModifierAggregatorType::CapturedAttributeBacked)
+		{
+			ScopedAggregator = ScopedModifierAggregators.Find(CurScopedMod.CapturedAttribute);
+			if (!ScopedAggregator)
 			{
-				ScopedAggregator = &(ScopedModifierAggregators.Add(CurScopedMod.CapturedAttribute, SnapshotAgg));
+				const FGameplayEffectAttributeCaptureSpec* CaptureSpec = InOwningSpec.CapturedRelevantAttributes.FindCaptureSpecByDefinition(CurScopedMod.CapturedAttribute, true);
+
+				FAggregator SnapshotAgg;
+				if (CaptureSpec && CaptureSpec->AttemptGetAttributeAggregatorSnapshot(SnapshotAgg))
+				{
+					ScopedAggregator = &(ScopedModifierAggregators.Add(CurScopedMod.CapturedAttribute, SnapshotAgg));
+				}
 			}
 		}
+		// Scoped mod is acting on a "temporary variable," find/create the appropriate transient aggregator
+		else
+		{
+			ScopedAggregator = &ScopedTransientAggregators.FindOrAdd(CurScopedMod.TransientAggregatorIdentifier);
+		}
 
 		float ModEvalValue = 0.f;
 		if (ScopedAggregator && CurScopedMod.ModifierMagnitude.AttemptCalculateMagnitude(InOwningSpec, ModEvalValue))
@@ -247,6 +258,66 @@ bool FGameplayEffectCustomExecutionParameters::ForEachQualifiedAttributeMod(cons
 	return false;
 }
 
+bool FGameplayEffectCustomExecutionParameters::AttemptCalculateTransientAggregatorMagnitude(const FGameplayTag& InAggregatorIdentifier, const FAggregatorEvaluateParameters& InEvalParams, OUT float& OutMagnitude) const
+{
+	const FAggregator* CalcAgg = ScopedTransientAggregators.Find(InAggregatorIdentifier);
+	if (CalcAgg)
+	{
+		OutMagnitude = CalcAgg->Evaluate(InEvalParams);
+		return true;
+	}
+	
+	return false;
+}
+
+bool FGameplayEffectCustomExecutionParameters::AttemptCalculateTransientAggregatorMagnitudeWithBase(const FGameplayTag& InAggregatorIdentifier, const FAggregatorEvaluateParameters& InEvalParams, float InBaseValue, OUT float& OutMagnitude) const
+{
+	const FAggregator* CalcAgg = ScopedTransientAggregators.Find(InAggregatorIdentifier);
+	if (CalcAgg)
+	{
+		OutMagnitude = CalcAgg->EvaluateWithBase(InBaseValue, InEvalParams);
+		return true;
+	}
+
+	return false;
+}
+
+bool FGameplayEffectCustomExecutionParameters::AttemptCalculateTransientAggregatorBaseValue(const FGameplayTag& InAggregatorIdentifier, OUT float& OutBaseValue) const
+{
+	const FAggregator* CalcAgg = ScopedTransientAggregators.Find(InAggregatorIdentifier);
+	if (CalcAgg)
+	{
+		OutBaseValue = CalcAgg->GetBaseValue();
+		return true;
+	}
+
+	return false;
+}
+
+bool FGameplayEffectCustomExecutionParameters::AttemptCalculateTransientAggregatorBonusMagnitude(const FGameplayTag& InAggregatorIdentifier, const FAggregatorEvaluateParameters& InEvalParams, OUT float& OutBonusMagnitude) const
+{
+	const FAggregator* CalcAgg = ScopedTransientAggregators.Find(InAggregatorIdentifier);
+	if (CalcAgg)
+	{
+		OutBonusMagnitude = CalcAgg->EvaluateBonus(InEvalParams);
+		return true;
+	}
+
+	return false;
+}
+
+bool FGameplayEffectCustomExecutionParameters::AttemptGetCapturedAttributeAggregatorSnapshot(const FGameplayTag& InAggregatorIdentifier, OUT FAggregator& OutSnapshottedAggregator) const
+{
+	const FAggregator* CalcAgg = ScopedTransientAggregators.Find(InAggregatorIdentifier);
+	if (CalcAgg)
+	{
+		OutSnapshottedAggregator.TakeSnapshotOf(*CalcAgg);
+		return true;
+	}
+
+	return false;
+}
+
 FGameplayEffectCustomExecutionOutput::FGameplayEffectCustomExecutionOutput()
 	: bTriggerConditionalGameplayEffects(false)
 	, bHandledStackCountManually(false)
@@ -320,6 +391,11 @@ void UGameplayEffectExecutionCalculation::GetValidScopedModifierAttributeCapture
 	}
 }
 
+const FGameplayTagContainer& UGameplayEffectExecutionCalculation::GetValidTransientAggregatorIdentifiers() const
+{
+	return ValidTransientAggregatorIdentifiers;
+}
+
 bool UGameplayEffectExecutionCalculation::DoesRequirePassedInTags() const
 {
 	return bRequiresPassedInTags;
diff --git a/Engine/Plugins/Runtime/GameplayAbilities/Source/GameplayAbilities/Public/GameplayEffect.h b/Engine/Plugins/Runtime/GameplayAbilities/Source/GameplayAbilities/Public/GameplayEffect.h
--- a/Engine/Plugins/Runtime/GameplayAbilities/Source/GameplayAbilities/Public/GameplayEffect.h
+++ b/Engine/Plugins/Runtime/GameplayAbilities/Source/GameplayAbilities/Public/GameplayEffect.h
@@ -351,6 +351,17 @@ struct TStructOpsTypeTraits<FGameplayEffectModifierMagnitude> : public TStructOp
 	};
 };
 
+/** Enumeration representing the types of scoped modifier aggregator usages available */
+UENUM()
+enum class EGameplayEffectScopedModifierAggregatorType : uint8
+{
+	/** Aggregator is backed by an attribute capture */
+	CapturedAttributeBacked,
+
+	/** Aggregator is entirely transient (acting as a "temporary variable") and must be identified via gameplay tag */
+	Transient
+};
+
 /** 
  * Struct representing modifier info used exclusively for "scoped" executions that happen instantaneously. These are
  * folded into a calculation only for the extent of the calculation and never permanently added to an aggregator.
@@ -362,11 +373,20 @@ struct FGameplayEffectExecutionScopedModifierInfo
 
 	// Constructors
 	FGameplayEffectExecutionScopedModifierInfo()
-		: ModifierOp(EGameplayModOp::Additive)
+		: AggregatorType(EGameplayEffectScopedModifierAggregatorType::CapturedAttributeBacked)
+		, ModifierOp(EGameplayModOp::Additive)
 	{}
 
 	FGameplayEffectExecutionScopedModifierInfo(const FGameplayEffectAttributeCaptureDefinition& InCaptureDef)
 		: CapturedAttribute(InCaptureDef)
+		, AggregatorType(EGameplayEffectScopedModifierAggregatorType::CapturedAttributeBacked)
+		, ModifierOp(EGameplayModOp::Additive)
+	{
+	}
+
+	FGameplayEffectExecutionScopedModifierInfo(const FGameplayTag& InTransientAggregatorIdentifier)
+		: TransientAggregatorIdentifier(InTransientAggregatorIdentifier)
+		, AggregatorType(EGameplayEffectScopedModifierAggregatorType::Transient)
 		, ModifierOp(EGameplayModOp::Additive)
 	{
 	}
@@ -375,6 +395,14 @@ struct FGameplayEffectExecutionScopedModifierInfo
 	UPROPERTY(VisibleDefaultsOnly, Category=Execution)
 	FGameplayEffectAttributeCaptureDefinition CapturedAttribute;
 
+	/** Identifier for aggregator if acting as a transient "temporary variable" aggregator */
+	UPROPERTY(VisibleDefaultsOnly, Category=Execution)
+	FGameplayTag TransientAggregatorIdentifier;
+
+	/** Type of aggregator backing the scoped mod */
+	UPROPERTY(VisibleDefaultsOnly, Category=Execution)
+	EGameplayEffectScopedModifierAggregatorType AggregatorType;
+
 	/** Modifier operation to perform */
 	UPROPERTY(EditDefaultsOnly, Category=Execution)
 	TEnumAsByte<EGameplayModOp::Type> ModifierOp;
diff --git a/Engine/Plugins/Runtime/GameplayAbilities/Source/GameplayAbilities/Public/GameplayEffectExecutionCalculation.h b/Engine/Plugins/Runtime/GameplayAbilities/Source/GameplayAbilities/Public/GameplayEffectExecutionCalculation.h
--- a/Engine/Plugins/Runtime/GameplayAbilities/Source/GameplayAbilities/Public/GameplayEffectExecutionCalculation.h
+++ b/Engine/Plugins/Runtime/GameplayAbilities/Source/GameplayAbilities/Public/GameplayEffectExecutionCalculation.h
@@ -117,11 +117,69 @@ public:
 	/** Runs given TFunction on every qualifier mod for a given AttributeCaptureDefinition */
 	bool ForEachQualifiedAttributeMod(const FGameplayEffectAttributeCaptureDefinition& InCaptureDef, const FAggregatorEvaluateParameters& InEvalParams, TFunction< void(EGameplayModEvaluationChannel, EGameplayModOp::Type, const FAggregatorMod&) >) const;
 
+	/**
+	 * Attempts to calculate the magnitude of a transient aggregator given the specified parameters.
+	 * 
+	 * @param InAggregatorIdentifier	Tag identifying the transient aggregator to attempt to calculate the magnitude of
+	 * @param InEvalParams				Parameters to evaluate the aggregator under
+	 * @param OutMagnitude				[OUT] Computed magnitude
+	 * 
+	 * @return True if the magnitude was successfully calculated, false if it was not
+	 */
+	bool AttemptCalculateTransientAggregatorMagnitude(const FGameplayTag& InAggregatorIdentifier, const FAggregatorEvaluateParameters& InEvalParams, OUT float& OutMagnitude) const;
+	
+	/**
+	 * Attempts to calculate the magnitude of a transient aggregator given the specified parameters, including a starting base value.
+	 * 
+	 * @param InAggregatorIdentifier	Tag identifying the transient aggregator to attempt to calculate the magnitude of
+	 * @param InEvalParams				Parameters to evaluate the attribute under
+	 * @param InBaseValue				Base value to evaluate the attribute under
+	 * @param OutMagnitude				[OUT] Computed magnitude
+	 * 
+	 * @return True if the magnitude was successfully calculated, false if it was not
+	 */
+	bool AttemptCalculateTransientAggregatorMagnitudeWithBase(const FGameplayTag& InAggregatorIdentifier, const FAggregatorEvaluateParameters& InEvalParams, float InBaseValue, OUT float& OutMagnitude) const;
+
+	/**
+	 * Attempts to calculate the base value of a transient aggregator given the specified parameters.
+	 * 
+	 * @param InAggregatorIdentifier	Tag identifying the transient aggregator to attempt to calculate the base value of
+	 * @param OutBaseValue				[OUT] Computed base value
+	 * 
+	 * @return True if the base value was successfully calculated, false if it was not
+	 */
+	bool AttemptCalculateTransientAggregatorBaseValue(const FGameplayTag& InAggregatorIdentifier, OUT float& OutBaseValue) const;
+
+	/**
+	 * Attempts to calculate the bonus magnitude of a transient aggregator given the specified parameters.
+	 * 
+	 * @param InAggregatorIdentifier		Tag identifying the transient aggregator to attempt to calculate the bonus magnitude of
+	 * @param InEvalParams					Parameters to evaluate the attribute under
+	 * @param OutBonusMagnitude				[OUT] Computed bonus magnitude
+	 * 
+	 * @return True if the bonus magnitude was successfully calculated, false if it was not
+	 */
+	bool AttemptCalculateTransientAggregatorBonusMagnitude(const FGameplayTag& InAggregatorIdentifier, const FAggregatorEvaluateParameters& InEvalParams, OUT float& OutBonusMagnitude) const;
+	
+	/**
+	 * Attempts to populate the specified aggregator with a snapshot of a backing transient aggregator. Can fail if the transient aggregator doesn't exist as a result
+	 * of no scoped mods targeting it.
+	 * 
+	 * @param InAggregatorIdentifier	Tag identifying the transient aggregator to attempt to snapshot
+	 * @param OutSnapshottedAggregator	[OUT] Snapshotted aggregator, if possible
+	 * 
+	 * @return True if the aggregator was successfully snapshotted, false if it was not
+	 */
+	bool AttemptGetCapturedAttributeAggregatorSnapshot(const FGameplayTag& InAggregatorIdentifier, OUT FAggregator& OutSnapshottedAggregator) const;
+
 private:
 
 	/** Mapping of capture definition to aggregator with scoped modifiers added in; Used to process scoped modifiers w/o modifying underlying aggregators in the capture */
 	TMap<FGameplayEffectAttributeCaptureDefinition, FAggregator> ScopedModifierAggregators;
 
+	/** Mapping of tag identifier to aggregators with scoped modifiers added in; Used to process "temporary variable" support w/o requiring backed attribute capture definitions in scoped modifiers */
+	TMap<FGameplayTag, FAggregator> ScopedTransientAggregators;
+
 	/** Owning gameplay effect spec */
 	FGameplayEffectSpec* OwningSpec;
 
@@ -215,6 +273,10 @@ protected:
 	UPROPERTY(EditDefaultsOnly, Category=Attributes)
 	TArray<FGameplayEffectAttributeCaptureDefinition> InvalidScopedModifierAttributes;
 
+	/** Any tag in this container will show up as a valid "temporary variable" for scoped modifiers; Used to allow for data-driven variable support that doesn't rely on scoped modifiers */
+	UPROPERTY(EditDefaultsOnly, Category=NonAttributeCalculation)
+	FGameplayTagContainer ValidTransientAggregatorIdentifiers;
+
 public:
 	/**
 	 * Gets the collection of capture attribute definitions that the calculation class will accept as valid scoped modifiers
@@ -223,6 +285,13 @@ public:
 	 */
 	virtual void GetValidScopedModifierAttributeCaptureDefinitions(OUT TArray<FGameplayEffectAttributeCaptureDefinition>& OutScopableModifiers) const;
 
+	/**
+	 * Gets the collection of identifiers of valid transient aggregators ("temporary variable aggregators")
+	 * 
+	 * @return Array of valid transient aggregator identifiers
+	 */
+	virtual const FGameplayTagContainer& GetValidTransientAggregatorIdentifiers() const;
+
 	/** Returns if this execution requires passed in tags */
 	virtual bool DoesRequirePassedInTags() const;
 
diff --git a/Engine/Plugins/Runtime/GameplayAbilities/Source/GameplayAbilitiesEditor/Private/GameplayEffectExecutionDefinitionDetails.cpp b/Engine/Plugins/Runtime/GameplayAbilities/Source/GameplayAbilitiesEditor/Private/GameplayEffectExecutionDefinitionDetails.cpp
--- a/Engine/Plugins/Runtime/GameplayAbilities/Source/GameplayAbilitiesEditor/Private/GameplayEffectExecutionDefinitionDetails.cpp
+++ b/Engine/Plugins/Runtime/GameplayAbilities/Source/GameplayAbilitiesEditor/Private/GameplayEffectExecutionDefinitionDetails.cpp
@@ -74,8 +74,9 @@ void FGameplayEffectExecutionDefinitionDetails::OnCalculationClassChanged()
 void FGameplayEffectExecutionDefinitionDetails::UpdateCalculationModifiers()
 {
 	TArray<FGameplayEffectAttributeCaptureDefinition> ValidCaptureDefinitions;
-	
-	// Try to extract the collection of valid capture definitions from the execution class CDO, if possible
+	FGameplayTagContainer ValidTransientAggregatorIdentifiers;
+
+	// Try to extract the collection of valid capture definitions & transient aggregator identifiers from the execution class CDO, if possible
 	if (CalculationClassPropHandle.IsValid())
 	{
 		UObject* ObjValue = nullptr;
@@ -88,6 +89,7 @@ void FGameplayEffectExecutionDefinitionDetails::UpdateCalculationModifiers()
 			if (ExecutionCDO)
 			{
 				ExecutionCDO->GetValidScopedModifierAttributeCaptureDefinitions(ValidCaptureDefinitions);
+				ValidTransientAggregatorIdentifiers = ExecutionCDO->GetValidTransientAggregatorIdentifiers();
 
 				// Grab this while we are at it so we know if we should show the 'Passed In Tags' property
 				bShowPassedInTags = ExecutionCDO->DoesRequirePassedInTags();
@@ -95,24 +97,21 @@ void FGameplayEffectExecutionDefinitionDetails::UpdateCalculationModifiers()
 		}
 	}
 
-	// Want to hide the calculation modifiers if there are no valid definitions
-	bShowCalculationModifiers = (ValidCaptureDefinitions.Num() > 0);
+	// Want to hide the calculation modifiers if there are no valid definitions or transient aggregator IDs
+	bShowCalculationModifiers = (ValidCaptureDefinitions.Num() > 0) || (ValidTransientAggregatorIdentifiers.Num() > 0);
 
-	// Need to prune out any modifiers that are specified for definitions that aren't specified by the execution class
+	// Need to prune out any modifiers that aren't supported by the execution class
 	if (CalculationModifiersArrayPropHandle.IsValid())
 	{
 		uint32 NumChildren = 0;
 		CalculationModifiersArrayPropHandle->GetNumElements(NumChildren);
 
-		// If there aren't any valid definitions, just dump the whole array
-		if (ValidCaptureDefinitions.Num() == 0)
+		// If there aren't any valid definitions or transient aggregator IDs, just dump the whole array
+		if (ValidCaptureDefinitions.Num() == 0 && ValidTransientAggregatorIdentifiers.Num() == 0 && NumChildren > 0)
 		{
-			if (NumChildren > 0)
-			{
-				CalculationModifiersArrayPropHandle->EmptyArray();
-			}
+			CalculationModifiersArrayPropHandle->EmptyArray();
 		}
-		// There are some valid definitions, so verify any existing ones to make sure they are in the valid array
+		// There are some valid definitions/IDs, so verify any existing ones to make sure they are in the valid array
 		else
 		{
 			for (int32 ChildIdx = NumChildren - 1; ChildIdx >= 0; --ChildIdx)
@@ -125,7 +124,10 @@ void FGameplayEffectExecutionDefinitionDetails::UpdateCalculationModifiers()
 				// @todo: For now, only allow single editing
 				ensure(RawScopedModInfoStructs.Num() == 1);
 				const FGameplayEffectExecutionScopedModifierInfo& CurModInfo = *reinterpret_cast<const FGameplayEffectExecutionScopedModifierInfo*>(RawScopedModInfoStructs[0]);
-				if (!ValidCaptureDefinitions.Contains(CurModInfo.CapturedAttribute))
+				
+				const bool bInvalidAttributeAggregator = (CurModInfo.AggregatorType == EGameplayEffectScopedModifierAggregatorType::CapturedAttributeBacked && !ValidCaptureDefinitions.Contains(CurModInfo.CapturedAttribute));
+				const bool bInvalidTransientAggregator = (CurModInfo.AggregatorType == EGameplayEffectScopedModifierAggregatorType::Transient && !ValidTransientAggregatorIdentifiers.HasTagExact(CurModInfo.TransientAggregatorIdentifier));
+				if (bInvalidAttributeAggregator || bInvalidTransientAggregator)
 				{
 					CalculationModifiersArrayPropHandle->DeleteItem(ChildIdx);
 				}
diff --git a/Engine/Plugins/Runtime/GameplayAbilities/Source/GameplayAbilitiesEditor/Private/GameplayEffectExecutionScopedModifierInfoDetails.cpp b/Engine/Plugins/Runtime/GameplayAbilities/Source/GameplayAbilitiesEditor/Private/GameplayEffectExecutionScopedModifierInfoDetails.cpp
--- a/Engine/Plugins/Runtime/GameplayAbilities/Source/GameplayAbilitiesEditor/Private/GameplayEffectExecutionScopedModifierInfoDetails.cpp
+++ b/Engine/Plugins/Runtime/GameplayAbilities/Source/GameplayAbilitiesEditor/Private/GameplayEffectExecutionScopedModifierInfoDetails.cpp
@@ -4,15 +4,63 @@
 #include "Widgets/Layout/SBorder.h"
 #include "Widgets/Text/STextBlock.h"
 #include "Widgets/Layout/SGridPanel.h"
+#include "Widgets/Layout/SWidgetSwitcher.h"
 #include "Widgets/Input/SComboBox.h"
 #include "DetailWidgetRow.h"
 #include "GameplayEffect.h"
 #include "GameplayEffectExecutionCalculation.h"
 #include "DetailLayoutBuilder.h"
 #include "SlateOptMacros.h"
+#include "GameplayEffectTypes.h"
 
 #define LOCTEXT_NAMESPACE "GameplayEffectExecutionScopedModifierInfoDetailsCustomization"
 
+/** Simple struct tracking scoped mod backing data from the scoped mod info */
+struct FAggregatorDetailsBackingData
+{
+	// Constructors
+	FAggregatorDetailsBackingData()
+		: AggregatorType(EGameplayEffectScopedModifierAggregatorType::CapturedAttributeBacked)
+		, CaptureDefinition()
+		, TransientAggregatorIdentifier()
+	{
+	}
+
+	FAggregatorDetailsBackingData(const FGameplayEffectAttributeCaptureDefinition& InCaptureDef)
+		: AggregatorType(EGameplayEffectScopedModifierAggregatorType::CapturedAttributeBacked)
+		, CaptureDefinition(InCaptureDef)
+		, TransientAggregatorIdentifier()
+	{
+	}
+
+	FAggregatorDetailsBackingData(const FGameplayTag& InIdentifier)
+		: AggregatorType(EGameplayEffectScopedModifierAggregatorType::Transient)
+		, CaptureDefinition()
+		, TransientAggregatorIdentifier(InIdentifier)
+	{
+	}
+
+	// Equality operators
+	bool operator==(const FAggregatorDetailsBackingData& Other) const
+	{
+		return (AggregatorType == Other.AggregatorType) && (CaptureDefinition == Other.CaptureDefinition) && (TransientAggregatorIdentifier == Other.TransientAggregatorIdentifier);
+	}
+
+	bool operator!=(const FAggregatorDetailsBackingData& Other) const
+	{
+		return (AggregatorType != Other.AggregatorType) || (CaptureDefinition != Other.CaptureDefinition) || (TransientAggregatorIdentifier != Other.TransientAggregatorIdentifier);
+	}
+
+	/** Aggregator type for the backing data */
+	EGameplayEffectScopedModifierAggregatorType AggregatorType;
+
+	/** Capture definition used if attribute-backed aggregator type */
+	FGameplayEffectAttributeCaptureDefinition CaptureDefinition;
+
+	/** Identifier used if transient/"temporary variable" aggregator type */
+	FGameplayTag TransientAggregatorIdentifier;
+};
+
 TSharedRef<IPropertyTypeCustomization> FGameplayEffectExecutionScopedModifierInfoDetails::MakeInstance()
 {
 	return MakeShareable(new FGameplayEffectExecutionScopedModifierInfoDetails());
@@ -29,12 +77,12 @@ void FGameplayEffectExecutionScopedModifierInfoDetails::CustomizeHeader(TSharedR
 
 BEGIN_SLATE_FUNCTION_BUILD_OPTIMIZATION
 
-/** Custom widget class to cleanly represent a capture definition in a combo box */
-class SCaptureDefWidget : public SCompoundWidget
+/** Custom widget class to cleanly represent mod backing data in a combo box */
+class SScopedModBackingDataWidget : public SCompoundWidget
 {
 public:
 
-	SLATE_BEGIN_ARGS(SCaptureDefWidget) {}
+	SLATE_BEGIN_ARGS(SScopedModBackingDataWidget) {}
 	SLATE_END_ARGS()
 
 	/** Construct the widget from a grid panel wrapped with a border */
@@ -42,79 +90,126 @@ public:
 	{
 		this->ChildSlot
 		[
-			SNew(SBorder)
-			.HAlign(HAlign_Fill)
+			SNew(SWidgetSwitcher)
+			.WidgetIndex(this, &SScopedModBackingDataWidget::GetWidgetSwitcherIdx)
+			+SWidgetSwitcher::Slot()
 			[
-				SNew(SGridPanel)
-				+SGridPanel::Slot(0, 0)
-				.HAlign(HAlign_Right)
-				.Padding(FMargin(2.f))
-				[
-					SNew(STextBlock)
-					.Text(NSLOCTEXT("ScopedModifierDetails", "CapturedAttributeLabel", "Captured Attribute:"))
-					.Font(IDetailLayoutBuilder::GetDetailFontBold())
-				]
-				+SGridPanel::Slot(1, 0)
-				.HAlign(HAlign_Left)
-				.Padding(FMargin(2.f))
+				SNew(SBorder)
+				.HAlign(HAlign_Fill)
 				[
-					SNew(STextBlock)
-					.Text(this, &SCaptureDefWidget::GetCapturedAttributeText)
-					.Font(IDetailLayoutBuilder::GetDetailFont())
+					SNew(SGridPanel)
+					+SGridPanel::Slot(0, 0)
+					.HAlign(HAlign_Right)
+					.Padding(FMargin(2.f))
+					[
+						SNew(STextBlock)
+						.Text(NSLOCTEXT("ScopedModifierDetails", "CapturedAttributeLabel", "Captured Attribute:"))
+						.Font(IDetailLayoutBuilder::GetDetailFontBold())
+					]
+					+SGridPanel::Slot(1, 0)
+					.HAlign(HAlign_Left)
+					.Padding(FMargin(2.f))
+					[
+						SNew(STextBlock)
+						.Text(this, &SScopedModBackingDataWidget::GetCapturedAttributeText)
+						.Font(IDetailLayoutBuilder::GetDetailFont())
+					]
+					+SGridPanel::Slot(0, 1)
+					.HAlign(HAlign_Right)
+					.Padding(FMargin(2.f))
+					[
+						SNew(STextBlock)
+						.Text(NSLOCTEXT("ScopedModifierDetails", "CapturedAttributeSourceLabel", "Captured Source:"))
+						.Font(IDetailLayoutBuilder::GetDetailFontBold())
+					]
+					+SGridPanel::Slot(1, 1)
+					.HAlign(HAlign_Left)
+					.Padding(FMargin(2.f))
+					[
+						SNew(STextBlock)
+						.Text(this, &SScopedModBackingDataWidget::GetCapturedAttributeSourceText)
+						.Font(IDetailLayoutBuilder::GetDetailFont())
+					]
+					+SGridPanel::Slot(0, 2)
+					.HAlign(HAlign_Right)
+					.Padding(FMargin(2.f))
+					[
+						SNew(STextBlock)
+						.Text(NSLOCTEXT("ScopedModifierDetails", "CapturedAttributeSnapshotLabel", "Captured Status:"))
+						.Font(IDetailLayoutBuilder::GetDetailFontBold())
+					]
+					+SGridPanel::Slot(1, 2)
+					.HAlign(HAlign_Left)
+					.Padding(FMargin(2.f))
+					[
+						SNew(STextBlock)
+						.Text(this, &SScopedModBackingDataWidget::GetCapturedAttributeSnapshotText)
+						.Font(IDetailLayoutBuilder::GetDetailFont())
+					]
 				]
-				+SGridPanel::Slot(0, 1)
-				.HAlign(HAlign_Right)
-				.Padding(FMargin(2.f))
-				[
-					SNew(STextBlock)
-					.Text(NSLOCTEXT("ScopedModifierDetails", "CapturedAttributeSourceLabel", "Captured Source:"))
-					.Font(IDetailLayoutBuilder::GetDetailFontBold())
-				]
-				+SGridPanel::Slot(1, 1)
-				.HAlign(HAlign_Left)
-				.Padding(FMargin(2.f))
-				[
-					SNew(STextBlock)
-					.Text(this, &SCaptureDefWidget::GetCapturedAttributeSourceText)
-					.Font(IDetailLayoutBuilder::GetDetailFont())
-				]
-				+SGridPanel::Slot(0, 2)
-				.HAlign(HAlign_Right)
-				.Padding(FMargin(2.f))
-				[
-					SNew(STextBlock)
-					.Text(NSLOCTEXT("ScopedModifierDetails", "CapturedAttributeSnapshotLabel", "Captured Status:"))
-					.Font(IDetailLayoutBuilder::GetDetailFontBold())
-				]
-				+SGridPanel::Slot(1, 2)
-				.HAlign(HAlign_Left)
-				.Padding(FMargin(2.f))
+			]
+			+SWidgetSwitcher::Slot()
+			[
+				SNew(SBorder)
+				.HAlign(HAlign_Fill)
 				[
-					SNew(STextBlock)
-					.Text(this, &SCaptureDefWidget::GetCapturedAttributeSnapshotText)
-					.Font(IDetailLayoutBuilder::GetDetailFont())
+					SNew(SGridPanel)
+					+SGridPanel::Slot(0, 0)
+					.HAlign(HAlign_Right)
+					.Padding(FMargin(2.f))
+					[
+						SNew(STextBlock)
+						.Text(NSLOCTEXT("ScopedModifierDetails", "TempVariableLabel", "Temporary Variable:"))
+						.Font(IDetailLayoutBuilder::GetDetailFontBold())
+					]
+					+SGridPanel::Slot(1, 0)
+					.HAlign(HAlign_Left)
+					.Padding(FMargin(2.f))
+					[
+						SNew(STextBlock)
+						.Text(this, &SScopedModBackingDataWidget::GetTemporaryVariableText)
+						.Font(IDetailLayoutBuilder::GetDetailFont())
+					]
+					+SGridPanel::Slot(0, 1)
+					.HAlign(HAlign_Center)
+					.Padding(FMargin(2.f))
+					.ColumnSpan(2)
+					[
+						SNew(STextBlock)
+						.Text(NSLOCTEXT("ScopedModifierDetails", "TempVariableDesc", "Temporary value exposed by calculation."))
+						.Font(IDetailLayoutBuilder::GetDetailFont())
+					]
 				]
 			]
 		];
 	}
 
-	/** Set the definition that backs the widget; This is used as a way to reduce expensive FText creations constantly */
-	void SetBackingDefinition(const FGameplayEffectAttributeCaptureDefinition& InDefinition)
+	/** Set the data that backs the widget; This is used as a way to reduce expensive FText creations constantly */
+	void SetBackingData(const FAggregatorDetailsBackingData& InBackingData)
 	{
 		static const FText SnapshotText(NSLOCTEXT("ScopedModifierDetails", "CapturedAttributeSnapshotted", "Snapshotted"));
 		static const FText NotSnapshotText(NSLOCTEXT("ScopedModifierDetails", "CapturedAttributeNotSnapshotted", "Not Snapshotted"));
 
-		if (InDefinition != BackingDefinition)
+		if (InBackingData != BackingData)
 		{
-			BackingDefinition = InDefinition;
-			CapturedAttributeText = FText::FromString(BackingDefinition.AttributeToCapture.GetName());
-			UEnum::GetDisplayValueAsText(TEXT("GameplayAbilities.EGameplayEffectAttributeCaptureSource"), BackingDefinition.AttributeSource, CapturedAttributeSourceText);
-			CapturedAttributeSnapshotText = BackingDefinition.bSnapshot ? SnapshotText : NotSnapshotText;
+			BackingData = InBackingData;
+			CapturedAttributeText = FText::FromString(BackingData.CaptureDefinition.AttributeToCapture.GetName());
+			UEnum::GetDisplayValueAsText(TEXT("GameplayAbilities.EGameplayEffectAttributeCaptureSource"), BackingData.CaptureDefinition.AttributeSource, CapturedAttributeSourceText);
+			CapturedAttributeSnapshotText = BackingData.CaptureDefinition.bSnapshot ? SnapshotText : NotSnapshotText;
+			TemporaryVariableText = FText::FromString(BackingData.TransientAggregatorIdentifier.ToString());
+
+			WidgetSwitcherIdx = (BackingData.AggregatorType == EGameplayEffectScopedModifierAggregatorType::CapturedAttributeBacked) ? 0 : 1;
 		}
 	}
 
 private:
 
+	/** Simple accessor to cached widget switcher index */
+	int32 GetWidgetSwitcherIdx() const
+	{
+		return WidgetSwitcherIdx;
+	}
+
 	/** Simple accessor to cached captured attribute text */
 	FText GetCapturedAttributeText() const
 	{
@@ -133,8 +228,14 @@ private:
 		return CapturedAttributeSnapshotText;
 	}
 
-	/** Capture definition backing the widget */
-	FGameplayEffectAttributeCaptureDefinition BackingDefinition;
+	/** Simple accessor to cached temporary variable text */
+	FText GetTemporaryVariableText() const
+	{
+		return TemporaryVariableText;
+	}
+
+	/** Aggregator data backing the widget */
+	FAggregatorDetailsBackingData BackingData;
 
 	/** Cached attribute text */
 	FText CapturedAttributeText;
@@ -144,14 +245,21 @@ private:
 
 	/** Cached attribute snapshot status text */
 	FText CapturedAttributeSnapshotText;
+
+	/** Cached temporary variable status text */
+	FText TemporaryVariableText;
+
+	/** Cached widget switcher index */
+	int32 WidgetSwitcherIdx;
 };
 
 END_SLATE_FUNCTION_BUILD_OPTIMIZATION
 
 void FGameplayEffectExecutionScopedModifierInfoDetails::CustomizeChildren(TSharedRef<IPropertyHandle> StructPropertyHandle, IDetailChildrenBuilder& StructBuilder, IPropertyTypeCustomizationUtils& StructCustomizationUtils)
 {
-	AvailableCaptureDefs.Empty();
-	CaptureDefPropertyHandle = StructPropertyHandle->GetChildHandle(GET_MEMBER_NAME_CHECKED(FGameplayEffectExecutionScopedModifierInfo, CapturedAttribute));
+	AvailableBackingData.Empty();
+
+	ScopedModifierStructPropertyHandle = StructPropertyHandle;
 
 	TSharedPtr<IPropertyHandle> ParentArrayHandle = StructPropertyHandle->GetParentHandle();
 	const bool bIsExecutionDefAttribute = (ParentArrayHandle.IsValid() && ParentArrayHandle->GetProperty()->GetOuter() == FGameplayEffectExecutionDefinition::StaticStruct());
@@ -184,7 +292,13 @@ void FGameplayEffectExecutionScopedModifierInfoDetails::CustomizeChildren(TShare
 
 							for (const FGameplayEffectAttributeCaptureDefinition& CurDef : CaptureDefs)
 							{
-								AvailableCaptureDefs.Add(MakeShareable(new FGameplayEffectAttributeCaptureDefinition(CurDef)));
+								AvailableBackingData.Add(MakeShareable(new FAggregatorDetailsBackingData(CurDef)));
+							}
+
+							FGameplayTagContainer ValidTransientAggregatorIdentifiers = ExecCalcCDO->GetValidTransientAggregatorIdentifiers();
+							for (const FGameplayTag& CurIdentifier : ValidTransientAggregatorIdentifiers)
+							{
+								AvailableBackingData.Add(MakeShareable(new FAggregatorDetailsBackingData(CurIdentifier)));
 							}
 						}
 					}
@@ -192,36 +306,36 @@ void FGameplayEffectExecutionScopedModifierInfoDetails::CustomizeChildren(TShare
 			}
 		}
 
-		// Construct a custom combo box widget outlining possible capture definition choices
-		if (AvailableCaptureDefs.Num() > 0)
+		// Construct a custom combo box widget outlining possible backing data choices
+		if (AvailableBackingData.Num() > 0)
 		{
-			TSharedPtr< SComboBox< TSharedPtr<FGameplayEffectAttributeCaptureDefinition> > > BackingComboBox;
+			TSharedPtr< SComboBox< TSharedPtr<FAggregatorDetailsBackingData> > > BackingComboBox;
 
-			StructBuilder.AddCustomRow(NSLOCTEXT("ScopedModifierDetails", "CaptureDefLabel", "Backing Capture Definition"))
+			StructBuilder.AddCustomRow(NSLOCTEXT("ScopedModifierDetails", "BackingDataLabel", "Backing Data"))
 			.NameContent()
 			[
 				SNew(STextBlock)
-				.Text(NSLOCTEXT("ScopedModifierDetails", "CaptureDefLabel", "Backing Capture Definition"))
-				.ToolTipText(NSLOCTEXT("ScopedModifierDetails", "CaptureDefTooltip", "The capture definition to use to populate the scoped modifier. Only options specified by the execution class are presented here."))
+				.Text(NSLOCTEXT("ScopedModifierDetails", "BackingDataLabel", "Backing Data"))
+				.ToolTipText(NSLOCTEXT("ScopedModifierDetails", "BackingDataTooltip", "The backing data to use to populate the scoped modifier. Only options specified by the execution class are presented here."))
 				.Font(StructCustomizationUtils.GetRegularFont())
 			]
 			.ValueContent()
 			.MinDesiredWidth(350.f)
 			[
-				SAssignNew(BackingComboBox, SComboBox< TSharedPtr<FGameplayEffectAttributeCaptureDefinition> >)
-				.OptionsSource(&AvailableCaptureDefs)
-				.OnSelectionChanged(this, &FGameplayEffectExecutionScopedModifierInfoDetails::OnCaptureDefComboBoxSelectionChanged)
-				.OnGenerateWidget(this, &FGameplayEffectExecutionScopedModifierInfoDetails::OnGenerateCaptureDefComboWidget)
+				SAssignNew(BackingComboBox, SComboBox< TSharedPtr<FAggregatorDetailsBackingData> >)
+				.OptionsSource(&AvailableBackingData)
+				.OnSelectionChanged(this, &FGameplayEffectExecutionScopedModifierInfoDetails::OnBackingDataComboBoxSelectionChanged)
+				.OnGenerateWidget(this, &FGameplayEffectExecutionScopedModifierInfoDetails::OnGenerateBackingDataComboWidget)
 				.Content()
 				[
-					SAssignNew(PrimaryCaptureDefWidget, SCaptureDefWidget)
+					SAssignNew(PrimaryBackingDataWidget, SScopedModBackingDataWidget)
 				]
 			];
 
 			// Set the initial value on the combo box; done this way to intentionally cause a change delegate
 			if (BackingComboBox.IsValid())
 			{
-				BackingComboBox->SetSelectedItem(GetCurrentCaptureDef());
+				BackingComboBox->SetSelectedItem(GetCurrentBackingData());
 			}
 		}
 	}
@@ -236,69 +350,100 @@ void FGameplayEffectExecutionScopedModifierInfoDetails::CustomizeChildren(TShare
 		const TSharedRef<IPropertyHandle> ChildHandle = StructPropertyHandle->GetChildHandle(ChildIndex).ToSharedRef();
 		const FName ChildPropName = ChildHandle->GetProperty()->GetFName();
 
-		if (!bIsExecutionDefAttribute || (ChildPropName != GET_MEMBER_NAME_CHECKED(FGameplayEffectExecutionScopedModifierInfo, CapturedAttribute)))
+		const bool bAlwaysShowableProperty = ((ChildPropName != GET_MEMBER_NAME_CHECKED(FGameplayEffectExecutionScopedModifierInfo, CapturedAttribute))
+			&& (ChildPropName != GET_MEMBER_NAME_CHECKED(FGameplayEffectExecutionScopedModifierInfo, TransientAggregatorIdentifier))
+			&& (ChildPropName != GET_MEMBER_NAME_CHECKED(FGameplayEffectExecutionScopedModifierInfo, AggregatorType)));
+
+		if (!bIsExecutionDefAttribute || bAlwaysShowableProperty)
 		{
 			StructBuilder.AddProperty(ChildHandle);
 		}
 	}
 }
 
-void FGameplayEffectExecutionScopedModifierInfoDetails::OnCaptureDefComboBoxSelectionChanged(TSharedPtr<FGameplayEffectAttributeCaptureDefinition> InSelectedItem, ESelectInfo::Type InSelectInfo)
+void FGameplayEffectExecutionScopedModifierInfoDetails::OnBackingDataComboBoxSelectionChanged(TSharedPtr<FAggregatorDetailsBackingData> InSelectedItem, ESelectInfo::Type InSelectInfo)
 {
-	SetCurrentCaptureDef(InSelectedItem);
+	SetCurrentBackingData(InSelectedItem);
 
-	// Need to update the base capture widget in the combo box manually due to caching strategy
-	if (PrimaryCaptureDefWidget.IsValid())
+	// Need to update the base backing data widget in the combo box manually due to caching strategy
+	if (PrimaryBackingDataWidget.IsValid())
 	{
-		PrimaryCaptureDefWidget->SetBackingDefinition(*InSelectedItem.Get());
+		PrimaryBackingDataWidget->SetBackingData(*InSelectedItem.Get());
 	}
 }
 
-TSharedRef<SWidget> FGameplayEffectExecutionScopedModifierInfoDetails::OnGenerateCaptureDefComboWidget(TSharedPtr<FGameplayEffectAttributeCaptureDefinition> InItem)
+TSharedRef<SWidget> FGameplayEffectExecutionScopedModifierInfoDetails::OnGenerateBackingDataComboWidget(TSharedPtr<FAggregatorDetailsBackingData> InItem)
 {
-	TSharedPtr<SCaptureDefWidget> NewCapDefWidget;
-	SAssignNew(NewCapDefWidget, SCaptureDefWidget);
-	NewCapDefWidget->SetBackingDefinition(*InItem.Get());
+	TSharedPtr<SScopedModBackingDataWidget> NewBackingDataWidget;
+	SAssignNew(NewBackingDataWidget, SScopedModBackingDataWidget);
+	NewBackingDataWidget->SetBackingData(*InItem.Get());
 
-	return NewCapDefWidget.ToSharedRef();
+	return NewBackingDataWidget.ToSharedRef();
 }
 
-TSharedPtr<FGameplayEffectAttributeCaptureDefinition> FGameplayEffectExecutionScopedModifierInfoDetails::GetCurrentCaptureDef() const
+TSharedPtr<FAggregatorDetailsBackingData> FGameplayEffectExecutionScopedModifierInfoDetails::GetCurrentBackingData() const
 {
-	if (CaptureDefPropertyHandle.IsValid() && CaptureDefPropertyHandle->GetProperty())
+	if (ScopedModifierStructPropertyHandle.IsValid() && ScopedModifierStructPropertyHandle->GetProperty())
 	{
 		TArray<const void*> RawStructPtrs;
-		CaptureDefPropertyHandle->AccessRawData(RawStructPtrs);
+		ScopedModifierStructPropertyHandle->AccessRawData(RawStructPtrs);
 
 		// Only showing the combo box for single-editing
-		const FGameplayEffectAttributeCaptureDefinition& BackingDef = *reinterpret_cast<const FGameplayEffectAttributeCaptureDefinition*>(RawStructPtrs[0]);
+		const FGameplayEffectExecutionScopedModifierInfo& ScopedModifierInfo = *reinterpret_cast<const FGameplayEffectExecutionScopedModifierInfo*>(RawStructPtrs[0]);
 
-		for (TSharedPtr<FGameplayEffectAttributeCaptureDefinition> CurCaptureDef : AvailableCaptureDefs)
+		for (TSharedPtr<FAggregatorDetailsBackingData> CurBackingData : AvailableBackingData)
 		{
-			if (CurCaptureDef.IsValid() && (*CurCaptureDef.Get() == BackingDef))
+			if (CurBackingData.IsValid() && (CurBackingData->AggregatorType == ScopedModifierInfo.AggregatorType))
 			{
-				return CurCaptureDef;
+				if (ScopedModifierInfo.AggregatorType == EGameplayEffectScopedModifierAggregatorType::CapturedAttributeBacked)
+				{
+					if (ScopedModifierInfo.CapturedAttribute == CurBackingData->CaptureDefinition)
+					{
+						return CurBackingData;
+					}
+				}
+				else
+				{
+					if (ScopedModifierInfo.TransientAggregatorIdentifier == CurBackingData->TransientAggregatorIdentifier)
+					{
+						return CurBackingData;
+					}
+				}
 			}
 		}
 	}
 
-	return AvailableCaptureDefs[0];
+	return AvailableBackingData[0];
 }
 
-void FGameplayEffectExecutionScopedModifierInfoDetails::SetCurrentCaptureDef(TSharedPtr<FGameplayEffectAttributeCaptureDefinition> InDef)
+void FGameplayEffectExecutionScopedModifierInfoDetails::SetCurrentBackingData(TSharedPtr<FAggregatorDetailsBackingData> InBackingData)
 {
-	if (CaptureDefPropertyHandle.IsValid() && CaptureDefPropertyHandle->GetProperty() && InDef.IsValid())
+	if (ScopedModifierStructPropertyHandle.IsValid() && ScopedModifierStructPropertyHandle->GetProperty() && InBackingData.IsValid())
 	{
 		TArray<void*> RawStructPtrs;
-		CaptureDefPropertyHandle->AccessRawData(RawStructPtrs);
+		ScopedModifierStructPropertyHandle->AccessRawData(RawStructPtrs);
 
-		FGameplayEffectAttributeCaptureDefinition& BackingDef = *reinterpret_cast<FGameplayEffectAttributeCaptureDefinition*>(RawStructPtrs[0]);
-		const FGameplayEffectAttributeCaptureDefinition& InDefRef = *InDef.Get();
-		if (BackingDef != InDefRef)
+		FGameplayEffectExecutionScopedModifierInfo& ScopedModifierInfo = *reinterpret_cast<FGameplayEffectExecutionScopedModifierInfo*>(RawStructPtrs[0]);
+		
+		const bool bSameData = (ScopedModifierInfo.AggregatorType == InBackingData->AggregatorType)
+			&& (ScopedModifierInfo.CapturedAttribute == InBackingData->CaptureDefinition)
+			&& (ScopedModifierInfo.TransientAggregatorIdentifier == InBackingData->TransientAggregatorIdentifier);
+
+		if (!bSameData)
 		{
-			CaptureDefPropertyHandle->NotifyPreChange();
-			BackingDef = InDefRef;
-			CaptureDefPropertyHandle->NotifyPostChange();
+			ScopedModifierStructPropertyHandle->NotifyPreChange();
+			ScopedModifierInfo.AggregatorType = InBackingData->AggregatorType;
+			if (ScopedModifierInfo.AggregatorType == EGameplayEffectScopedModifierAggregatorType::CapturedAttributeBacked)
+			{
+				ScopedModifierInfo.CapturedAttribute = InBackingData->CaptureDefinition;
+				ScopedModifierInfo.TransientAggregatorIdentifier = FGameplayTag::EmptyTag;
+			}
+			else
+			{
+				ScopedModifierInfo.CapturedAttribute = FGameplayEffectAttributeCaptureDefinition();
+				ScopedModifierInfo.TransientAggregatorIdentifier = InBackingData->TransientAggregatorIdentifier;
+			}
+			ScopedModifierStructPropertyHandle->NotifyPostChange();
 		}
 	}
 }
diff --git a/Engine/Plugins/Runtime/GameplayAbilities/Source/GameplayAbilitiesEditor/Private/GameplayEffectExecutionScopedModifierInfoDetails.h b/Engine/Plugins/Runtime/GameplayAbilities/Source/GameplayAbilitiesEditor/Private/GameplayEffectExecutionScopedModifierInfoDetails.h
--- a/Engine/Plugins/Runtime/GameplayAbilities/Source/GameplayAbilitiesEditor/Private/GameplayEffectExecutionScopedModifierInfoDetails.h
+++ b/Engine/Plugins/Runtime/GameplayAbilities/Source/GameplayAbilitiesEditor/Private/GameplayEffectExecutionScopedModifierInfoDetails.h
@@ -5,17 +5,18 @@
 #include "CoreMinimal.h"
 #include "Widgets/SWidget.h"
 #include "IPropertyTypeCustomization.h"
-#include "GameplayEffectTypes.h"
 #include "PropertyHandle.h"
 #include "IDetailChildrenBuilder.h"
 
 class FDetailWidgetRow;
-class SCaptureDefWidget;
+class SScopedModBackingDataWidget;
+struct FAggregatorDetailsBackingData;
 
 /** Details customization for FGameplayEffectExecutionScopedModifierInfo */
 class FGameplayEffectExecutionScopedModifierInfoDetails : public IPropertyTypeCustomization
 {
 public:
+	
 	static TSharedRef<IPropertyTypeCustomization> MakeInstance();
 
 	/** Overridden to provide the property name */
@@ -27,23 +28,23 @@ public:
 private:
 
 	/** Delegate called when combo box selection is changed */
-	void OnCaptureDefComboBoxSelectionChanged(TSharedPtr<FGameplayEffectAttributeCaptureDefinition> InSelectedItem, ESelectInfo::Type InSelectInfo);
+	void OnBackingDataComboBoxSelectionChanged(TSharedPtr<FAggregatorDetailsBackingData> InSelectedItem, ESelectInfo::Type InSelectInfo);
 	
 	/* Called to generate the widgets for custom combo box entries */
-	TSharedRef<SWidget> OnGenerateCaptureDefComboWidget(TSharedPtr<FGameplayEffectAttributeCaptureDefinition> InItem);
+	TSharedRef<SWidget> OnGenerateBackingDataComboWidget(TSharedPtr<FAggregatorDetailsBackingData> InItem);
 	
-	/** Get the current capture definition as specified by the backing property, if possible; Otherwise falls back to first available definition from execution class */
-	TSharedPtr<FGameplayEffectAttributeCaptureDefinition> GetCurrentCaptureDef() const;
+	/** Get the current aggregator backing data, if possible; Otherwise falls back to first available definition from execution class */
+	TSharedPtr<FAggregatorDetailsBackingData> GetCurrentBackingData() const;
 
 	/** Set the current capture definition */
-	void SetCurrentCaptureDef(TSharedPtr<FGameplayEffectAttributeCaptureDefinition> InDef);
+	void SetCurrentBackingData(TSharedPtr<FAggregatorDetailsBackingData> InBackingData);
 
-	/** Cached property handle for the capture definition property */
-	TSharedPtr<IPropertyHandle> CaptureDefPropertyHandle;
+	/** Cached property handle for the overall scoped modifier info struct */
+	TSharedPtr<IPropertyHandle> ScopedModifierStructPropertyHandle;
 
-	/** Primary capture definition widget shown for the custom combo box */
-	TSharedPtr<SCaptureDefWidget> PrimaryCaptureDefWidget;
+	/** Primary backing data widget shown for the custom combo box */
+	TSharedPtr<SScopedModBackingDataWidget> PrimaryBackingDataWidget;
 
-	/** Backing source for the custom combo box; Populated by all valid definitions from the execution class */
-	TArray< TSharedPtr<FGameplayEffectAttributeCaptureDefinition> > AvailableCaptureDefs;
+	/** Backing source for the custom combo box; Populated by all valid backing data from the execution class */
+	TArray< TSharedPtr<FAggregatorDetailsBackingData> > AvailableBackingData;
 };
