Speculative fix for GameplayCueNotify_Actor crashing when Pawn::OnDestroyed delegate is broadcast

diff --git a/Engine/Plugins/Runtime/GameplayAbilities/Source/GameplayAbilities/Public/GameplayCue_Types.h b/Engine/Plugins/Runtime/GameplayAbilities/Source/GameplayAbilities/Public/GameplayCue_Types.h
--- a/Engine/Plugins/Runtime/GameplayAbilities/Source/GameplayAbilities/Public/GameplayCue_Types.h
+++ b/Engine/Plugins/Runtime/GameplayAbilities/Source/GameplayAbilities/Public/GameplayCue_Types.h
@@ -73,6 +73,7 @@ struct FPreallocationInfo
 	GENERATED_USTRUCT_BODY()
 
 	/** Raw list of pooled instances. This relies on NotifyGameplayCueActorEndPlay always being called when actor is destroyed */
+	UPROPERTY(transient)
 	TMap<UClass*, TArray<AGameplayCueNotify_Actor*> >	PreallocatedInstances;
 
 	/** List of calsses that will be pooled */
