Add GC references to pooled gameplay cues

diff --git a/Engine/Plugins/Runtime/GameplayAbilities/Source/GameplayAbilities/Private/GameplayCueManager.cpp b/Engine/Plugins/Runtime/GameplayAbilities/Source/GameplayAbilities/Private/GameplayCueManager.cpp
--- a/Engine/Plugins/Runtime/GameplayAbilities/Source/GameplayAbilities/Private/GameplayCueManager.cpp
+++ b/Engine/Plugins/Runtime/GameplayAbilities/Source/GameplayAbilities/Private/GameplayCueManager.cpp
@@ -1635,6 +1635,22 @@ void UGameplayCueManager::OnWorldCleanup(UWorld* World, bool bSessionEnded, bool
 	IGameplayCueInterface::ClearTagToFunctionMap();
 }
 
+void UGameplayCueManager::AddReferencedObjects(UObject* InThis, FReferenceCollector& Collector)
+{
+	Super::AddReferencedObjects(InThis, Collector);
+
+	UGameplayCueManager* GCManager = Cast<UGameplayCueManager>(InThis);
+	for (int32 idx = 0; idx < GCManager->PreallocationInfoList_Internal.Num(); ++idx)
+	{
+		FPreallocationInfo& Info = GCManager->PreallocationInfoList_Internal[idx];
+		for (auto &It : Info.PreallocatedInstances)
+		{
+			TArray<AGameplayCueNotify_Actor*>& GameplayCueArray = It.Value;
+			Collector.AddReferencedObjects(GameplayCueArray);
+		}
+	}
+}
+
 void UGameplayCueManager::DumpPreallocationStats(UWorld* World)
 {
 	if (World == nullptr)
diff --git a/Engine/Plugins/Runtime/GameplayAbilities/Source/GameplayAbilities/Public/GameplayCueManager.h b/Engine/Plugins/Runtime/GameplayAbilities/Source/GameplayAbilities/Public/GameplayCueManager.h
--- a/Engine/Plugins/Runtime/GameplayAbilities/Source/GameplayAbilities/Public/GameplayCueManager.h
+++ b/Engine/Plugins/Runtime/GameplayAbilities/Source/GameplayAbilities/Public/GameplayCueManager.h
@@ -259,6 +259,7 @@ protected:
 	TSharedPtr<FStreamableHandle> GameplayCueAssetHandle;
 
 public:		
+	static void AddReferencedObjects(UObject* InThis, FReferenceCollector& Collector);
 
 	/** Called before loading any gameplay cue notifies from object libraries. Allows subclasses to skip notifies. */
 	virtual bool ShouldLoadGameplayCueAssetData(const FAssetData& Data) const { return true; }
