Security fixes for FGameplayAbilityTargetDataHandle::NetSerialize

#rnx

diff --git a/Engine/Plugins/Runtime/GameplayAbilities/Source/GameplayAbilities/Private/AbilitySystemGlobals.cpp b/Engine/Plugins/Runtime/GameplayAbilities/Source/GameplayAbilities/Private/AbilitySystemGlobals.cpp
--- a/Engine/Plugins/Runtime/GameplayAbilities/Source/GameplayAbilities/Private/AbilitySystemGlobals.cpp
+++ b/Engine/Plugins/Runtime/GameplayAbilities/Source/GameplayAbilities/Private/AbilitySystemGlobals.cpp
@@ -46,6 +46,8 @@ void UAbilitySystemGlobals::InitGlobalData()
 	GetGameplayTagResponseTable();
 	InitGlobalTags();
 
+	TargetDataStructCache.InitForType(FGameplayAbilityTargetData::StaticStruct());
+
 	// Register for PreloadMap so cleanup can occur on map transitions
 	FCoreUObjectDelegates::PreLoadMap.AddUObject(this, &UAbilitySystemGlobals::HandlePreLoadMap);
 
@@ -502,3 +504,48 @@ void UAbilitySystemGlobals::NonShipping_ApplyGlobalAbilityScaler_Duration(float&
 	}
 #endif
 }
+
+void FNetSerializeScriptStructCache::InitForType(UScriptStruct* InScriptStruct)
+{
+	// Find all script structs of this type and add them to the list
+	// (not sure of a better way to do this but it should only happen once at startup)
+	for (TObjectIterator<UScriptStruct> It; It; ++It)
+	{
+		if (It->IsChildOf(InScriptStruct))
+		{
+			ScriptStructs.Add(*It);
+		}
+	}
+	
+	ScriptStructs.Sort([](const UScriptStruct& A, const UScriptStruct& B) { return A.GetName().ToLower() > B.GetName().ToLower(); });
+}
+
+bool FNetSerializeScriptStructCache::NetSerialize(FArchive& Ar, UScriptStruct*& Struct)
+{
+	if (Ar.IsSaving())
+	{
+		int32 idx;
+		if (ScriptStructs.Find(Struct, idx))
+		{
+			check(idx < (1 << 8));
+			uint8 b = idx;
+			Ar.SerializeBits(&b, 8);
+			return true;
+		}
+		ABILITY_LOG(Error, TEXT("Could not find %s in ScriptStructCache"), *GetNameSafe(Struct));
+		return false;
+	}
+	else
+	{
+		uint8 b = 0;
+		Ar.SerializeBits(&b, 8);
+		if (ScriptStructs.IsValidIndex(b))
+		{
+			Struct = ScriptStructs[b];
+			return true;
+		}
+
+		ABILITY_LOG(Error, TEXT("Could not script struct at idx %d"), b);
+		return false;
+	}
+}
\ No newline at end of file
diff --git a/Engine/Plugins/Runtime/GameplayAbilities/Source/GameplayAbilities/Private/GameplayAbilityTargetTypes.cpp b/Engine/Plugins/Runtime/GameplayAbilities/Source/GameplayAbilities/Private/GameplayAbilityTargetTypes.cpp
--- a/Engine/Plugins/Runtime/GameplayAbilities/Source/GameplayAbilities/Private/GameplayAbilityTargetTypes.cpp
+++ b/Engine/Plugins/Runtime/GameplayAbilities/Source/GameplayAbilities/Private/GameplayAbilityTargetTypes.cpp
@@ -4,6 +4,7 @@
 #include "GameplayEffect.h"
 #include "AbilitySystemComponent.h"
 #include "AbilitySystemBlueprintLibrary.h"
+#include "AbilitySystemGlobals.h"
 
 
 TArray<FActiveGameplayEffectHandle> FGameplayAbilityTargetData::ApplyGameplayEffect(const UGameplayEffect* GameplayEffect, const FGameplayEffectContextHandle& InEffectContext, float Level, FPredictionKey PredictionKey)
@@ -134,6 +135,11 @@ FGameplayAbilityTargetDataHandle FGameplayAbilityTargetingLocationInfo::MakeTarg
 	return ReturnDataHandle;
 }
 
+// If defined, we'll serialize target data in a safer way (untested/unproven still: goal should be to remove old code asap)
+#ifndef TARGETDATAHANDLE_SAFE_NET_SERIALIZE
+#define TARGETDATAHANDLE_SAFE_NET_SERIALIZE 1
+#endif
+
 bool FGameplayAbilityTargetDataHandle::NetSerialize(FArchive& Ar, class UPackageMap* Map, bool& bOutSuccess)
 {
 	Ar << UniqueId;
@@ -158,7 +164,11 @@ bool FGameplayAbilityTargetDataHandle::NetSerialize(FArchive& Ar, class UPackage
 	{
 		TCheckedObjPtr<UScriptStruct> ScriptStruct = Data[i].IsValid() ? Data[i]->GetScriptStruct() : NULL;
 
+#if TARGETDATAHANDLE_SAFE_NET_SERIALIZE
+		UAbilitySystemGlobals::Get().TargetDataStructCache.NetSerialize(Ar, ScriptStruct.Get());
+#else
 		Ar << ScriptStruct;
+#endif
 
 		if (ScriptStruct.IsValid())
 		{
diff --git a/Engine/Plugins/Runtime/GameplayAbilities/Source/GameplayAbilities/Public/AbilitySystemGlobals.h b/Engine/Plugins/Runtime/GameplayAbilities/Source/GameplayAbilities/Public/AbilitySystemGlobals.h
--- a/Engine/Plugins/Runtime/GameplayAbilities/Source/GameplayAbilities/Public/AbilitySystemGlobals.h
+++ b/Engine/Plugins/Runtime/GameplayAbilities/Source/GameplayAbilities/Public/AbilitySystemGlobals.h
@@ -22,6 +22,21 @@ struct FGameplayEffectSpecForRPC;
 DECLARE_MULTICAST_DELEGATE_TwoParams(FOnAbilitySystemAssetOpenedDelegate, FString , int );
 DECLARE_MULTICAST_DELEGATE_TwoParams(FOnAbilitySystemAssetFoundDelegate, FString, int);
 
+//  Container for safely replicating  script struct references (constrained to a specified parent struct)
+USTRUCT()
+struct FNetSerializeScriptStructCache
+{
+	GENERATED_BODY()
+
+	void InitForType(UScriptStruct* InScriptStruct);
+
+	// Serializes reference to given script struct (must be in the cache)
+	bool NetSerialize(FArchive& Ar, UScriptStruct*& Struct);
+
+	UPROPERTY()
+	TArray<UScriptStruct*> ScriptStructs;
+};
+
 
 /** Holds global data for the ability system. Can be configured per project via config file */
 UCLASS(config=Game)
@@ -263,6 +278,9 @@ class GAMEPLAYABILITIES_API UAbilitySystemGlobals : public UObject
 	/** Path where the engine will load gameplay cue notifies from */
 	virtual TArray<FString> GetGameplayCueNotifyPaths() { return GameplayCueNotifyPaths; }
 
+	UPROPERTY()
+	FNetSerializeScriptStructCache	TargetDataStructCache;
+
 protected:
 
 	virtual void InitAttributeDefaults();
