Non unity compile fix
#rnx

diff --git a/Engine/Plugins/Runtime/GameplayAbilities/Source/GameplayAbilities/Private/AbilitySystemGlobals.cpp b/Engine/Plugins/Runtime/GameplayAbilities/Source/GameplayAbilities/Private/AbilitySystemGlobals.cpp
--- a/Engine/Plugins/Runtime/GameplayAbilities/Source/GameplayAbilities/Private/AbilitySystemGlobals.cpp
+++ b/Engine/Plugins/Runtime/GameplayAbilities/Source/GameplayAbilities/Private/AbilitySystemGlobals.cpp
@@ -9,6 +9,7 @@
 #include "GameplayCueManager.h"
 #include "GameplayTagResponseTable.h"
 #include "GameplayTagsManager.h"
+#include "UObjectIterator.h"
 
 #if WITH_EDITOR
 #include "Editor.h"
