Remove dead code from the ability system.

diff --git a/Engine/Plugins/Runtime/GameplayAbilities/Source/GameplayAbilities/Private/AbilitySystemComponent_Abilities.cpp b/Engine/Plugins/Runtime/GameplayAbilities/Source/GameplayAbilities/Private/AbilitySystemComponent_Abilities.cpp
--- a/Engine/Plugins/Runtime/GameplayAbilities/Source/GameplayAbilities/Private/AbilitySystemComponent_Abilities.cpp
+++ b/Engine/Plugins/Runtime/GameplayAbilities/Source/GameplayAbilities/Private/AbilitySystemComponent_Abilities.cpp
@@ -963,7 +963,7 @@ void UAbilitySystemComponent::OnRep_ActivateAbilities()
 
 	// Try to run any pending activations that couldn't run before. If they don't work now, kill them
 
-	for (auto PendingAbilityInfo : PendingServerActivatedAbilities)
+	for (const FPendingAbilityInfo& PendingAbilityInfo : PendingServerActivatedAbilities)
 	{
 		if (PendingAbilityInfo.bPartiallyActivated)
 		{
@@ -1827,48 +1827,10 @@ bool UAbilitySystemComponent::TriggerAbilityFromGameplayEvent(FGameplayAbilitySp
 	// Run on the non-instanced ability
 	if (Ability->ShouldAbilityRespondToEvent(ActorInfo, &TempEventData))
 	{
-		int32 ExecutingAbilityIndex = -1;
-
-		// if we're the server and this is coming from a predicted event we should check if the client has already predicted it
-		if (ScopedPredictionKey.IsValidKey()
-			&& Ability->GetNetExecutionPolicy() == EGameplayAbilityNetExecutionPolicy::LocalPredicted
-			&& ActorInfo->OwnerActor->Role == ROLE_Authority)
-		{
-			bool bPendingClientAbilityFound = false;
-			for (auto PendingAbilityInfo : Component.PendingClientActivatedAbilities)
-			{
-				if (ScopedPredictionKey.Current == PendingAbilityInfo.PredictionKey.Base && Handle == PendingAbilityInfo.Handle) // found a match
-				{
-					Component.PendingClientActivatedAbilities.RemoveSingleSwap(PendingAbilityInfo);
-					bPendingClientAbilityFound = true;
-					break;
-				}
-			}
-
-			// we haven't received the client's copy of the triggered ability
-			// keep track of this so we can associate the prediction keys when it comes in
-			if (bPendingClientAbilityFound == false)
-			{
-				UAbilitySystemComponent::FExecutingAbilityInfo Info;
-				Info.PredictionKey = ScopedPredictionKey;
-				Info.Handle = Handle;
-
-				ExecutingAbilityIndex = Component.ExecutingServerAbilities.Add(Info);
-			}
-		}
-
 		if (InternalTryActivateAbility(Handle, ScopedPredictionKey, nullptr, nullptr, &TempEventData))
 		{
-			if (ExecutingAbilityIndex >= 0)
-			{
-				Component.ExecutingServerAbilities[ExecutingAbilityIndex].State = UAbilitySystemComponent::EAbilityExecutionState::Succeeded;
-			}
 			return true;
 		}
-		else if (ExecutingAbilityIndex >= 0)
-		{
-			Component.ExecutingServerAbilities[ExecutingAbilityIndex].State = UAbilitySystemComponent::EAbilityExecutionState::Failed;
-		}
 	}
 	return false;
 }
diff --git a/Engine/Plugins/Runtime/GameplayAbilities/Source/GameplayAbilities/Public/AbilitySystemComponent.h b/Engine/Plugins/Runtime/GameplayAbilities/Source/GameplayAbilities/Public/AbilitySystemComponent.h
--- a/Engine/Plugins/Runtime/GameplayAbilities/Source/GameplayAbilities/Public/AbilitySystemComponent.h
+++ b/Engine/Plugins/Runtime/GameplayAbilities/Source/GameplayAbilities/Public/AbilitySystemComponent.h
@@ -249,40 +249,9 @@ class GAMEPLAYABILITIES_API UAbilitySystemComponent : public UGameplayTasksCompo
 		{}
 	};
 
-	/**
-	 * This is a list of GameplayAbilities that are predicted by the client and were triggered by abilities that were also predicted by the client
-	 * When the server version of the predicted ability executes it should trigger copies of these and the copies will be associated with the correct prediction keys
-	 */
-	TArray<FPendingAbilityInfo> PendingClientActivatedAbilities;
-
 	/** This is a list of GameplayAbilities that were activated on the server and can't yet execute on the client. It will try to execute these at a later point */
 	TArray<FPendingAbilityInfo> PendingServerActivatedAbilities;
 
-	/** State of execution for an ability, used to track on the server */
-	enum class EAbilityExecutionState : uint8
-	{
-		Executing,
-		Succeeded,
-		Failed,
-	};
-
-	struct FExecutingAbilityInfo
-	{
-		FExecutingAbilityInfo() : State(EAbilityExecutionState::Executing) {};
-
-		bool operator==(const FExecutingAbilityInfo& Other) const
-		{
-			return PredictionKey == Other.PredictionKey	&& State == Other.State;
-		}
-
-		FPredictionKey PredictionKey;
-		EAbilityExecutionState State;
-		FGameplayAbilitySpecHandle Handle;
-	};
-
-	/** List of all executing abilities the server knows about */
-	TArray<FExecutingAbilityInfo> ExecutingServerAbilities;
-
 	// ----------------------------------------------------------------------------------------------------------------
 	//	GameplayEffects: Primary outward facing API for other systems
 	// ----------------------------------------------------------------------------------------------------------------
