Fix the creation of the object libraries for UGameplayCueNotify_Static and AGameplayCueNotify_Actor. Previously they were being assigned to the incorrect variables.

diff --git a/Engine/Plugins/Runtime/GameplayAbilities/Source/GameplayAbilities/Private/GameplayCueManager.cpp b/Engine/Plugins/Runtime/GameplayAbilities/Source/GameplayAbilities/Private/GameplayCueManager.cpp
--- a/Engine/Plugins/Runtime/GameplayAbilities/Source/GameplayAbilities/Private/GameplayCueManager.cpp
+++ b/Engine/Plugins/Runtime/GameplayAbilities/Source/GameplayAbilities/Private/GameplayCueManager.cpp
@@ -732,7 +732,7 @@ TSharedPtr<FStreamableHandle> UGameplayCueManager::InitObjectLibrary(FGameplayCu
 	// Instantiate the UObjectLibraries if they aren't there already
 	if (!Lib.StaticObjectLibrary)
 	{
-		Lib.StaticObjectLibrary = UObjectLibrary::CreateLibrary(AGameplayCueNotify_Actor::StaticClass(), true, GIsEditor && !IsRunningCommandlet());
+		Lib.StaticObjectLibrary = UObjectLibrary::CreateLibrary(UGameplayCueNotify_Static::StaticClass(), true, GIsEditor && !IsRunningCommandlet());
 		if (GIsEditor)
 		{
 			Lib.StaticObjectLibrary->bIncludeOnlyOnDiskAssets = false;
@@ -740,7 +740,7 @@ TSharedPtr<FStreamableHandle> UGameplayCueManager::InitObjectLibrary(FGameplayCu
 	}
 	if (!Lib.ActorObjectLibrary)
 	{
-		Lib.ActorObjectLibrary = UObjectLibrary::CreateLibrary(UGameplayCueNotify_Static::StaticClass(), true, GIsEditor && !IsRunningCommandlet());
+		Lib.ActorObjectLibrary = UObjectLibrary::CreateLibrary(AGameplayCueNotify_Actor::StaticClass(), true, GIsEditor && !IsRunningCommandlet());
 		if (GIsEditor)
 		{
 			Lib.ActorObjectLibrary->bIncludeOnlyOnDiskAssets = false;
