Make UGameplayAbility::MontageStop use the OverrideBlendOutTime parameter.

PR #6247: Fix GameplayAbility not passing OverrideBlendOutTime to CurrentMontageStop (Contributed by [REDACTED_PERSON])

diff --git a/Engine/Plugins/Runtime/GameplayAbilities/Source/GameplayAbilities/Private/Abilities/GameplayAbility.cpp b/Engine/Plugins/Runtime/GameplayAbilities/Source/GameplayAbilities/Private/Abilities/GameplayAbility.cpp
--- a/Engine/Plugins/Runtime/GameplayAbilities/Source/GameplayAbilities/Private/Abilities/GameplayAbility.cpp
+++ b/Engine/Plugins/Runtime/GameplayAbilities/Source/GameplayAbilities/Private/Abilities/GameplayAbility.cpp
@@ -1157,7 +1157,7 @@ void UGameplayAbility::MontageStop(float OverrideBlendOutTime)
 		// We should only stop the current montage if we are the animating ability
 		if (AbilitySystemComponent->IsAnimatingAbility(this))
 		{
-			AbilitySystemComponent->CurrentMontageStop();
+			AbilitySystemComponent->CurrentMontageStop(OverrideBlendOutTime);
 		}
 	}
 }
diff --git a/Engine/Plugins/Runtime/GameplayAbilities/Source/GameplayAbilities/Public/Abilities/GameplayAbility.h b/Engine/Plugins/Runtime/GameplayAbilities/Source/GameplayAbilities/Public/Abilities/GameplayAbility.h
--- a/Engine/Plugins/Runtime/GameplayAbilities/Source/GameplayAbilities/Public/Abilities/GameplayAbility.h
+++ b/Engine/Plugins/Runtime/GameplayAbilities/Source/GameplayAbilities/Public/Abilities/GameplayAbility.h
@@ -775,7 +775,7 @@ protected:
 	/**
 	 * Stops the current animation montage.
 	 *
-	 * @param OverrideBlendTime If < 0, will override the BlendOutTime parameter on the AnimMontage instance
+	 * @param OverrideBlendTime If >= 0, will override the BlendOutTime parameter on the AnimMontage instance
 	 */
 	UFUNCTION(BlueprintCallable, Category="Ability|Animation", Meta = (AdvancedDisplay = "OverrideBlendOutTime"))
 	void MontageStop(float OverrideBlendOutTime = -1.0f);
