Don't try to fixup pin string references if we're not saving the package.
Update parsing of the pin string values for FGameplayAttribute pins to properly roundtrip.

diff --git a/Engine/Plugins/Runtime/GameplayAbilities/Source/GameplayAbilitiesEditor/Private/SGameplayAttributeGraphPin.cpp b/Engine/Plugins/Runtime/GameplayAbilities/Source/GameplayAbilitiesEditor/Private/SGameplayAttributeGraphPin.cpp
--- a/Engine/Plugins/Runtime/GameplayAbilities/Source/GameplayAbilitiesEditor/Private/SGameplayAttributeGraphPin.cpp
+++ b/Engine/Plugins/Runtime/GameplayAbilities/Source/GameplayAbilitiesEditor/Private/SGameplayAttributeGraphPin.cpp
@@ -10,30 +10,18 @@
 void SGameplayAttributeGraphPin::Construct( const FArguments& InArgs, UEdGraphPin* InGraphPinObj )
 {
 	SGraphPin::Construct( SGraphPin::FArguments(), InGraphPinObj );
-	LastSelectedProperty = NULL;
+	LastSelectedProperty = nullptr;
 
 }
 
 TSharedRef<SWidget>	SGameplayAttributeGraphPin::GetDefaultValueWidget()
 {
-	// Parse out current default value
-	// It will be in the form (Attribute=/Script/<PackageName>.<ObjectName>:<PropertyName>)
-	
+	// Parse out current default value	
 	FString DefaultString = GraphPinObj->GetDefaultAsString();
 	FGameplayAttribute DefaultAttribute;
 
-	if (DefaultString.StartsWith(TEXT("(")) && DefaultString.EndsWith(TEXT(")")))
-	{
-		DefaultString = DefaultString.LeftChop(1);
-		DefaultString = DefaultString.RightChop(1);
-
-		DefaultString.Split("=", NULL, &DefaultString);
-
-		DefaultString = DefaultString.LeftChop(1);
-		DefaultString = DefaultString.RightChop(1);
-
-		DefaultAttribute.SetUProperty(FindObject<UProperty>(ANY_PACKAGE, *DefaultString));
-	}
+	UScriptStruct* PinLiteralStructType = FGameplayAttribute::StaticStruct();
+	PinLiteralStructType->ImportText(*DefaultString, &DefaultAttribute, nullptr, EPropertyPortFlags::PPF_SerializedAsImportText, nullptr, PinLiteralStructType->GetName(), true);
 
 	//Create widget
 	return SNew(SVerticalBox)
@@ -49,23 +37,13 @@ TSharedRef<SWidget>	SGameplayAttributeGraphPin::GetDefaultValueWidget()
 void SGameplayAttributeGraphPin::OnAttributeChanged(UProperty* SelectedAttribute)
 {
 	FString FinalValue;
+	FGameplayAttribute NewAttributeStruct;
+	NewAttributeStruct.SetUProperty(SelectedAttribute);
 
-	if (SelectedAttribute == nullptr)
-	{
-		FinalValue = FString(TEXT("()"));
-	}
-	else
-	{
-		FinalValue = FString::Printf(TEXT("(Attribute=\"%s\")"), *SelectedAttribute->GetPathName());
-	}
+	FGameplayAttribute::StaticStruct()->ExportText(FinalValue, &NewAttributeStruct, &NewAttributeStruct, nullptr, EPropertyPortFlags::PPF_SerializedAsImportText, nullptr);
 
 	GraphPinObj->GetSchema()->TrySetDefaultValue(*GraphPinObj, FinalValue);
-
-
-	FString DefTagString = GraphPinObj->GetDefaultAsString();
-
 	LastSelectedProperty = SelectedAttribute;
-
 }
 
 #undef LOCTEXT_NAMESPACE
diff --git a/Engine/Source/Editor/BlueprintGraph/Private/K2Node.cpp b/Engine/Source/Editor/BlueprintGraph/Private/K2Node.cpp
--- a/Engine/Source/Editor/BlueprintGraph/Private/K2Node.cpp
+++ b/Engine/Source/Editor/BlueprintGraph/Private/K2Node.cpp
@@ -91,7 +91,7 @@ void UK2Node::Serialize(FArchive& Ar)
 {
 	Ar.UsingCustomVersion(FFrameworkObjectVersion::GUID);
 
-	if (Ar.IsSaving())
+	if (Ar.IsSaving() && !GIsSavingPackage)
 	{
 		if (Ar.IsObjectReferenceCollector() || Ar.Tell() < 0)
 		{
