Always create server initiated prediction key in mixed mode hack. This fixes the OnActive event being called twice on clients in some situations. The `ScopedPredictionKey.IsValidKey()` clause was actually a typo from the previous version of that function, we want to always be creating the server side prediction key for this case.

#rnx

diff --git a/Engine/Plugins/Runtime/GameplayAbilities/Source/GameplayAbilities/Private/AbilitySystemComponent.cpp b/Engine/Plugins/Runtime/GameplayAbilities/Source/GameplayAbilities/Private/AbilitySystemComponent.cpp
--- a/Engine/Plugins/Runtime/GameplayAbilities/Source/GameplayAbilities/Private/AbilitySystemComponent.cpp
+++ b/Engine/Plugins/Runtime/GameplayAbilities/Source/GameplayAbilities/Private/AbilitySystemComponent.cpp
@@ -1164,10 +1164,7 @@ void UAbilitySystemComponent::AddGameplayCue_Internal(const FGameplayTag Gamepla
 				if (GameplayCueContainer.bMinimalReplication)
 				{
 					// For *replicated to sim proxies only* container, Create a Server Initiated PK to avoid double playing on the auto proxy in mixed replication mode (Original Hack)
-					if (ScopedPredictionKey.IsValidKey())
-					{
-						PredictionKeyForRPC = FPredictionKey::CreateNewServerInitiatedKey(this);
-					}
+					PredictionKeyForRPC = FPredictionKey::CreateNewServerInitiatedKey(this);
 				}
 				else
 				{
