Don't try to fixup pin string references if we're not saving the package.
Update parsing of the pin string values for attributes.

diff --git a/Engine/Plugins/Runtime/GameplayAbilities/Source/GameplayAbilitiesEditor/Private/SGameplayAttributeGraphPin.cpp b/Engine/Plugins/Runtime/GameplayAbilities/Source/GameplayAbilitiesEditor/Private/SGameplayAttributeGraphPin.cpp
--- a/Engine/Plugins/Runtime/GameplayAbilities/Source/GameplayAbilitiesEditor/Private/SGameplayAttributeGraphPin.cpp
+++ b/Engine/Plugins/Runtime/GameplayAbilities/Source/GameplayAbilitiesEditor/Private/SGameplayAttributeGraphPin.cpp
@@ -21,7 +21,10 @@ TSharedRef<SWidget>	SGameplayAttributeGraphPin::GetDefaultValueWidget()
 	FGameplayAttribute DefaultAttribute;
 
 	UScriptStruct* PinLiteralStructType = FGameplayAttribute::StaticStruct();
-	PinLiteralStructType->ImportText(*DefaultString, &DefaultAttribute, nullptr, EPropertyPortFlags::PPF_SerializedAsImportText, nullptr, PinLiteralStructType->GetName(), true);
+	if (!DefaultString.IsEmpty())
+	{
+		PinLiteralStructType->ImportText(*DefaultString, &DefaultAttribute, nullptr, EPropertyPortFlags::PPF_SerializedAsImportText, GError, PinLiteralStructType->GetName(), true);
+	}
 
 	//Create widget
 	return SNew(SVerticalBox)
