Implemented Task - Ability to remove a GE based on tag application

GameplayEffect
RemovalTagRequirements - Tag requirements that if met will remove this effect. Also prevents effect application.
RemoveGameplayEffectsQuery - On Application of this effect, any active effect that matches this query will be removed. Queries are more powerful but slightly slower than RemoveGameplayEffectsWithTags.

CheckRemovalTagRequirements - helper method to check if the removal tag reqruirements have been met
AttempRemoveActiveEffectsOnEffectApplication - Called on application of an effect to remove all effects that pass RemoveGameplayEffectsWithTags or RemovalTagRequirements checks
OnOwnerTagChange - Updated to remove any effects that meet the removal tag requirements

AbilitySystemComponent
ApplyGameplayEffectSpecToSelf - Updated to check removal requirements and to remove all effects that pass the new checks

#rnx

diff --git a/Engine/Plugins/Runtime/GameplayAbilities/Source/GameplayAbilities/Private/AbilitySystemComponent.cpp b/Engine/Plugins/Runtime/GameplayAbilities/Source/GameplayAbilities/Private/AbilitySystemComponent.cpp
--- a/Engine/Plugins/Runtime/GameplayAbilities/Source/GameplayAbilities/Private/AbilitySystemComponent.cpp
+++ b/Engine/Plugins/Runtime/GameplayAbilities/Source/GameplayAbilities/Private/AbilitySystemComponent.cpp
@@ -650,6 +650,11 @@ FActiveGameplayEffectHandle UAbilitySystemComponent::ApplyGameplayEffectSpecToSe
 		{
 			return FActiveGameplayEffectHandle();
 		}
+
+		if (!Spec.Def->RemovalTagRequirements.IsEmpty() && Spec.Def->RemovalTagRequirements.RequirementsMet(MyTags) == true)
+		{
+			return FActiveGameplayEffectHandle();
+		}
 	}
 
 	// Custom application requirement check
@@ -780,21 +785,11 @@ FActiveGameplayEffectHandle UAbilitySystemComponent::ApplyGameplayEffectSpecToSe
 	{
 		ABILITY_LOG(Warning, TEXT("%s is periodic but also applies GameplayEffects to its target. GameplayEffects will only be applied once, not every period."), *Spec.Def->GetPathName());
 	}
-
-	// ------------------------------------------------------
-	//	Remove gameplay effects with tags
-	//		Remove any active gameplay effects that match the RemoveGameplayEffectsWithTags in the definition for this spec
-	//		Only call this if we are the Authoritative owner and we have some RemoveGameplayEffectsWithTags.CombinedTag to remove
-	// ------------------------------------------------------
-	if (bIsNetAuthority && Spec.Def->RemoveGameplayEffectsWithTags.CombinedTags.Num() > 0)
+	
+	// evaluate if any active effects need to be removed by the application of this effect
+	if (bIsNetAuthority)
 	{
-		// Clear tags is always removing all stacks.
-		FGameplayEffectQuery ClearQuery = FGameplayEffectQuery::MakeQuery_MatchAnyOwningTags(Spec.Def->RemoveGameplayEffectsWithTags.CombinedTags);
-		if (MyHandle.IsValid())
-		{
-			ClearQuery.IgnoreHandles.Add(MyHandle);
-		}
-		ActiveGameplayEffects.RemoveActiveEffects(ClearQuery, -1);
+		ActiveGameplayEffects.AttemptRemoveActiveEffectsOnEffectApplication(Spec, MyHandle);
 	}
 	
 	// ------------------------------------------------------
diff --git a/Engine/Plugins/Runtime/GameplayAbilities/Source/GameplayAbilities/Private/GameplayEffect.cpp b/Engine/Plugins/Runtime/GameplayAbilities/Source/GameplayAbilities/Private/GameplayEffect.cpp
--- a/Engine/Plugins/Runtime/GameplayAbilities/Source/GameplayAbilities/Private/GameplayEffect.cpp
+++ b/Engine/Plugins/Runtime/GameplayAbilities/Source/GameplayAbilities/Private/GameplayEffect.cpp
@@ -92,6 +92,7 @@ void UGameplayEffect::PostLoad()
 	UpdateInheritedTagProperties();
 
 	HasGrantedApplicationImmunityQuery = !GrantedApplicationImmunityQuery.IsEmpty();
+	HasRemoveGameplayEffectsQuery = !RemoveGameplayEffectQuery.IsEmpty();
 
 #if WITH_EDITOR
 	GETCURVE_REPORTERROR_WITHPOSTLOAD(Period.Curve);
@@ -156,6 +157,7 @@ void UGameplayEffect::PostEditChangeProperty(FPropertyChangedEvent& PropertyChan
 	}
 
 	HasGrantedApplicationImmunityQuery = !GrantedApplicationImmunityQuery.IsEmpty();
+	HasRemoveGameplayEffectsQuery = !RemoveGameplayEffectQuery.IsEmpty();
 
 	UAbilitySystemGlobals::Get().GameplayEffectPostEditChangeProperty(this, PropertyChangedEvent);
 }
@@ -166,6 +168,7 @@ void UGameplayEffect::PreSave(const class ITargetPlatform* TargetPlatform)
 {
 	Super::PreSave(TargetPlatform);
 	HasGrantedApplicationImmunityQuery = !GrantedApplicationImmunityQuery.IsEmpty();
+	HasRemoveGameplayEffectsQuery = !RemoveGameplayEffectQuery.IsEmpty();
 }
 
 void UGameplayEffect::UpdateInheritedTagProperties()
@@ -1711,6 +1714,18 @@ void FActiveGameplayEffect::CheckOngoingTagRequirements(const FGameplayTagContai
 	}
 }
 
+bool FActiveGameplayEffect::CheckRemovalTagRequirements(const FGameplayTagContainer& OwnerTags, FActiveGameplayEffectsContainer& OwningContainer) const
+{
+	if (!Spec.Def->RemovalTagRequirements.IsEmpty())
+	{
+		if (Spec.Def->RemovalTagRequirements.RequirementsMet(OwnerTags))
+		{
+			return true;
+		}
+	}
+	return false;
+}
+
 void FActiveGameplayEffect::PreReplicatedRemove(const struct FActiveGameplayEffectsContainer &InArray)
 {
 	if (Spec.Def == nullptr)
@@ -3028,6 +3043,17 @@ void FActiveGameplayEffectsContainer::InternalOnActiveGameplayEffectAdded(FActiv
 		ActiveEffectTagDependencies.FindOrAdd(Tag).Add(Effect.Handle);
 	}
 
+	// Add our removal tag requirements to the dependency map. We will actually check for these tags below.
+	for (const FGameplayTag& Tag : EffectDef->RemovalTagRequirements.IgnoreTags)
+	{
+		ActiveEffectTagDependencies.FindOrAdd(Tag).Add(Effect.Handle);
+	}
+
+	for (const FGameplayTag& Tag : EffectDef->RemovalTagRequirements.RequireTags)
+	{
+		ActiveEffectTagDependencies.FindOrAdd(Tag).Add(Effect.Handle);
+	}
+
 	// Add any external dependencies that might dirty the effect, if necessary
 	AddCustomMagnitudeExternalDependencies(Effect);
 
@@ -3350,6 +3376,8 @@ void FActiveGameplayEffectsContainer::InternalOnActiveGameplayEffectRemoved(FAct
 		// Remove our tag requirements from the dependency map
 		RemoveActiveEffectTagDependency(Effect.Spec.Def->OngoingTagRequirements.IgnoreTags, Effect.Handle);
 		RemoveActiveEffectTagDependency(Effect.Spec.Def->OngoingTagRequirements.RequireTags, Effect.Handle);
+		RemoveActiveEffectTagDependency(Effect.Spec.Def->RemovalTagRequirements.IgnoreTags, Effect.Handle);
+		RemoveActiveEffectTagDependency(Effect.Spec.Def->RemovalTagRequirements.RequireTags, Effect.Handle);
 
 		// Only Need to update tags and modifiers if the gameplay effect is active.
 		if (!Effect.bIsInhibited)
@@ -3642,22 +3670,36 @@ void FActiveGameplayEffectsContainer::RestartActiveGameplayEffectDuration(FActiv
 void FActiveGameplayEffectsContainer::OnOwnerTagChange(FGameplayTag TagChange, int32 NewCount)
 {
 	// It may be beneficial to do a scoped lock on attribute re-evaluation during this function
+	GAMEPLAYEFFECT_SCOPE_LOCK();
 
-	auto Ptr = ActiveEffectTagDependencies.Find(TagChange);
-	if (Ptr)
+	FGameplayTagContainer OwnerTags;
+	if (Owner)
 	{
-		GAMEPLAYEFFECT_SCOPE_LOCK();
+		Owner->GetOwnedGameplayTags(OwnerTags);
+	}
+	else
+	{
+		UE_LOG(LogGameplayEffects, Warning, TEXT("No Owner for FActiveGameplayEffectsContainer in OnOwnerTagChange"));
+	}
 
-		FGameplayTagContainer OwnerTags;
-		if (Owner)
-		{
-			Owner->GetOwnedGameplayTags(OwnerTags);
-		}
-		else
+	// if this tag has any dependencies, we need to evaluate any removals
+	if (ActiveEffectTagDependencies.Contains(TagChange))
+	{
+		for (int32 idx = GetNumGameplayEffects() - 1; idx >= 0; --idx)
 		{
-			UE_LOG(LogGameplayEffects, Warning, TEXT("No Owner for FActiveGameplayEffectsContainer in OnOwnerTagChange"));
+			const FActiveGameplayEffect& Effect = *GetActiveGameplayEffect(idx);
+			if (Effect.IsPendingRemove == false && Effect.CheckRemovalTagRequirements(OwnerTags, *this))
+			{
+				InternalRemoveActiveGameplayEffect(idx, -1, true);
+			}
 		}
+	}
 
+	// Removing effects could change and/or remove the dependency entries, so find the set after we
+	// perform any removals.
+	auto Ptr = ActiveEffectTagDependencies.Find(TagChange);
+	if (Ptr)
+	{
 		TSet<FActiveGameplayEffectHandle>& Handles = *Ptr;
 		for (const FActiveGameplayEffectHandle& Handle : Handles)
 		{
@@ -4135,6 +4177,49 @@ int32 FActiveGameplayEffectsContainer::RemoveActiveEffects(const FGameplayEffect
 	return NumRemoved;
 }
 
+void FActiveGameplayEffectsContainer::AttemptRemoveActiveEffectsOnEffectApplication(const FGameplayEffectSpec &InSpec, const FActiveGameplayEffectHandle& InHandle)
+{
+	GAMEPLAYEFFECT_SCOPE_LOCK();
+	if (InSpec.Def)
+	{
+		// Clear tags is always removing all stacks.
+		FGameplayEffectQuery ClearQuery;
+		if (InSpec.Def->RemoveGameplayEffectsWithTags.CombinedTags.Num() > 0)
+		{
+			ClearQuery = FGameplayEffectQuery::MakeQuery_MatchAnyOwningTags(InSpec.Def->RemoveGameplayEffectsWithTags.CombinedTags);
+			if (InHandle.IsValid())
+			{
+				ClearQuery.IgnoreHandles.Add(InHandle);
+			}
+		}
+
+		const FGameplayTagContainer* AggregatedSourceTags = InSpec.CapturedSourceTags.GetAggregatedTags();
+		for (int32 idx = GetNumGameplayEffects() - 1; idx >= 0; --idx)
+		{
+			const FActiveGameplayEffect& ActiveEffect = *GetActiveGameplayEffect(idx);
+			if (!ActiveEffect.IsPendingRemove)
+			{
+				const UGameplayEffect* Effect = ActiveEffect.Spec.Def;
+				// check if the effect has a tag query
+				if (Effect && Effect->HasRemoveGameplayEffectsQuery && Effect->RemoveGameplayEffectQuery.Matches(InSpec))
+				{
+					InternalRemoveActiveGameplayEffect(idx, -1, true);
+				}
+				// check if the effect has removal tag requirements
+				else if (Effect && AggregatedSourceTags && !Effect->RemovalTagRequirements.IsEmpty() && Effect->RemovalTagRequirements.RequirementsMet(*AggregatedSourceTags))
+				{
+					InternalRemoveActiveGameplayEffect(idx, -1, true);
+				}
+				// clear query check
+				else if (!ClearQuery.IsEmpty() && ClearQuery.Matches(ActiveEffect))
+				{
+					InternalRemoveActiveGameplayEffect(idx, -1, true);
+				}
+			}
+		}
+	}
+}
+
 int32 FActiveGameplayEffectsContainer::GetActiveEffectCount(const FGameplayEffectQuery& Query, bool bEnforceOnGoingCheck) const
 {
 	int32 Count = 0;
diff --git a/Engine/Plugins/Runtime/GameplayAbilities/Source/GameplayAbilities/Public/GameplayEffect.h b/Engine/Plugins/Runtime/GameplayAbilities/Source/GameplayAbilities/Public/GameplayEffect.h
--- a/Engine/Plugins/Runtime/GameplayAbilities/Source/GameplayAbilities/Public/GameplayEffect.h
+++ b/Engine/Plugins/Runtime/GameplayAbilities/Source/GameplayAbilities/Public/GameplayEffect.h
@@ -1269,6 +1269,9 @@ struct GAMEPLAYABILITIES_API FActiveGameplayEffect : public FFastArraySerializer
 	}
 
 	void CheckOngoingTagRequirements(const FGameplayTagContainer& OwnerTags, struct FActiveGameplayEffectsContainer& OwningContainer, bool bInvokeGameplayCueEvents = false);
+	
+	/** Method to check if this effect should remove because the owner tags pass the RemovalTagRequirements requirement check */
+	bool CheckRemovalTagRequirements(const FGameplayTagContainer& OwnerTags, struct FActiveGameplayEffectsContainer& OwningContainer) const;
 
 	void PrintAll() const;
 
@@ -1655,6 +1658,9 @@ struct GAMEPLAYABILITIES_API FActiveGameplayEffectsContainer : public FFastArray
 
 	int32 RemoveActiveEffects(const FGameplayEffectQuery& Query, int32 StacksToRemove);
 
+	/** Method called during effect application to process if any active effects should be removed from this effects application */
+	void AttemptRemoveActiveEffectsOnEffectApplication(const FGameplayEffectSpec &InSpec, const FActiveGameplayEffectHandle& InHandle);
+
 	/**
 	 * Get the count of the effects matching the specified query (including stack count)
 	 * 
@@ -2019,6 +2025,10 @@ public:
 	UPROPERTY(EditDefaultsOnly, BlueprintReadOnly, Category = Tags, meta=(Categories="ApplicationTagRequirementsCategory"))
 	FGameplayTagRequirements ApplicationTagRequirements;
 
+	/** Tag requirements that if met will remove this effect. Also prevents effect application. */
+	UPROPERTY(EditDefaultsOnly, BlueprintReadOnly, Category = Tags, meta=(Categories="ApplicationTagRequirementsCategory"))
+	FGameplayTagRequirements RemovalTagRequirements;
+
 	/** GameplayEffects that *have* tags in this container will be cleared upon effect application. */
 	UPROPERTY(EditDefaultsOnly, BlueprintReadOnly, Category = Tags, meta=(Categories="RemoveTagRequirementsCategory"))
 	FInheritedTagContainer RemoveGameplayEffectsWithTags;
@@ -2032,7 +2042,14 @@ public:
 	FGameplayEffectQuery GrantedApplicationImmunityQuery;
 
 	/** Cached !GrantedApplicationImmunityQuery.IsEmpty(). Set on PostLoad. */
-	bool HasGrantedApplicationImmunityQuery;
+	bool HasGrantedApplicationImmunityQuery = false;
+
+	/** On Application of an effect, any active effects with this this query that matches against the added effect will be removed. Queries are more powerful but slightly slower than RemoveGameplayEffectsWithTags. */
+	UPROPERTY(EditDefaultsOnly, BlueprintReadOnly, Category = Tags, meta = (DisplayAfter = "RemovalTagRequirements"))
+	FGameplayEffectQuery RemoveGameplayEffectQuery;
+
+	/** Cached !RemoveGameplayEffectsQuery.IsEmpty(). Set on PostLoad. */
+	bool HasRemoveGameplayEffectsQuery = false;
 
 	// ----------------------------------------------------------------------
 	//	Stacking
