Added the ability to query an ASC for all active gameplay effects that have a specified set of tags.

diff --git a/Engine/Plugins/Runtime/GameplayAbilities/Source/GameplayAbilities/Private/AbilitySystemComponent.cpp b/Engine/Plugins/Runtime/GameplayAbilities/Source/GameplayAbilities/Private/AbilitySystemComponent.cpp
--- a/Engine/Plugins/Runtime/GameplayAbilities/Source/GameplayAbilities/Private/AbilitySystemComponent.cpp
+++ b/Engine/Plugins/Runtime/GameplayAbilities/Source/GameplayAbilities/Private/AbilitySystemComponent.cpp
@@ -1356,6 +1356,11 @@ TArray<FActiveGameplayEffectHandle> UAbilitySystemComponent::GetActiveEffects(co
 	return ActiveGameplayEffects.GetActiveEffects(Query);
 }
 
+TArray<FActiveGameplayEffectHandle> UAbilitySystemComponent::GetActiveEffectsWithAllTags(FGameplayTagContainer Tags) const
+{
+	return GetActiveEffects(FGameplayEffectQuery::MakeQuery_MatchAllEffectTags(Tags));
+}
+
 float UAbilitySystemComponent::GetActiveEffectsEndTime(const FGameplayEffectQuery& Query) const
 {
 	TArray<AActor*> DummyInstigators;
diff --git a/Engine/Plugins/Runtime/GameplayAbilities/Source/GameplayAbilities/Public/AbilitySystemComponent.h b/Engine/Plugins/Runtime/GameplayAbilities/Source/GameplayAbilities/Public/AbilitySystemComponent.h
--- a/Engine/Plugins/Runtime/GameplayAbilities/Source/GameplayAbilities/Public/AbilitySystemComponent.h
+++ b/Engine/Plugins/Runtime/GameplayAbilities/Source/GameplayAbilities/Public/AbilitySystemComponent.h
@@ -713,6 +713,10 @@ class GAMEPLAYABILITIES_API UAbilitySystemComponent : public UGameplayTasksCompo
 	UFUNCTION(BlueprintCallable, BlueprintPure=false, Category = "GameplayEffects", meta=(DisplayName = "Get Activate Gameplay Effects for Query"))
 	TArray<FActiveGameplayEffectHandle> GetActiveEffects(const FGameplayEffectQuery& Query) const;
 
+	/** Returns list of active effects that have all of the passed in tags */
+	UFUNCTION(BlueprintCallable, BlueprintPure = false, Category = "GameplayEffects")
+	TArray<FActiveGameplayEffectHandle> GetActiveEffectsWithAllTags(FGameplayTagContainer Tags) const;
+
 	/** This will give the world time that all effects matching this query will be finished. If multiple effects match, it returns the one that returns last */
 	float GetActiveEffectsEndTime(const FGameplayEffectQuery& Query) const;
 	float GetActiveEffectsEndTimeWithInstigators(const FGameplayEffectQuery& Query, TArray<AActor*>& Instigators) const;
