Update gameplay ability policy enums in preparation for future changes

diff --git a/Engine/Plugins/Runtime/GameplayAbilities/Source/GameplayAbilities/Public/Abilities/GameplayAbility.h b/Engine/Plugins/Runtime/GameplayAbilities/Source/GameplayAbilities/Public/Abilities/GameplayAbility.h
--- a/Engine/Plugins/Runtime/GameplayAbilities/Source/GameplayAbilities/Public/Abilities/GameplayAbility.h
+++ b/Engine/Plugins/Runtime/GameplayAbilities/Source/GameplayAbilities/Public/Abilities/GameplayAbility.h
@@ -138,12 +138,18 @@ public:
 		return ReplicationPolicy;
 	}
 
-	/** How does an ability execute on the network. Does a client "ask and predict", "ask and wait", "don't ask (just do it)" */
+	/** Where does an ability execute on the network? Does a client "ask and predict", "ask and wait", "don't ask (just do it)" */
 	EGameplayAbilityNetExecutionPolicy::Type GetNetExecutionPolicy() const
 	{
 		return NetExecutionPolicy;
 	}
 
+	/** Where should an ability execute on the network? Provides protection from clients attempting to execute restricted abilities. */
+	EGameplayAbilityNetSecurityPolicy::Type GetNetSecurityPolicy() const
+	{
+		return NetSecurityPolicy;
+	}
+
 	/** Returns the actor info associated with this ability, has cached pointers to useful objects */
 	UFUNCTION(BlueprintCallable, Category=Ability)
 	FGameplayAbilityActorInfo GetActorInfo() const;
@@ -672,6 +678,9 @@ protected:
 	UPROPERTY(EditDefaultsOnly, Category=Advanced)
 	TEnumAsByte<EGameplayAbilityNetExecutionPolicy::Type> NetExecutionPolicy;
 
+	UPROPERTY(EditDefaultsOnly, Category = Advanced)
+	TEnumAsByte<EGameplayAbilityNetSecurityPolicy::Type> NetSecurityPolicy;
+
 	/** This GameplayEffect represents the cost (mana, stamina, etc) of the ability. It will be applied when the ability is committed. */
 	UPROPERTY(EditDefaultsOnly, Category = Costs)
 	TSubclassOf<class UGameplayEffect> CostGameplayEffectClass;
diff --git a/Engine/Plugins/Runtime/GameplayAbilities/Source/GameplayAbilities/Public/Abilities/GameplayAbilityTypes.h b/Engine/Plugins/Runtime/GameplayAbilities/Source/GameplayAbilities/Public/Abilities/GameplayAbilityTypes.h
--- a/Engine/Plugins/Runtime/GameplayAbilities/Source/GameplayAbilities/Public/Abilities/GameplayAbilityTypes.h
+++ b/Engine/Plugins/Runtime/GameplayAbilities/Source/GameplayAbilities/Public/Abilities/GameplayAbilityTypes.h
@@ -53,7 +53,7 @@ namespace EGameplayAbilityInstancingPolicy
 UENUM(BlueprintType)
 namespace EGameplayAbilityNetExecutionPolicy
 {
-	/** How does an ability execute on the network. Does a client "ask and predict", "ask and wait", "don't ask (just do it)" */
+	/** Where does an ability execute on the network. Does a client "ask and predict", "ask and wait", "don't ask (just do it)" */
 	enum Type
 	{
 		// Part of this ability runs predictively on the local client if there is one
@@ -70,6 +70,26 @@ namespace EGameplayAbilityNetExecutionPolicy
 	};
 }
 
+UENUM(BlueprintType)
+namespace EGameplayAbilityNetSecurityPolicy
+{
+	/** What protections does this ability have? Should the client be allowed to request changes to the execution of the ability? */
+	enum Type
+	{
+		// No security requirements. Client or server can trigger execution and termination of this ability freely.
+		ClientOrServer			UMETA(DisplayName = "Client Or Server"),
+
+		// A client requesting execution of this ability will be ignored by the server. Clients can still request that the server cancel or end this ability.
+		ServerOnlyExecution		UMETA(DisplayName = "Server Only Execution"),
+
+		// A client requesting cancellation or ending of this ability will be ignored by the server. Clients can still request execution of the ability.
+		ServerOnlyTermination	UMETA(DisplayName = "Server Only Termination"),
+
+		// Server controls both execution and termination of this ability. A client making any requests will be ignored.
+		ServerOnly				UMETA(DisplayName = "Server Only"),
+	};
+}
+
 UENUM(BlueprintType)
 namespace EGameplayAbilityReplicationPolicy
 {
