Dev Only (AbilitySystemComponent);

Adding tag count to the blocked tags when displayed for debug.

-- Code --
AbilitySystemComponent:
Debug_Internal() - Getting and dsiplaying the tag count for blocked tags next to them when we draw the debug text.

#tests PIE, do dev cheat ShowDebug and CycleDebugAbilitySystem to see owned and blocked tags.

diff --git a/Engine/Plugins/Runtime/GameplayAbilities/Source/GameplayAbilities/Private/AbilitySystemComponent.cpp b/Engine/Plugins/Runtime/GameplayAbilities/Source/GameplayAbilities/Private/AbilitySystemComponent.cpp
--- a/Engine/Plugins/Runtime/GameplayAbilities/Source/GameplayAbilities/Private/AbilitySystemComponent.cpp
+++ b/Engine/Plugins/Runtime/GameplayAbilities/Source/GameplayAbilities/Private/AbilitySystemComponent.cpp
@@ -2193,7 +2193,18 @@ void UAbilitySystemComponent::Debug_Internal(FAbilitySystemComponentDebugInfo& I
 
 	if (BlockedAbilityTags.GetExplicitGameplayTags().Num() > 0)
 	{
-		DebugLine(Info, FString::Printf(TEXT("BlockedAbilityTags: %s"), *BlockedAbilityTags.GetExplicitGameplayTags().ToStringSimple()), 4.f, 0.f);
+		FString BlockedTagsStrings;
+		int32 BlockedTagCount = 1;
+		for (FGameplayTag Tag : BlockedAbilityTags.GetExplicitGameplayTags())
+		{
+			BlockedTagsStrings.Append(FString::Printf(TEXT("%s (%d)"), *Tag.ToString(), BlockedAbilityTags.GetTagCount(Tag)));
+
+			if (BlockedTagCount++ < NumTags)
+			{
+				BlockedTagsStrings += TEXT(", ");
+			}
+		}
+		DebugLine(Info, FString::Printf(TEXT("BlockedAbilityTags: %s"), *BlockedTagsStrings), 4.f, 0.f);
 	}
 	else
 	{
