Fixes for potential security issues. These changes cut off a couple pathways identified in a security audit that a malicious attacker could use to crash a server via the ability system.

diff --git a/Engine/Plugins/Runtime/GameplayAbilities/Source/GameplayAbilities/Private/AbilitySystemComponent_Abilities.cpp b/Engine/Plugins/Runtime/GameplayAbilities/Source/GameplayAbilities/Private/AbilitySystemComponent_Abilities.cpp
--- a/Engine/Plugins/Runtime/GameplayAbilities/Source/GameplayAbilities/Private/AbilitySystemComponent_Abilities.cpp
+++ b/Engine/Plugins/Runtime/GameplayAbilities/Source/GameplayAbilities/Private/AbilitySystemComponent_Abilities.cpp
@@ -3037,6 +3037,15 @@ void UAbilitySystemComponent::ServerSetReplicatedTargetData_Implementation(FGame
 
 bool UAbilitySystemComponent::ServerSetReplicatedTargetData_Validate(FGameplayAbilitySpecHandle AbilityHandle, FPredictionKey AbilityOriginalPredictionKey, const FGameplayAbilityTargetDataHandle& ReplicatedTargetDataHandle, FGameplayTag ApplicationTag, FPredictionKey CurrentPredictionKey)
 {
+	// check the data coming from the client to ensure it's valid
+	for (const TSharedPtr<FGameplayAbilityTargetData>& TgtData : ReplicatedTargetDataHandle.Data)
+	{
+		if (!ensure(TgtData.IsValid()))
+		{
+			return false;
+		}
+	}
+
 	return true;
 }
 
diff --git a/Engine/Plugins/Runtime/GameplayAbilities/Source/GameplayAbilities/Private/GameplayEffect.cpp b/Engine/Plugins/Runtime/GameplayAbilities/Source/GameplayAbilities/Private/GameplayEffect.cpp
--- a/Engine/Plugins/Runtime/GameplayAbilities/Source/GameplayAbilities/Private/GameplayEffect.cpp
+++ b/Engine/Plugins/Runtime/GameplayAbilities/Source/GameplayAbilities/Private/GameplayEffect.cpp
@@ -1906,7 +1906,7 @@ FActiveGameplayEffectsContainer::~FActiveGameplayEffectsContainer()
 
 void FActiveGameplayEffectsContainer::RegisterWithOwner(UAbilitySystemComponent* InOwner)
 {
-	if (Owner != InOwner)
+	if (Owner != InOwner && InOwner != nullptr)
 	{
 		Owner = InOwner;
 		OwnerIsNetAuthority = Owner->IsOwnerActorAuthoritative();
@@ -1924,6 +1924,11 @@ void FActiveGameplayEffectsContainer::ExecuteActiveEffectsFrom(FGameplayEffectSp
 	SCOPE_CYCLE_COUNTER(STAT_ExecuteActiveEffectsFrom);
 #endif
 
+	if (!Owner)
+	{
+		return;
+	}
+
 	FGameplayEffectSpec& SpecToUse = Spec;
 
 	// Capture our own tags.
