In UGameplayCueManager::OnMissingCueAsyncLoadComplete, if TargetActor has since been destroyed, still attempt to call HandleGameplayCue with the remaining data
#rnx

diff --git a/Engine/Plugins/Runtime/GameplayAbilities/Source/GameplayAbilities/Private/GameplayCueInterface.cpp b/Engine/Plugins/Runtime/GameplayAbilities/Source/GameplayAbilities/Private/GameplayCueInterface.cpp
--- a/Engine/Plugins/Runtime/GameplayAbilities/Source/GameplayAbilities/Private/GameplayCueInterface.cpp
+++ b/Engine/Plugins/Runtime/GameplayAbilities/Source/GameplayAbilities/Private/GameplayCueInterface.cpp
@@ -56,6 +56,11 @@ void IGameplayCueInterface::HandleGameplayCue(AActor *Self, FGameplayTag Gamepla
 {
 	SCOPE_CYCLE_COUNTER(STAT_GameplayCueInterface_HandleGameplayCue);
 
+	if (!Self)
+	{
+		return;
+	}
+
 	// Look up a custom function for this gameplay tag. 
 	UClass* Class = Self->GetClass();
 	FGameplayTagContainer TagAndParentsContainer = GameplayCueTag.GetGameplayTagParents();
diff --git a/Engine/Plugins/Runtime/GameplayAbilities/Source/GameplayAbilities/Private/GameplayCueManager.cpp b/Engine/Plugins/Runtime/GameplayAbilities/Source/GameplayAbilities/Private/GameplayCueManager.cpp
--- a/Engine/Plugins/Runtime/GameplayAbilities/Source/GameplayAbilities/Private/GameplayCueManager.cpp
+++ b/Engine/Plugins/Runtime/GameplayAbilities/Source/GameplayAbilities/Private/GameplayCueManager.cpp
@@ -313,9 +313,20 @@ void UGameplayCueManager::OnMissingCueAsyncLoadComplete(FSoftObjectPath LoadedOb
 		return;
 	}
 
-	if (OwningSet.IsValid() && TargetActor.IsValid())
+	if (OwningSet.IsValid())
 	{
-		CurrentWorld = TargetActor->GetWorld();
+		CurrentWorld = TargetActor.IsValid() ? TargetActor->GetWorld() : nullptr;
+		if (!CurrentWorld)
+		{
+			// TargetActor has since been destroyed.  Attempt to get the world from the other actors.
+			const AActor* CueInstigator = Parameters.GetInstigator();
+			CurrentWorld = CueInstigator ? CueInstigator->GetWorld() : nullptr;
+			if (!CurrentWorld)
+			{
+				const AActor* EffectCauser = Parameters.GetEffectCauser();
+				CurrentWorld = EffectCauser ? EffectCauser->GetWorld() : nullptr;
+			}
+		}
 
 		// Don't handle gameplay cues when world is tearing down
 		if (!GetWorld() || GetWorld()->bIsTearingDown)
diff --git a/Engine/Plugins/Runtime/GameplayAbilities/Source/GameplayAbilities/Private/GameplayCueNotify_HitImpact.cpp b/Engine/Plugins/Runtime/GameplayAbilities/Source/GameplayAbilities/Private/GameplayCueNotify_HitImpact.cpp
--- a/Engine/Plugins/Runtime/GameplayAbilities/Source/GameplayAbilities/Private/GameplayCueNotify_HitImpact.cpp
+++ b/Engine/Plugins/Runtime/GameplayAbilities/Source/GameplayAbilities/Private/GameplayCueNotify_HitImpact.cpp
@@ -2,6 +2,8 @@
 
 #include "GameplayCueNotify_HitImpact.h"
 #include "Kismet/GameplayStatics.h"
+#include "GameplayCueManager.h"
+#include "AbilitySystemGlobals.h"
 
 UGameplayCueNotify_HitImpact::UGameplayCueNotify_HitImpact(const FObjectInitializer& PCIP)
 : Super(PCIP)
@@ -17,21 +19,25 @@ bool UGameplayCueNotify_HitImpact::HandlesEvent(EGameplayCueEvent::Type EventTyp
 void UGameplayCueNotify_HitImpact::HandleGameplayCue(AActor* Self, EGameplayCueEvent::Type EventType, const FGameplayCueParameters& Parameters)
 {
 	check(EventType == EGameplayCueEvent::Executed);
-	check(Self);
 	
-	const FHitResult* HitResult = Parameters.EffectContext.GetHitResult();
-	if (HitResult)
+	const UObject* WorldContextObject = Self;
+	if (!WorldContextObject)
 	{
-		if (ParticleSystem)
-		{
-			UGameplayStatics::SpawnEmitterAtLocation(Self, ParticleSystem, HitResult->ImpactPoint, HitResult->ImpactNormal.Rotation(), true);
-		}
+		WorldContextObject = UAbilitySystemGlobals::Get().GetGameplayCueManager();
 	}
-	else
+
+	if (ParticleSystem && WorldContextObject)
 	{
-		if (ParticleSystem)
+		const FHitResult* HitResult = Parameters.EffectContext.GetHitResult();
+		if (HitResult)
+		{
+			UGameplayStatics::SpawnEmitterAtLocation(WorldContextObject, ParticleSystem, HitResult->ImpactPoint, HitResult->ImpactNormal.Rotation(), true);
+		}
+		else
 		{
-			UGameplayStatics::SpawnEmitterAtLocation(Self, ParticleSystem, Self->GetActorLocation(), Self->GetActorRotation(), true);
+			const FVector Location = Self ? Self->GetActorLocation() : Parameters.Location;
+			const FRotator Rotation = Self ? Self->GetActorRotation() : Parameters.Normal.Rotation();
+			UGameplayStatics::SpawnEmitterAtLocation(WorldContextObject, ParticleSystem, Location, Rotation, true);
 		}
 	}
 }
