Attempting to fix compile error in HandleGameplayCue
#rnx

diff --git a/Engine/Plugins/Runtime/GameplayAbilities/Source/GameplayAbilities/Private/GameplayCueNotify_HitImpact.cpp b/Engine/Plugins/Runtime/GameplayAbilities/Source/GameplayAbilities/Private/GameplayCueNotify_HitImpact.cpp
--- a/Engine/Plugins/Runtime/GameplayAbilities/Source/GameplayAbilities/Private/GameplayCueNotify_HitImpact.cpp
+++ b/Engine/Plugins/Runtime/GameplayAbilities/Source/GameplayAbilities/Private/GameplayCueNotify_HitImpact.cpp
@@ -35,8 +35,18 @@ void UGameplayCueNotify_HitImpact::HandleGameplayCue(AActor* Self, EGameplayCueE
 		}
 		else
 		{
-			const FVector Location = Self ? Self->GetActorLocation() : Parameters.Location;
-			const FRotator Rotation = Self ? Self->GetActorRotation() : Parameters.Normal.Rotation();
+			FVector Location = FVector::ZeroVector;
+			FRotator Rotation = FRotator::ZeroRotator;
+			if (Self)
+			{
+				Location = Self->GetActorLocation();
+				Rotation = Self->GetActorRotation();
+			}
+			else
+			{
+				Location = Parameters.Location;
+				Rotation = Parameters.Normal.Rotation();
+			}
 			UGameplayStatics::SpawnEmitterAtLocation(WorldContextObject, ParticleSystem, Location, Rotation, true);
 		}
 	}
