Checking if the spec.def is valid when attempting to apply a gameplay effect spec (also adding some additional logging to existing failure cases).

#rnx

diff --git a/Engine/Plugins/Runtime/GameplayAbilities/Source/GameplayAbilities/Private/GameplayEffect.cpp b/Engine/Plugins/Runtime/GameplayAbilities/Source/GameplayAbilities/Private/GameplayEffect.cpp
--- a/Engine/Plugins/Runtime/GameplayAbilities/Source/GameplayAbilities/Private/GameplayEffect.cpp
+++ b/Engine/Plugins/Runtime/GameplayAbilities/Source/GameplayAbilities/Private/GameplayEffect.cpp
@@ -1757,7 +1757,7 @@ void FActiveGameplayEffect::PostReplicatedAdd(const struct FActiveGameplayEffect
 {
 	if (Spec.Def == nullptr)
 	{
-		ABILITY_LOG(Error, TEXT("Received ReplicatedGameplayEffect with no UGameplayEffect def."));
+		ABILITY_LOG(Error, TEXT("FActiveGameplayEffect::PostReplicatedAdd Received ReplicatedGameplayEffect with no UGameplayEffect def. (%s)"), *Spec.GetEffectContext().ToString());
 		return;
 	}
 
@@ -1826,7 +1826,7 @@ void FActiveGameplayEffect::PostReplicatedChange(const struct FActiveGameplayEff
 {
 	if (Spec.Def == nullptr)
 	{
-		ABILITY_LOG(Error, TEXT("Received ReplicatedGameplayEffect with no UGameplayEffect def."));
+		ABILITY_LOG(Error, TEXT("FActiveGameplayEffect::PostReplicatedChange Received ReplicatedGameplayEffect with no UGameplayEffect def. (%s)"), *Spec.GetEffectContext().ToString());
 		return;
 	}
 
@@ -2734,6 +2734,11 @@ FActiveGameplayEffect* FActiveGameplayEffectsContainer::ApplyGameplayEffectSpec(
 
 	GAMEPLAYEFFECT_SCOPE_LOCK();
 
+	if (!ensureMsgf(Spec.Def, TEXT("Tried to apply GE with no Def (context == %s)"), *Spec.GetEffectContext().ToString()))
+	{
+		return nullptr;
+	}
+
 	bFoundExistingStackableGE = false;
 
 	if (Owner && Owner->OwnerActor && IsNetAuthority())
@@ -3692,17 +3697,22 @@ void FActiveGameplayEffectsContainer::OnOwnerTagChange(FGameplayTag TagChange, i
 	{
 		for (int32 idx = GetNumGameplayEffects() - 1; idx >= 0; --idx)
 		{
-			const FActiveGameplayEffect& Effect = *GetActiveGameplayEffect(idx);
-			if (Effect.IsPendingRemove == false && Effect.CheckRemovalTagRequirements(OwnerTags, *this))
+			const FActiveGameplayEffect* Effect = GetActiveGameplayEffect(idx);
+			if (ensureMsgf(Effect != nullptr && Effect->Spec.Def != nullptr, TEXT("GetActiveGameplayEffect(%i) returned %p. GetNumGameplayEffects is %i."), idx, Effect, GetNumGameplayEffects()))
 			{
-				InternalRemoveActiveGameplayEffect(idx, -1, true);
+				if (Effect->IsPendingRemove == false && Effect->CheckRemovalTagRequirements(OwnerTags, *this))
+				{
+					InternalRemoveActiveGameplayEffect(idx, -1, true);
+				}
 			}
+
+			ensure(Effect == GetActiveGameplayEffect(idx));
 		}
 	}
 
 	// Removing effects could change and/or remove the dependency entries, so find the set after we
 	// perform any removals.
-	auto Ptr = ActiveEffectTagDependencies.Find(TagChange);
+	TSet<FActiveGameplayEffectHandle>* Ptr = ActiveEffectTagDependencies.Find(TagChange);
 	if (Ptr)
 	{
 		TSet<FActiveGameplayEffectHandle>& Handles = *Ptr;
