Don't call GetInstancedCueActor with a null TargetActor
#rnx

diff --git a/Engine/Plugins/Runtime/GameplayAbilities/Source/GameplayAbilities/Private/GameplayCueSet.cpp b/Engine/Plugins/Runtime/GameplayAbilities/Source/GameplayAbilities/Private/GameplayCueSet.cpp
--- a/Engine/Plugins/Runtime/GameplayAbilities/Source/GameplayAbilities/Private/GameplayCueSet.cpp
+++ b/Engine/Plugins/Runtime/GameplayAbilities/Source/GameplayAbilities/Private/GameplayCueSet.cpp
@@ -264,15 +264,18 @@ bool UGameplayCueSet::HandleGameplayCueNotify_Internal(AActor* TargetActor, int3
 		{
 			if (InstancedCue->HandlesEvent(EventType))
 			{
-				//Get our instance. We should probably have a flag or something to determine if we want to reuse or stack instances. That would mean changing our map to have a list of active instances.
-				AGameplayCueNotify_Actor* SpawnedInstancedCue = CueManager->GetInstancedCueActor(TargetActor, CueData.LoadedGameplayCueClass, Parameters);
-				if (ensure(SpawnedInstancedCue))
+				if (TargetActor)
 				{
-					SpawnedInstancedCue->HandleGameplayCue(TargetActor, EventType, Parameters);
-					bReturnVal = true;
-					if (!SpawnedInstancedCue->IsOverride)
+					//Get our instance. We should probably have a flag or something to determine if we want to reuse or stack instances. That would mean changing our map to have a list of active instances.
+					AGameplayCueNotify_Actor* SpawnedInstancedCue = CueManager->GetInstancedCueActor(TargetActor, CueData.LoadedGameplayCueClass, Parameters);
+					if (ensure(SpawnedInstancedCue))
 					{
-						HandleGameplayCueNotify_Internal(TargetActor, CueData.ParentDataIdx, EventType, Parameters);
+						SpawnedInstancedCue->HandleGameplayCue(TargetActor, EventType, Parameters);
+						bReturnVal = true;
+						if (!SpawnedInstancedCue->IsOverride)
+						{
+							HandleGameplayCueNotify_Internal(TargetActor, CueData.ParentDataIdx, EventType, Parameters);
+						}
 					}
 				}
 			}
