Added function OnRemoveAbility() to UGameplayAbility.  It follows the same pattern as OnGiveAbility() and is only called on the primary instance of the ability or the class default object.

diff --git a/Engine/Plugins/Runtime/GameplayAbilities/Source/GameplayAbilities/Private/AbilitySystemComponent_Abilities.cpp b/Engine/Plugins/Runtime/GameplayAbilities/Source/GameplayAbilities/Private/AbilitySystemComponent_Abilities.cpp
--- a/Engine/Plugins/Runtime/GameplayAbilities/Source/GameplayAbilities/Private/AbilitySystemComponent_Abilities.cpp
+++ b/Engine/Plugins/Runtime/GameplayAbilities/Source/GameplayAbilities/Private/AbilitySystemComponent_Abilities.cpp
@@ -481,6 +481,18 @@ void UAbilitySystemComponent::OnRemoveAbility(FGameplayAbilitySpec& Spec)
 			}
 		}
 	}
+
+	// Notify the ability that it has been removed.  It follows the same pattern as OnGiveAbility() and is only called on the primary instance of the ability or the CDO.
+	UGameplayAbility* PrimaryInstance = Spec.GetPrimaryInstance();
+	if (PrimaryInstance)
+	{
+		PrimaryInstance->OnRemoveAbility(AbilityActorInfo.Get(), Spec);
+	}
+	else
+	{
+		Spec.Ability->OnRemoveAbility(AbilityActorInfo.Get(), Spec);
+	}
+
 	Spec.ReplicatedInstances.Empty();
 	Spec.NonReplicatedInstances.Empty();
 }
diff --git a/Engine/Plugins/Runtime/GameplayAbilities/Source/GameplayAbilities/Public/Abilities/GameplayAbility.h b/Engine/Plugins/Runtime/GameplayAbilities/Source/GameplayAbilities/Public/Abilities/GameplayAbility.h
--- a/Engine/Plugins/Runtime/GameplayAbilities/Source/GameplayAbilities/Public/Abilities/GameplayAbility.h
+++ b/Engine/Plugins/Runtime/GameplayAbilities/Source/GameplayAbilities/Public/Abilities/GameplayAbility.h
@@ -438,6 +438,9 @@ public:
 	/** Called when the ability is given to an AbilitySystemComponent */
 	virtual void OnGiveAbility(const FGameplayAbilityActorInfo* ActorInfo, const FGameplayAbilitySpec& Spec);
 
+	/** Called when the ability is removed from an AbilitySystemComponent */
+	virtual void OnRemoveAbility(const FGameplayAbilityActorInfo* ActorInfo, const FGameplayAbilitySpec& Spec) {}
+
 	/** Called when the avatar actor is set/changes */
 	virtual void OnAvatarSet(const FGameplayAbilityActorInfo* ActorInfo, const FGameplayAbilitySpec& Spec);
 
