#UE4 - Added ABILITYLIST_SCOPE_LOCK() to UAbilitySystemComponent::MonitoredTagChanged(), to fix potential issues with ability spec being invalidated by InternalTryActivateAbility() while looping throuhg triggered abilities

diff --git a/Engine/Plugins/Runtime/GameplayAbilities/Source/GameplayAbilities/Private/AbilitySystemComponent_Abilities.cpp b/Engine/Plugins/Runtime/GameplayAbilities/Source/GameplayAbilities/Private/AbilitySystemComponent_Abilities.cpp
--- a/Engine/Plugins/Runtime/GameplayAbilities/Source/GameplayAbilities/Private/AbilitySystemComponent_Abilities.cpp
+++ b/Engine/Plugins/Runtime/GameplayAbilities/Source/GameplayAbilities/Private/AbilitySystemComponent_Abilities.cpp
@@ -1963,6 +1963,8 @@ void UAbilitySystemComponent::RemoveGameplayEventTagContainerDelegate(const FGam
 
 void UAbilitySystemComponent::MonitoredTagChanged(const FGameplayTag Tag, int32 NewCount)
 {
+	ABILITYLIST_SCOPE_LOCK();
+
 	int32 TriggeredCount = 0;
 	if (OwnedTagTriggeredAbilities.Contains(Tag))
 	{
