Fix shadow variable build error

#rnx

diff --git a/Engine/Plugins/Runtime/GameplayAbilities/Source/GameplayAbilities/Private/AbilitySystemComponent_Abilities.cpp b/Engine/Plugins/Runtime/GameplayAbilities/Source/GameplayAbilities/Private/AbilitySystemComponent_Abilities.cpp
--- a/Engine/Plugins/Runtime/GameplayAbilities/Source/GameplayAbilities/Private/AbilitySystemComponent_Abilities.cpp
+++ b/Engine/Plugins/Runtime/GameplayAbilities/Source/GameplayAbilities/Private/AbilitySystemComponent_Abilities.cpp
@@ -1355,10 +1355,9 @@ bool UAbilitySystemComponent::InternalTryActivateAbility(FGameplayAbilitySpecHan
 	{
 		// Flush server moves that occurred before this ability activation so that the server receives the RPCs in the correct order
 		// Necessary to prevent abilities that trigger animation root motion or impact movement from causing network corrections
-		AActor* const AvatarActor = ActorInfo->AvatarActor.Get();
-		if (AvatarActor && !ActorInfo->IsNetAuthority())
+		if (!ActorInfo->IsNetAuthority())
 		{
-			ACharacter* AvatarCharacter = Cast<ACharacter>(AvatarActor);
+			ACharacter* AvatarCharacter = Cast<ACharacter>(ActorInfo->AvatarActor.Get());
 			if (AvatarCharacter)
 			{
 				UCharacterMovementComponent* AvatarCharMoveComp = Cast<UCharacterMovementComponent>(AvatarCharacter->GetMovementComponent());
