#UE4 - Cache Spec->Ability->AbilityTriggers array before looping in UAbilitySystemComponent::MonitoredTagChanged(), to fix potential issues with ability spec being invalidated by InternalTryActivateAbility() while looping throuhg triggered abilities

diff --git a/Engine/Plugins/Runtime/GameplayAbilities/Source/GameplayAbilities/Private/AbilitySystemComponent_Abilities.cpp b/Engine/Plugins/Runtime/GameplayAbilities/Source/GameplayAbilities/Private/AbilitySystemComponent_Abilities.cpp
--- a/Engine/Plugins/Runtime/GameplayAbilities/Source/GameplayAbilities/Private/AbilitySystemComponent_Abilities.cpp
+++ b/Engine/Plugins/Runtime/GameplayAbilities/Source/GameplayAbilities/Private/AbilitySystemComponent_Abilities.cpp
@@ -1996,29 +1996,33 @@ void UAbilitySystemComponent::MonitoredTagChanged(const FGameplayTag Tag, int32
 				return;
 			}
 
-			for (const FAbilityTriggerData& TriggerData : Spec->Ability->AbilityTriggers)
+			if (Spec->Ability)
 			{
-				FGameplayTag EventTag = TriggerData.TriggerTag;
-
-				if (EventTag == Tag)
+				TArray<FAbilityTriggerData> AbilityTriggers = Spec->Ability->AbilityTriggers;
+				for (const FAbilityTriggerData& TriggerData : AbilityTriggers)
 				{
-					if (NewCount > 0)
-					{
-						// Populate event data so this will use the same blueprint node to activate as gameplay triggers
-						FGameplayEventData EventData;
-						EventData.EventMagnitude = NewCount;
-						EventData.EventTag = EventTag;
-						EventData.Instigator = OwnerActor;
-						EventData.Target = OwnerActor;
-						// Try to activate it
-						InternalTryActivateAbility(Spec->Handle, FPredictionKey(), nullptr, nullptr, &EventData);
-
-						// TODO: Check client/server type
-					}
-					else if (NewCount == 0 && TriggerData.TriggerSource == EGameplayAbilityTriggerSource::OwnedTagPresent)
+					FGameplayTag EventTag = TriggerData.TriggerTag;
+
+					if (EventTag == Tag)
 					{
-						// Try to cancel, but only if the type is right
-						CancelAbilitySpec(*Spec, nullptr);
+						if (NewCount > 0)
+						{
+							// Populate event data so this will use the same blueprint node to activate as gameplay triggers
+							FGameplayEventData EventData;
+							EventData.EventMagnitude = NewCount;
+							EventData.EventTag = EventTag;
+							EventData.Instigator = OwnerActor;
+							EventData.Target = OwnerActor;
+							// Try to activate it
+							InternalTryActivateAbility(Spec->Handle, FPredictionKey(), nullptr, nullptr, &EventData);
+
+							// TODO: Check client/server type
+						}
+						else if (NewCount == 0 && TriggerData.TriggerSource == EGameplayAbilityTriggerSource::OwnedTagPresent)
+						{
+							// Try to cancel, but only if the type is right
+							CancelAbilitySpec(*Spec, nullptr);
+						}
 					}
 				}
 			}
