Fixed issue with gameplay effects's tags not being replicated on with multiple viewports

diff --git a/Engine/Plugins/Runtime/GameplayAbilities/Source/GameplayAbilities/Private/GameplayEffect.cpp b/Engine/Plugins/Runtime/GameplayAbilities/Source/GameplayAbilities/Private/GameplayEffect.cpp
--- a/Engine/Plugins/Runtime/GameplayAbilities/Source/GameplayAbilities/Private/GameplayEffect.cpp
+++ b/Engine/Plugins/Runtime/GameplayAbilities/Source/GameplayAbilities/Private/GameplayEffect.cpp
@@ -3,8 +3,9 @@
 #include "GameplayEffect.h"
 #include "TimerManager.h"
 #include "GameFramework/GameStateBase.h"
-#include "Engine/PackageMapClient.h"
+#include "Engine/ChildConnection.h"
 #include "Engine/NetConnection.h"
+#include "Engine/PackageMapClient.h"
 #include "AbilitySystemStats.h"
 #include "GameplayTagsModule.h"
 #include "AbilitySystemGlobals.h"
@@ -3794,9 +3795,22 @@ bool FActiveGameplayEffectsContainer::NetDeltaSerialize(FNetDeltaSerializeInfo&
 				{
 					// In mixed mode, we only want to replicate to the owner of this channel, minimal replication
 					// data will go to everyone else.
-					if (!Owner->GetOwner()->IsOwnedBy(Connection->OwningActor))
+					const AActor* ParentOwner = Owner->GetOwner();
+					if (!ParentOwner->IsOwnedBy(Connection->OwningActor))
 					{
-						return false;
+						bool bIsChildConnection = false;
+						for (UChildConnection* ChildConnection : Connection->Children)
+						{
+							if (ParentOwner->IsOwnedBy(ChildConnection->OwningActor))
+							{
+								bIsChildConnection = true;
+								break;
+							}
+						}
+						if (!bIsChildConnection)
+						{
+							return false;
+						}
 					}
 				}
 			}
