Fix relevancy issues around player pawn proxy replicateed tags:
-clear tags on pawn going out of relevancy
-properly handle tags that are serialized in first actor bunch before Owner is set on proxy struct.

#rnx

diff --git a/Engine/Plugins/Runtime/GameplayAbilities/Source/GameplayAbilities/Private/GameplayEffectTypes.cpp b/Engine/Plugins/Runtime/GameplayAbilities/Source/GameplayAbilities/Private/GameplayEffectTypes.cpp
--- a/Engine/Plugins/Runtime/GameplayAbilities/Source/GameplayAbilities/Private/GameplayEffectTypes.cpp
+++ b/Engine/Plugins/Runtime/GameplayAbilities/Source/GameplayAbilities/Private/GameplayEffectTypes.cpp
@@ -1079,7 +1079,7 @@ bool FMinimalReplicationTagCountMap::NetSerialize(FArchive& Ar, class UPackageMa
 
 		if (Owner)
 		{
-			bool UpdateOwnerTagMap = true;
+			bool bUpdateOwnerTagMap = true;
 			if (bRequireNonOwningNetConnection)
 			{
 				if (AActor* OwningActor = Owner->GetOwner())
@@ -1089,26 +1089,15 @@ bool FMinimalReplicationTagCountMap::NetSerialize(FArchive& Ar, class UPackageMa
 					{
 						if (OwnerNetConnection == CastChecked<UPackageMapClient>(Map)->GetConnection())
 						{
-							UpdateOwnerTagMap = false;
+							bUpdateOwnerTagMap = false;
 						}
 					}
 				}
 			}
 
-			if (UpdateOwnerTagMap)
+			if (bUpdateOwnerTagMap)
 			{
-				// Update our tags with owner tags
-				for(auto It = TagMap.CreateIterator(); It; ++It)
-				{
-					Owner->SetTagMapCount(It->Key, It->Value);
-
-					// Remove tags with a count of zero from the map. This prevents them
-					// from being replicated incorrectly when recording client replays.
-					if (It->Value == 0)
-					{
-						It.RemoveCurrent();
-					}
-				}
+				UpdateOwnerTagMap();
 			}
 		}
 	}
@@ -1118,4 +1107,45 @@ bool FMinimalReplicationTagCountMap::NetSerialize(FArchive& Ar, class UPackageMa
 	return true;
 }
 
+void FMinimalReplicationTagCountMap::SetOwner(UAbilitySystemComponent* ASC)
+{
+	Owner = ASC;
+	if (Owner && TagMap.Num() > 0)
+	{
+		// Invoke events in case we skipped them during ::NetSerialize
+		UpdateOwnerTagMap();
+	}
+}
+
+void FMinimalReplicationTagCountMap::RemoveAllTags()
+{
+	if (Owner)
+	{
+		for(auto It = TagMap.CreateIterator(); It; ++It)
+		{
+			Owner->SetTagMapCount(It->Key, 0);
+		}
+
+		TagMap.Reset();
+	}
+}
+
+void FMinimalReplicationTagCountMap::UpdateOwnerTagMap()
+{
+	if (Owner)
+	{
+		for (auto It = TagMap.CreateIterator(); It; ++It)
+		{
+			Owner->SetTagMapCount(It->Key, It->Value);
+
+			// Remove tags with a count of zero from the map. This prevents them
+			// from being replicated incorrectly when recording client replays.
+			if (It->Value == 0)
+			{
+				It.RemoveCurrent();
+			}
+		}
+	}
+}
+
 #undef LOCTEXT_NAMESPACE
diff --git a/Engine/Plugins/Runtime/GameplayAbilities/Source/GameplayAbilities/Public/GameplayEffectTypes.h b/Engine/Plugins/Runtime/GameplayAbilities/Source/GameplayAbilities/Public/GameplayEffectTypes.h
--- a/Engine/Plugins/Runtime/GameplayAbilities/Source/GameplayAbilities/Public/GameplayEffectTypes.h
+++ b/Engine/Plugins/Runtime/GameplayAbilities/Source/GameplayAbilities/Public/GameplayEffectTypes.h
@@ -1643,9 +1643,15 @@ struct GAMEPLAYABILITIES_API FMinimalReplicationTagCountMap
 	/** If true, we will skip updating the Owner ASC if we replicate on a connection owned by the ASC */
 	void SetRequireNonOwningNetConnection(bool b) { bRequireNonOwningNetConnection = b; }
 
+	void SetOwner(UAbilitySystemComponent* ASC);
+
+	// Removes all tags that this container is granting
+	void RemoveAllTags();
+
 private:
 
 	bool bRequireNonOwningNetConnection = false;
+	void UpdateOwnerTagMap();
 };
 
 template<>
