Ability System Serialization: Added logging when an Ability System Component add too many Replication Tags and can not serialize them all. This logging prints out all of the tags on the Ability System Component so that programmers will be able to see what tags are causing the issue.

#rnx

diff --git a/Engine/Plugins/Runtime/GameplayAbilities/Source/GameplayAbilities/Private/GameplayEffectTypes.cpp b/Engine/Plugins/Runtime/GameplayAbilities/Source/GameplayAbilities/Private/GameplayEffectTypes.cpp
--- a/Engine/Plugins/Runtime/GameplayAbilities/Source/GameplayAbilities/Private/GameplayEffectTypes.cpp
+++ b/Engine/Plugins/Runtime/GameplayAbilities/Source/GameplayAbilities/Private/GameplayEffectTypes.cpp
@@ -1036,10 +1036,34 @@ bool FMinimalReplicationTagCountMap::NetSerialize(FArchive& Ar, class UPackageMa
 
 	if (Ar.IsSaving())
 	{
+		// if Count is too high then print a warning and clamp Count
 		int32 Count = TagMap.Num();
 		if (Count > MaxCount)
 		{
-			ABILITY_LOG(Error, TEXT("FMinimapReplicationTagCountMap has too many tags (%d). This will cause tags to not replicate. See FMinimapReplicationTagCountMap::NetSerialize"), TagMap.Num());
+#if UE_BUILD_SHIPPING
+			ABILITY_LOG(Error, TEXT("FMinimalReplicationTagCountMap has too many tags (%d) when the limit is %d. This will cause tags to not replicate. See FMinimapReplicationTagCountMap::NetSerialize"), TagMap.Num(), MaxCount);
+#else
+			TArray<FGameplayTag> TagKeys;
+			TagMap.GetKeys(TagKeys);
+
+			const int32 SpaceToReservePerEntry = 40;
+			FString TagsString;
+			// reserve lots of space for our tags string
+			TagsString.Reserve(Count * SpaceToReservePerEntry);
+
+			TagsString = TEXT("");
+			TagKeys[0].GetTagName().AppendString(TagsString);
+			for (int32 TagIndex = 1; TagIndex < Count; ++TagIndex)
+			{
+				// appends ", %tag_name" to the string
+				TagsString.Append(TEXT(", "));
+				TagKeys[TagIndex].GetTagName().AppendString(TagsString);
+			}
+
+			ABILITY_LOG(Error, TEXT("FMinimalReplicationTagCountMap has too many tags (%d) when the limit is %d. This will cause tags to not replicate. See FMinimapReplicationTagCountMap::NetSerialize\nTagMap tags: %s"), TagMap.Num(), MaxCount, *TagsString);
+#endif	
+
+			//clamp the count
 			Count = MaxCount;
 		}
 
