It is now possible to listen to GameplayCues within a UObject.

diff --git a/Engine/Plugins/Runtime/GameplayAbilities/Source/GameplayAbilities/Private/AbilitySystemBlueprintLibrary.cpp b/Engine/Plugins/Runtime/GameplayAbilities/Source/GameplayAbilities/Private/AbilitySystemBlueprintLibrary.cpp
--- a/Engine/Plugins/Runtime/GameplayAbilities/Source/GameplayAbilities/Private/AbilitySystemBlueprintLibrary.cpp
+++ b/Engine/Plugins/Runtime/GameplayAbilities/Source/GameplayAbilities/Private/AbilitySystemBlueprintLibrary.cpp
@@ -581,10 +581,10 @@ bool UAbilitySystemBlueprintLibrary::HasHitResult(FGameplayCueParameters Paramet
 
 void UAbilitySystemBlueprintLibrary::ForwardGameplayCueToTarget(TScriptInterface<IGameplayCueInterface> TargetCueInterface, EGameplayCueEvent::Type EventType, FGameplayCueParameters Parameters)
 {
-	AActor* ActorTarget = Cast<AActor>(TargetCueInterface.GetObject());
-	if (TargetCueInterface && ActorTarget)
+	UObject* TargetObject = TargetCueInterface.GetObject();
+	if (TargetCueInterface && TargetObject)
 	{
-		TargetCueInterface->HandleGameplayCue(ActorTarget, Parameters.OriginalTag, EventType, Parameters);
+		TargetCueInterface->HandleGameplayCue(TargetObject, Parameters.OriginalTag, EventType, Parameters);
 	}
 }
 
diff --git a/Engine/Plugins/Runtime/GameplayAbilities/Source/GameplayAbilities/Private/GameplayCueInterface.cpp b/Engine/Plugins/Runtime/GameplayAbilities/Source/GameplayAbilities/Private/GameplayCueInterface.cpp
--- a/Engine/Plugins/Runtime/GameplayAbilities/Source/GameplayAbilities/Private/GameplayCueInterface.cpp
+++ b/Engine/Plugins/Runtime/GameplayAbilities/Source/GameplayAbilities/Private/GameplayCueInterface.cpp
@@ -25,13 +25,13 @@ UGameplayCueInterface::UGameplayCueInterface(const FObjectInitializer& ObjectIni
 {
 }
 
-void IGameplayCueInterface::DispatchBlueprintCustomHandler(AActor* Actor, UFunction* Func, EGameplayCueEvent::Type EventType, FGameplayCueParameters Parameters)
+void IGameplayCueInterface::DispatchBlueprintCustomHandler(UObject* Object, UFunction* Func, EGameplayCueEvent::Type EventType, FGameplayCueParameters Parameters)
 {
-	GameplayCueInterface_eventBlueprintCustomHandler_Parms Parms;
-	Parms.EventType = EventType;
-	Parms.Parameters = Parameters;
+	GameplayCueInterface_eventBlueprintCustomHandler_Parms Params;
+	Params.EventType = EventType;
+	Params.Parameters = Parameters;
 
-	Actor->ProcessEvent(Func, &Parms);
+	Object->ProcessEvent(Func, &Params);
 }
 
 void IGameplayCueInterface::ClearTagToFunctionMap()
@@ -41,18 +41,33 @@ void IGameplayCueInterface::ClearTagToFunctionMap()
 
 void IGameplayCueInterface::HandleGameplayCues(AActor *Self, const FGameplayTagContainer& GameplayCueTags, EGameplayCueEvent::Type EventType, FGameplayCueParameters Parameters)
 {
-	for (auto TagIt = GameplayCueTags.CreateConstIterator(); TagIt; ++TagIt)
+	HandleGameplayCues((UObject*)Self, GameplayCueTags, EventType, Parameters);
+}
+
+void IGameplayCueInterface::HandleGameplayCues(UObject* Self, const FGameplayTagContainer& GameplayCueTags, EGameplayCueEvent::Type EventType, FGameplayCueParameters Parameters)
+{
+	for (FGameplayTag CueTag : GameplayCueTags)
 	{
-		HandleGameplayCue(Self, *TagIt, EventType, Parameters);
+		HandleGameplayCue(Self, CueTag, EventType, Parameters);
 	}
 }
 
 bool IGameplayCueInterface::ShouldAcceptGameplayCue(AActor *Self, FGameplayTag GameplayCueTag, EGameplayCueEvent::Type EventType, FGameplayCueParameters Parameters)
+{
+	return ShouldAcceptGameplayCue((UObject*)Self, GameplayCueTag, EventType, Parameters);
+}
+
+bool IGameplayCueInterface::ShouldAcceptGameplayCue(UObject* Self, FGameplayTag GameplayCueTag, EGameplayCueEvent::Type EventType, FGameplayCueParameters Parameters)
 {
 	return true;
 }
 
 void IGameplayCueInterface::HandleGameplayCue(AActor *Self, FGameplayTag GameplayCueTag, EGameplayCueEvent::Type EventType, FGameplayCueParameters Parameters)
+{
+	HandleGameplayCue((UObject*)Self, GameplayCueTag, EventType, Parameters);
+}
+
+void IGameplayCueInterface::HandleGameplayCue(UObject* Self, FGameplayTag GameplayCueTag, EGameplayCueEvent::Type EventType, FGameplayCueParameters Parameters)
 {
 	SCOPE_CYCLE_COUNTER(STAT_GameplayCueInterface_HandleGameplayCue);
 
@@ -128,14 +143,17 @@ void IGameplayCueInterface::HandleGameplayCue(AActor *Self, FGameplayTag Gamepla
 
 	if (bShouldContinue)
 	{
-		TArray<UGameplayCueSet*> Sets;
-		GetGameplayCueSets(Sets);
-		for (UGameplayCueSet* Set : Sets)
+		if (AActor* SelfActor = Cast<AActor>(Self))
 		{
-			bShouldContinue = Set->HandleGameplayCue(Self, GameplayCueTag, EventType, Parameters);
-			if (!bShouldContinue)
+			TArray<UGameplayCueSet*> Sets;
+			GetGameplayCueSets(Sets);
+			for (UGameplayCueSet* Set : Sets)
 			{
-				break;
+				bShouldContinue = Set->HandleGameplayCue(SelfActor, GameplayCueTag, EventType, Parameters);
+				if (!bShouldContinue)
+				{
+					break;
+				}
 			}
 		}
 	}
diff --git a/Engine/Plugins/Runtime/GameplayAbilities/Source/GameplayAbilities/Private/GameplayCueManager.cpp b/Engine/Plugins/Runtime/GameplayAbilities/Source/GameplayAbilities/Private/GameplayCueManager.cpp
--- a/Engine/Plugins/Runtime/GameplayAbilities/Source/GameplayAbilities/Private/GameplayCueManager.cpp
+++ b/Engine/Plugins/Runtime/GameplayAbilities/Source/GameplayAbilities/Private/GameplayCueManager.cpp
@@ -112,7 +112,7 @@ void UGameplayCueManager::HandleGameplayCues(AActor* TargetActor, const FGamepla
 #if WITH_EDITOR
 	if (GIsEditor && TargetActor == nullptr && UGameplayCueManager::PreviewComponent)
 	{
-		TargetActor = Cast<AActor>(AActor::StaticClass()->GetDefaultObject());
+		TargetActor = GetMutableDefault<AActor>();
 	}
 #endif
 
@@ -155,22 +155,15 @@ void UGameplayCueManager::HandleGameplayCue(AActor* TargetActor, FGameplayTag Ga
 
 bool UGameplayCueManager::ShouldSuppressGameplayCues(AActor* TargetActor)
 {
-	if (DisableGameplayCues)
-	{
-		return true;
-	}
-
-	if (GameplayCueRunOnDedicatedServer == 0 && IsDedicatedServerForGameplayCue())
-	{
-		return true;
-	}
-
-	if (TargetActor == nullptr)
+	if (DisableGameplayCues ||
+		!TargetActor ||
+		(GameplayCueRunOnDedicatedServer == 0 && IsDedicatedServerForGameplayCue()))
 	{
 		return true;
 	}
 
 	return false;
+
 }
 
 void UGameplayCueManager::RouteGameplayCue(AActor* TargetActor, FGameplayTag GameplayCueTag, EGameplayCueEvent::Type EventType, const FGameplayCueParameters& Parameters, EGameplayCueExecutionOptions Options)
@@ -1296,7 +1289,7 @@ void UGameplayCueManager::InvokeGameplayCueExecuted_FromSpec(UAbilitySystemCompo
 
 	if (AbilitySystemAlwaysConvertGESpecToGCParams)
 	{
-		// Transform the GE Spec into GameplayCue parmameters here (on the server)
+		// Transform the GE Spec into GameplayCue parameters here (on the server)
 		PendingCue.PayloadType = EGameplayCuePayloadType::CueParameters;
 		PendingCue.OwningComponent = OwningComponent;
 		PendingCue.PredictionKey = PredictionKey;
@@ -1312,55 +1305,48 @@ void UGameplayCueManager::InvokeGameplayCueExecuted_FromSpec(UAbilitySystemCompo
 	}
 	else
 	{
-		// Transform the GE Spec into a FGameplayEffectSpecForRPC (holds less information than the GE Spec itself, but more information that the FGamepalyCueParameter)
+		// Transform the GE Spec into a FGameplayEffectSpecForRPC (holds less information than the GE Spec itself, but more information than the FGameplayCueParameter)
 		PendingCue.PayloadType = EGameplayCuePayloadType::FromSpec;
 		PendingCue.OwningComponent = OwningComponent;
 		PendingCue.FromSpec = FGameplayEffectSpecForRPC(Spec);
 		PendingCue.PredictionKey = PredictionKey;
 	}
 
-	if (ProcessPendingCueExecute(PendingCue))
-	{
-		PendingExecuteCues.Add(PendingCue);
-	}
-
-	if (GameplayCueSendContextCount == 0)
-	{
-		// Not in a context, flush now
-		FlushPendingCues();
-	}
+	AddPendingCueExecuteInternal(PendingCue);
 }
 
 void UGameplayCueManager::InvokeGameplayCueExecuted(UAbilitySystemComponent* OwningComponent, const FGameplayTag GameplayCueTag, FPredictionKey PredictionKey, FGameplayEffectContextHandle EffectContext)
 {
-	FGameplayCuePendingExecute PendingCue;
-	PendingCue.PayloadType = EGameplayCuePayloadType::CueParameters;
-	PendingCue.GameplayCueTags.Add(GameplayCueTag);
-	PendingCue.OwningComponent = OwningComponent;
-	UAbilitySystemGlobals::Get().InitGameplayCueParameters(PendingCue.CueParameters, EffectContext);
-	PendingCue.PredictionKey = PredictionKey;
-
-	if (ProcessPendingCueExecute(PendingCue))
+	if (OwningComponent)
 	{
-		PendingExecuteCues.Add(PendingCue);
-	}
+		FGameplayCuePendingExecute PendingCue;
+		PendingCue.PayloadType = EGameplayCuePayloadType::CueParameters;
+		PendingCue.GameplayCueTags.Add(GameplayCueTag);
+		PendingCue.OwningComponent = OwningComponent;
+		UAbilitySystemGlobals::Get().InitGameplayCueParameters(PendingCue.CueParameters, EffectContext);
+		PendingCue.PredictionKey = PredictionKey;
 
-	if (GameplayCueSendContextCount == 0)
-	{
-		// Not in a context, flush now
-		FlushPendingCues();
+		AddPendingCueExecuteInternal(PendingCue);
 	}
 }
 
 void UGameplayCueManager::InvokeGameplayCueExecuted_WithParams(UAbilitySystemComponent* OwningComponent, const FGameplayTag GameplayCueTag, FPredictionKey PredictionKey, FGameplayCueParameters GameplayCueParameters)
 {
-	FGameplayCuePendingExecute PendingCue;
-	PendingCue.PayloadType = EGameplayCuePayloadType::CueParameters;
-	PendingCue.GameplayCueTags.Add(GameplayCueTag);
-	PendingCue.OwningComponent = OwningComponent;
-	PendingCue.CueParameters = GameplayCueParameters;
-	PendingCue.PredictionKey = PredictionKey;
+	if (OwningComponent)
+	{
+		FGameplayCuePendingExecute PendingCue;
+		PendingCue.PayloadType = EGameplayCuePayloadType::CueParameters;
+		PendingCue.GameplayCueTags.Add(GameplayCueTag);
+		PendingCue.OwningComponent = OwningComponent;
+		PendingCue.CueParameters = GameplayCueParameters;
+		PendingCue.PredictionKey = PredictionKey;
+
+		AddPendingCueExecuteInternal(PendingCue);
+	}
+}
 
+void UGameplayCueManager::AddPendingCueExecuteInternal(FGameplayCuePendingExecute& PendingCue)
+{
 	if (ProcessPendingCueExecute(PendingCue))
 	{
 		PendingExecuteCues.Add(PendingCue);
@@ -1416,9 +1402,8 @@ void UGameplayCueManager::FlushPendingCues()
 			}
 
 			// TODO: Could implement non-rpc method for replicating if desired
-			switch (PendingCue.PayloadType)
+			if (PendingCue.PayloadType == EGameplayCuePayloadType::CueParameters)
 			{
-			case EGameplayCuePayloadType::CueParameters:
 				if (ensure(PendingCue.GameplayCueTags.Num() >= 1))
 				{
 					if (bHasAuthority)
@@ -1444,34 +1429,9 @@ void UGameplayCueManager::FlushPendingCues()
 						}
 					}
 				}
-				break;
-			case EGameplayCuePayloadType::EffectContext:
-				if (ensure(PendingCue.GameplayCueTags.Num() >= 1))
-				{
-					if (bHasAuthority)
-					{
-						PendingCue.OwningComponent->ForceReplication();
-						if (PendingCue.GameplayCueTags.Num() > 1)
-						{
-							RepInterface->Call_InvokeGameplayCuesExecuted(FGameplayTagContainer::CreateFromArray(PendingCue.GameplayCueTags), PendingCue.PredictionKey, PendingCue.CueParameters.EffectContext);
-						}
-						else
-						{
-							RepInterface->Call_InvokeGameplayCueExecuted(PendingCue.GameplayCueTags[0], PendingCue.PredictionKey, PendingCue.CueParameters.EffectContext);
-							static FName NetMulticast_InvokeGameplayCueExecutedName = TEXT("NetMulticast_InvokeGameplayCueExecuted");
-							CheckForTooManyRPCs(NetMulticast_InvokeGameplayCueExecutedName, PendingCue, PendingCue.GameplayCueTags[0].ToString(), PendingCue.CueParameters.EffectContext.Get());
-						}
-					}
-					else if (bLocalPredictionKey)
-					{
-						for (const FGameplayTag& Tag : PendingCue.GameplayCueTags)
-						{
-							PendingCue.OwningComponent->InvokeGameplayCueEvent(Tag, EGameplayCueEvent::Executed, PendingCue.CueParameters.EffectContext);
-						}
-					}
-				}
-				break;
-			case EGameplayCuePayloadType::FromSpec:
+			}
+			else if (PendingCue.PayloadType == EGameplayCuePayloadType::FromSpec)
+			{
 				if (bHasAuthority)
 				{
 					RepInterface->ForceReplication();
@@ -1484,7 +1444,6 @@ void UGameplayCueManager::FlushPendingCues()
 				{
 					PendingCue.OwningComponent->InvokeGameplayCueEvent(PendingCue.FromSpec, EGameplayCueEvent::Executed);
 				}
-				break;
 			}
 		}
 	}
diff --git a/Engine/Plugins/Runtime/GameplayAbilities/Source/GameplayAbilities/Public/GameplayCueInterface.h b/Engine/Plugins/Runtime/GameplayAbilities/Source/GameplayAbilities/Public/GameplayCueInterface.h
--- a/Engine/Plugins/Runtime/GameplayAbilities/Source/GameplayAbilities/Public/GameplayCueInterface.h
+++ b/Engine/Plugins/Runtime/GameplayAbilities/Source/GameplayAbilities/Public/GameplayCueInterface.h
@@ -23,6 +23,23 @@ class GAMEPLAYABILITIES_API IGameplayCueInterface
 {
 	GENERATED_IINTERFACE_BODY()
 
+public:
+
+	/** Handle a single gameplay cue */
+	virtual void HandleGameplayCue(UObject* Self, FGameplayTag GameplayCueTag, EGameplayCueEvent::Type EventType, FGameplayCueParameters Parameters);
+
+	/** Wrapper that handles multiple cues */
+	virtual void HandleGameplayCues(UObject* Self, const FGameplayTagContainer& GameplayCueTags, EGameplayCueEvent::Type EventType, FGameplayCueParameters Parameters);
+
+	/**
+	* Returns true if the object can currently accept gameplay cues associated with the given tag. Returns true by default.
+	* Allows objects to opt out of cues in cases such as pending death
+	*/
+	virtual bool ShouldAcceptGameplayCue(UObject* Self, FGameplayTag GameplayCueTag, EGameplayCueEvent::Type EventType, FGameplayCueParameters Parameters);
+
+
+	// DEPRECATED - use the UObject* signatures above
+
 	/** Handle a single gameplay cue */
 	virtual void HandleGameplayCue(AActor *Self, FGameplayTag GameplayCueTag, EGameplayCueEvent::Type EventType, FGameplayCueParameters Parameters);
 
@@ -32,6 +49,9 @@ class GAMEPLAYABILITIES_API IGameplayCueInterface
 	/** Returns true if the actor can currently accept gameplay cues associated with the given tag. Returns true by default. Allows actors to opt out of cues in cases such as pending death */
 	virtual bool ShouldAcceptGameplayCue(AActor *Self, FGameplayTag GameplayCueTag, EGameplayCueEvent::Type EventType, FGameplayCueParameters Parameters);
 
+	// END DEPRECATED
+
+
 	/** Return the cue sets used by this object. This is optional and it is possible to leave this list empty. */
 	virtual void GetGameplayCueSets(TArray<class UGameplayCueSet*>& OutSets) const {}
 
@@ -47,7 +67,7 @@ class GAMEPLAYABILITIES_API IGameplayCueInterface
 	virtual void ForwardGameplayCueToParent();
 
 	/** Calls the UFunction override for a specific gameplay cue */
-	static void DispatchBlueprintCustomHandler(AActor* Actor, UFunction* Func, EGameplayCueEvent::Type EventType, FGameplayCueParameters Parameters);
+	static void DispatchBlueprintCustomHandler(UObject* Object, UFunction* Func, EGameplayCueEvent::Type EventType, FGameplayCueParameters Parameters);
 
 	/** Clears internal cache of what classes implement which functions */
 	static void ClearTagToFunctionMap();
diff --git a/Engine/Plugins/Runtime/GameplayAbilities/Source/GameplayAbilities/Public/GameplayCueManager.h b/Engine/Plugins/Runtime/GameplayAbilities/Source/GameplayAbilities/Public/GameplayCueManager.h
--- a/Engine/Plugins/Runtime/GameplayAbilities/Source/GameplayAbilities/Public/GameplayCueManager.h
+++ b/Engine/Plugins/Runtime/GameplayAbilities/Source/GameplayAbilities/Public/GameplayCueManager.h
@@ -370,4 +370,9 @@ protected:
 	TArray<FPreallocationInfo>	PreallocationInfoList_Internal;
 
 	FOnRouteGameplayCue OnRouteGameplayCue;
+
+private:
+
+	void AddPendingCueExecuteInternal(FGameplayCuePendingExecute& PendingCue);
+
 };
\ No newline at end of file
diff --git a/Engine/Plugins/Runtime/GameplayAbilities/Source/GameplayAbilities/Public/GameplayCue_Types.h b/Engine/Plugins/Runtime/GameplayAbilities/Source/GameplayAbilities/Public/GameplayCue_Types.h
--- a/Engine/Plugins/Runtime/GameplayAbilities/Source/GameplayAbilities/Public/GameplayCue_Types.h
+++ b/Engine/Plugins/Runtime/GameplayAbilities/Source/GameplayAbilities/Public/GameplayCue_Types.h
@@ -20,8 +20,6 @@ class UGameplayCueSet;
 UENUM()
 enum class EGameplayCuePayloadType : uint8
 {
-	/** Uses FGameplayEffectContextHandle inside GameplayCueParameters */
-	EffectContext,
 	/** Uses FGameplayCueParameters */
 	CueParameters,
 	/** Uses FGameplayEffectSpecForRPC */
@@ -35,11 +33,7 @@ struct FGameplayCuePendingExecute
 {
 	GENERATED_USTRUCT_BODY()
 
-	FGameplayCuePendingExecute()
-		: PayloadType(EGameplayCuePayloadType::EffectContext)
-		, OwningComponent(nullptr)
-	{
-	}
+	FGameplayCuePendingExecute() {}
 
 	/** List of tags, we allocate one as there is almost always exactly one tag */
 	TArray<FGameplayTag, TInlineAllocator<1> > GameplayCueTags;
@@ -50,11 +44,11 @@ struct FGameplayCuePendingExecute
 
 	/** What type of payload is attached to this cue */
 	UPROPERTY()
-	EGameplayCuePayloadType PayloadType;
+	EGameplayCuePayloadType PayloadType = EGameplayCuePayloadType::CueParameters;
 
 	/** What component to send the cue on */
 	UPROPERTY()
-	UAbilitySystemComponent* OwningComponent;
+	UAbilitySystemComponent* OwningComponent = nullptr;
 
 	/** If this cue is from a spec, here's the copy of that spec */
 	UPROPERTY()
