#UE4 Fix memory problem of calling malloc to allocate FGameplayAbilityTargetData, but using delete to deallocate it. InitializeStruct was already properly calling the constructor for this memory, so now we are using DestroyStruct to call the destructor instead of implicitly using delete to call it. We are now using free to release the memory to match the malloc.

diff --git a/Engine/Plugins/Runtime/GameplayAbilities/Source/GameplayAbilities/Private/GameplayAbilityTargetTypes.cpp b/Engine/Plugins/Runtime/GameplayAbilities/Source/GameplayAbilities/Private/GameplayAbilityTargetTypes.cpp
--- a/Engine/Plugins/Runtime/GameplayAbilities/Source/GameplayAbilities/Private/GameplayAbilityTargetTypes.cpp
+++ b/Engine/Plugins/Runtime/GameplayAbilities/Source/GameplayAbilities/Private/GameplayAbilityTargetTypes.cpp
@@ -140,6 +140,18 @@ FGameplayAbilityTargetDataHandle FGameplayAbilityTargetingLocationInfo::MakeTarg
 #define TARGETDATAHANDLE_SAFE_NET_SERIALIZE 1
 #endif
 
+struct FGameplayAbilityTargetDataDeleter
+{
+	FORCEINLINE void operator()(FGameplayAbilityTargetData* Object) const
+	{
+		check(Object);
+		UScriptStruct* ScriptStruct = Object->GetScriptStruct();
+		check(ScriptStruct);
+		ScriptStruct->DestroyStruct(Object);
+		FMemory::Free(Object);
+	}
+};
+
 bool FGameplayAbilityTargetDataHandle::NetSerialize(FArchive& Ar, class UPackageMap* Map, bool& bOutSuccess)
 {
 	Ar << UniqueId;
@@ -182,7 +194,7 @@ bool FGameplayAbilityTargetDataHandle::NetSerialize(FArchive& Ar, class UPackage
 				FGameplayAbilityTargetData * NewData = (FGameplayAbilityTargetData*)FMemory::Malloc(ScriptStruct->GetStructureSize());
 				ScriptStruct->InitializeStruct(NewData);
 
-				Data[i] = TSharedPtr<FGameplayAbilityTargetData>(NewData);
+				Data[i] = TSharedPtr<FGameplayAbilityTargetData>(NewData, FGameplayAbilityTargetDataDeleter());
 			}
 
 			void* ContainerPtr = Data[i].Get();
