Improved serialization of FGameplayAbilityRepAnimMontage

- Added 2 modes of support to either replicate the position or the start section id give to PlayMontage via the StartSectionName
  - Mode 1: Position serialization now uses a quantized float that is packed to improve struct size and the cost of some precision
  - Mode 2: At the cost of not replicating position, we replicate only the Section Id to Play which is specified by the StartSectionName to PlayMontage. This packs into 7 bits so it is the optimal way to rep montage data when 99% of the time PlayIng Montage is from the start of the montage, and doesn't need to the precision on the client to maintain synchronization

#rnx

diff --git a/Engine/Plugins/Runtime/GameplayAbilities/Source/GameplayAbilities/Private/AbilitySystemComponent_Abilities.cpp b/Engine/Plugins/Runtime/GameplayAbilities/Source/GameplayAbilities/Private/AbilitySystemComponent_Abilities.cpp
--- a/Engine/Plugins/Runtime/GameplayAbilities/Source/GameplayAbilities/Private/AbilitySystemComponent_Abilities.cpp
+++ b/Engine/Plugins/Runtime/GameplayAbilities/Source/GameplayAbilities/Private/AbilitySystemComponent_Abilities.cpp
@@ -2377,6 +2377,13 @@ float UAbilitySystemComponent::PlayMontage(UGameplayAbility* InAnimatingAbility,
 				RepAnimMontageInfo.AnimMontage = NewAnimMontage;
 				RepAnimMontageInfo.ForcePlayBit = !bool(RepAnimMontageInfo.ForcePlayBit);
 
+				RepAnimMontageInfo.SectionIdToPlay = 0;
+				if (RepAnimMontageInfo.AnimMontage && StartSectionName != NAME_None)
+				{
+					// we add one so INDEX_NONE can be used in the on rep
+					RepAnimMontageInfo.SectionIdToPlay = RepAnimMontageInfo.AnimMontage->GetSectionIndex(StartSectionName) + 1;
+				}
+
 				// Update parameters that change during Montage life time.
 				AnimMontage_UpdateReplicatedData();
 
@@ -2567,6 +2574,16 @@ void UAbilitySystemComponent::OnRep_ReplicatedAnimMontage()
 				AnimInstance->Montage_SetPlayRate(LocalAnimMontageInfo.AnimMontage, RepAnimMontageInfo.PlayRate);
 			}
 
+			const int32 SectionIdToPlay = (static_cast<int32>(RepAnimMontageInfo.SectionIdToPlay) - 1);
+			if (SectionIdToPlay != INDEX_NONE)
+			{
+				FName SectionNameToJumpTo = LocalAnimMontageInfo.AnimMontage->GetSectionName(SectionIdToPlay);
+				if (SectionNameToJumpTo != NAME_None)
+				{
+					AnimInstance->Montage_JumpToSection(SectionNameToJumpTo);
+				}
+			}
+
 			// Compressed Flags
 			const bool bIsStopped = AnimInstance->Montage_GetIsStopped(LocalAnimMontageInfo.AnimMontage);
 			const bool bReplicatedIsStopped = bool(RepAnimMontageInfo.IsStopped);
@@ -2677,6 +2694,13 @@ void UAbilitySystemComponent::CurrentMontageJumpToSection(FName SectionName)
 		AnimInstance->Montage_JumpToSection(SectionName, LocalAnimMontageInfo.AnimMontage);
 		if (IsOwnerActorAuthoritative())
 		{
+			RepAnimMontageInfo.SectionIdToPlay = 0;
+			if (RepAnimMontageInfo.AnimMontage)
+			{
+				// we add one so INDEX_NONE can be used in the on rep
+				RepAnimMontageInfo.SectionIdToPlay = RepAnimMontageInfo.AnimMontage->GetSectionIndex(SectionName) + 1;
+			}
+
 			AnimMontage_UpdateReplicatedData();
 		}
 		else
@@ -2784,6 +2808,13 @@ void UAbilitySystemComponent::ServerCurrentMontageJumpToSectionName_Implementati
 			// Update replicated version for Simulated Proxies if we are on the server.
 			if (IsOwnerActorAuthoritative())
 			{
+				RepAnimMontageInfo.SectionIdToPlay = 0;
+				if (RepAnimMontageInfo.AnimMontage && SectionName != NAME_None)
+				{
+					// we add one so INDEX_NONE can be used in the on rep
+					RepAnimMontageInfo.SectionIdToPlay = RepAnimMontageInfo.AnimMontage->GetSectionIndex(SectionName) + 1;
+				}
+
 				AnimMontage_UpdateReplicatedData();
 			}
 		}
@@ -2899,6 +2930,11 @@ float UAbilitySystemComponent::GetCurrentMontageSectionTimeLeft() const
 	return -1.f;
 }
 
+void UAbilitySystemComponent::SetMontageRepAnimPositionMethod(ERepAnimPositionMethod InMethod)
+{
+	RepAnimMontageInfo.SetRepAnimPositionMethod(InMethod);
+}
+
 bool UAbilitySystemComponent::IsAnimatingAbility(UGameplayAbility* InAbility) const
 {
 	return (LocalAnimMontageInfo.AnimatingAbility == InAbility);
diff --git a/Engine/Plugins/Runtime/GameplayAbilities/Source/GameplayAbilities/Private/GameplayAbilityTypes.cpp b/Engine/Plugins/Runtime/GameplayAbilities/Source/GameplayAbilities/Private/GameplayAbilityTypes.cpp
--- a/Engine/Plugins/Runtime/GameplayAbilities/Source/GameplayAbilities/Private/GameplayAbilityTypes.cpp
+++ b/Engine/Plugins/Runtime/GameplayAbilities/Source/GameplayAbilities/Private/GameplayAbilityTypes.cpp
@@ -1,10 +1,13 @@
 // Copyright Epic Games, Inc. All Rights Reserved.
 
 #include "Abilities/GameplayAbilityTypes.h"
+
 #include "GameFramework/PlayerController.h"
 #include "GameFramework/MovementComponent.h"
 #include "Abilities/GameplayAbility.h"
 #include "AbilitySystemComponent.h"
+#include "Animation/AnimMontage.h"
+#include "GameplayPrediction.h"
 
 //----------------------------------------------------------------------
 
@@ -368,3 +371,73 @@ void FGameplayAbilityReplicatedDataContainer::PrintDebug()
 	ABILITY_LOG(Warning, TEXT("============================="));
 }
 
+bool FGameplayAbilityRepAnimMontage::NetSerialize(FArchive& Ar, class UPackageMap* Map, bool& bOutSuccess)
+{
+	uint8 RepPosition = bRepPosition;
+	Ar.SerializeBits(&RepPosition, 1);
+	if (RepPosition)
+	{
+		bRepPosition = true;
+
+		// when rep'ing position, we don't want to skip correction
+		// and we don't need to force the section id to play
+		SectionIdToPlay = 0;
+		SkipPositionCorrection = false;
+
+		if (Ar.IsSaving())
+		{
+			uint32 PackedPosition = static_cast<uint32>(Position * 100.0f);
+			Ar.SerializeIntPacked(PackedPosition);
+		}
+		else if (Ar.IsLoading())
+		{
+			uint32 PackedPosition = 0;
+			Ar.SerializeIntPacked(PackedPosition);
+			Position = static_cast<float>(PackedPosition) / 100.0f;
+		}
+	}
+	else
+	{
+		bRepPosition = false;
+
+		// when rep'ing the section to play id, we want to skip
+		// correction, and don't want a position
+		SkipPositionCorrection = true;
+		Position = 0.0f;
+		Ar.SerializeBits(&SectionIdToPlay, 7);
+	}
+
+	uint8 bIsStopped = IsStopped;
+	Ar.SerializeBits(&bIsStopped, 1);
+	IsStopped = bIsStopped & 1;
+
+	uint8 bForcePlayBit = ForcePlayBit;
+	Ar.SerializeBits(&bForcePlayBit, 1);
+	ForcePlayBit = bForcePlayBit & 1;
+
+	uint8 bSkipPositionCorrection = SkipPositionCorrection;
+	Ar.SerializeBits(&bSkipPositionCorrection, 1);
+	SkipPositionCorrection = bSkipPositionCorrection & 1;
+
+	uint8 SkipPlayRate = bSkipPlayRate;
+	Ar.SerializeBits(&SkipPlayRate, 1);
+	bSkipPlayRate = SkipPlayRate & 1;
+
+	Ar << AnimMontage;
+	Ar << PlayRate;
+	Ar << BlendTime;
+	Ar << NextSectionID;
+	PredictionKey.NetSerialize(Ar, Map, bOutSuccess);
+
+	bOutSuccess = true;
+	return true;
+}
+
+void FGameplayAbilityRepAnimMontage::SetRepAnimPositionMethod(ERepAnimPositionMethod InMethod)
+{
+	switch (InMethod)
+	{
+	case ERepAnimPositionMethod::Position: bRepPosition = true; break;
+	case ERepAnimPositionMethod::CurrentSectionId: bRepPosition = false; break;
+	}
+}
diff --git a/Engine/Plugins/Runtime/GameplayAbilities/Source/GameplayAbilities/Public/Abilities/GameplayAbilityTypes.h b/Engine/Plugins/Runtime/GameplayAbilities/Source/GameplayAbilities/Public/Abilities/GameplayAbilityTypes.h
--- a/Engine/Plugins/Runtime/GameplayAbilities/Source/GameplayAbilities/Public/Abilities/GameplayAbilityTypes.h
+++ b/Engine/Plugins/Runtime/GameplayAbilities/Source/GameplayAbilities/Public/Abilities/GameplayAbilityTypes.h
@@ -192,6 +192,14 @@ struct GAMEPLAYABILITIES_API FGameplayAbilityActorInfo
 	virtual void ClearActorInfo();
 };
 
+/** Enum used by the Ability Rep Anim Montage struct to rep the quantized position or the current section id */
+UENUM()
+enum class ERepAnimPositionMethod
+{
+	Position = 0,			// reps the position in the montage to keep the client in sync (heavier, quantized, more precise)
+	CurrentSectionId = 1,	// reps the current section id we want to play on the client (compact, less precise)
+};
+
 /** Data about montages that is replicated to simulated clients */
 USTRUCT()
 struct GAMEPLAYABILITIES_API FGameplayAbilityRepAnimMontage
@@ -207,7 +215,7 @@ struct GAMEPLAYABILITIES_API FGameplayAbilityRepAnimMontage
 	float PlayRate;
 
 	/** Montage position */
-	UPROPERTY()
+	UPROPERTY(NotReplicated)
 	float Position;
 
 	/** Montage current blend time */
@@ -218,6 +226,10 @@ struct GAMEPLAYABILITIES_API FGameplayAbilityRepAnimMontage
 	UPROPERTY()
 	uint8 NextSectionID;
 
+	/** flag indicating we should serialize the position or the current section id */
+	UPROPERTY()
+	uint8 bRepPosition : 1;
+
 	/** Bit set when montage has been stopped. */
 	UPROPERTY()
 	uint8 IsStopped : 1;
@@ -234,21 +246,40 @@ struct GAMEPLAYABILITIES_API FGameplayAbilityRepAnimMontage
 	UPROPERTY()
 	uint8 bSkipPlayRate : 1;
 
+	UPROPERTY()
+	FPredictionKey PredictionKey;
+
+	/** The current section Id used by the montage. Will only be valid if bRepPosition is false */
+	UPROPERTY(NotReplicated)
+	uint8 SectionIdToPlay;
+
 	FGameplayAbilityRepAnimMontage()
-		: AnimMontage(nullptr),
-		PlayRate(0.f),
-		Position(0.f),
-		BlendTime(0.f),
-		NextSectionID(0),
-		IsStopped(true),
-		ForcePlayBit(0),
-		SkipPositionCorrection(false),
-		bSkipPlayRate(false)
+	: AnimMontage(nullptr),
+	PlayRate(0.f),
+	Position(0.f),
+	BlendTime(0.f),
+	NextSectionID(0),
+	bRepPosition(true),
+	IsStopped(true),
+	ForcePlayBit(0),
+	SkipPositionCorrection(false),
+	bSkipPlayRate(false),
+	SectionIdToPlay(0)
 	{
 	}
 
-	UPROPERTY()
-	FPredictionKey PredictionKey;
+	bool NetSerialize(FArchive& Ar, class UPackageMap* Map, bool& bOutSuccess);
+
+	void SetRepAnimPositionMethod(ERepAnimPositionMethod InMethod);
+};
+
+template<>
+struct TStructOpsTypeTraits<FGameplayAbilityRepAnimMontage> : public TStructOpsTypeTraitsBase2<FGameplayAbilityRepAnimMontage>
+{
+	enum
+	{
+		WithNetSerializer = true,
+	};
 };
 
 
diff --git a/Engine/Plugins/Runtime/GameplayAbilities/Source/GameplayAbilities/Public/AbilitySystemComponent.h b/Engine/Plugins/Runtime/GameplayAbilities/Source/GameplayAbilities/Public/AbilitySystemComponent.h
--- a/Engine/Plugins/Runtime/GameplayAbilities/Source/GameplayAbilities/Public/AbilitySystemComponent.h
+++ b/Engine/Plugins/Runtime/GameplayAbilities/Source/GameplayAbilities/Public/AbilitySystemComponent.h
@@ -1215,6 +1215,9 @@ class GAMEPLAYABILITIES_API UAbilitySystemComponent : public UGameplayTasksCompo
 	/** Returns amount of time left in current section */
 	float GetCurrentMontageSectionTimeLeft() const;
 
+	/** Method to set the replication method for the position in the montage */
+	void SetMontageRepAnimPositionMethod(ERepAnimPositionMethod InMethod);
+
 	// ----------------------------------------------------------------------------------------------------------------
 	//	Actor interaction
 	// ----------------------------------------------------------------------------------------------------------------	
