Framework: Fixed issues in operator== for FAttributeBasedFloat, FCustomCalculationBasedFloat, and FGameplayModifierInfo that would consider functionally different instances to be equal

diff --git a/Engine/Plugins/Runtime/GameplayAbilities/Source/GameplayAbilities/Private/GameplayEffect.cpp b/Engine/Plugins/Runtime/GameplayAbilities/Source/GameplayAbilities/Private/GameplayEffect.cpp
--- a/Engine/Plugins/Runtime/GameplayAbilities/Source/GameplayAbilities/Private/GameplayEffect.cpp
+++ b/Engine/Plugins/Runtime/GameplayAbilities/Source/GameplayAbilities/Private/GameplayEffect.cpp
@@ -250,17 +250,18 @@ bool FAttributeBasedFloat::operator==(const FAttributeBasedFloat& Other) const
 		PostMultiplyAdditiveValue != Other.PostMultiplyAdditiveValue ||
 		BackingAttribute != Other.BackingAttribute ||
 		AttributeCurve != Other.AttributeCurve ||
-		AttributeCalculationType != Other.AttributeCalculationType)
+		AttributeCalculationType != Other.AttributeCalculationType ||
+		FinalChannel != Other.FinalChannel)
 	{
 		return false;
 	}
 	if (SourceTagFilter.Num() != Other.SourceTagFilter.Num() ||
-		!SourceTagFilter.HasAll(Other.SourceTagFilter))
+		!SourceTagFilter.HasAllExact(Other.SourceTagFilter))
 	{
 		return false;
 	}
 	if (TargetTagFilter.Num() != Other.TargetTagFilter.Num() ||
-		!TargetTagFilter.HasAll(Other.TargetTagFilter))
+		!TargetTagFilter.HasAllExact(Other.TargetTagFilter))
 	{
 		return false;
 	}
@@ -307,7 +308,8 @@ bool FCustomCalculationBasedFloat::operator==(const FCustomCalculationBasedFloat
 	}
 	if (Coefficient != Other.Coefficient ||
 		PreMultiplyAdditiveValue != Other.PreMultiplyAdditiveValue ||
-		PostMultiplyAdditiveValue != Other.PostMultiplyAdditiveValue)
+		PostMultiplyAdditiveValue != Other.PostMultiplyAdditiveValue ||
+		FinalLookupCurve != Other.FinalLookupCurve)
 	{
 		return false;
 	}
@@ -4854,20 +4856,20 @@ bool FGameplayModifierInfo::operator==(const FGameplayModifierInfo& Other) const
 		return false;
 	}
 
-	if (SourceTags.RequireTags.Num() != Other.SourceTags.RequireTags.Num() || !SourceTags.RequireTags.HasAll(Other.SourceTags.RequireTags))
+	if (SourceTags.RequireTags.Num() != Other.SourceTags.RequireTags.Num() || !SourceTags.RequireTags.HasAllExact(Other.SourceTags.RequireTags))
 	{
 		return false;
 	}
-	if (SourceTags.IgnoreTags.Num() != Other.SourceTags.IgnoreTags.Num() || !SourceTags.IgnoreTags.HasAll(Other.SourceTags.IgnoreTags))
+	if (SourceTags.IgnoreTags.Num() != Other.SourceTags.IgnoreTags.Num() || !SourceTags.IgnoreTags.HasAllExact(Other.SourceTags.IgnoreTags))
 	{
 		return false;
 	}
 
-	if (TargetTags.RequireTags.Num() != Other.TargetTags.RequireTags.Num() || !TargetTags.RequireTags.HasAll(Other.TargetTags.RequireTags))
+	if (TargetTags.RequireTags.Num() != Other.TargetTags.RequireTags.Num() || !TargetTags.RequireTags.HasAllExact(Other.TargetTags.RequireTags))
 	{
 		return false;
 	}
-	if (TargetTags.IgnoreTags.Num() != Other.TargetTags.IgnoreTags.Num() || !TargetTags.IgnoreTags.HasAll(Other.TargetTags.IgnoreTags))
+	if (TargetTags.IgnoreTags.Num() != Other.TargetTags.IgnoreTags.Num() || !TargetTags.IgnoreTags.HasAllExact(Other.TargetTags.IgnoreTags))
 	{
 		return false;
 	}
