Change how we handle updating gameplay tags inside of tag count containers. When deferring the update of parent tags when removing gameplay tags we will now call the change related delegates after the parent tags have updated. This ensures that the tag table is in a consistent state when the delegates broadcast.

diff --git a/Engine/Plugins/Runtime/GameplayAbilities/Source/GameplayAbilities/Private/AbilitySystemComponent.cpp b/Engine/Plugins/Runtime/GameplayAbilities/Source/GameplayAbilities/Private/AbilitySystemComponent.cpp
--- a/Engine/Plugins/Runtime/GameplayAbilities/Source/GameplayAbilities/Private/AbilitySystemComponent.cpp
+++ b/Engine/Plugins/Runtime/GameplayAbilities/Source/GameplayAbilities/Private/AbilitySystemComponent.cpp
@@ -469,6 +469,55 @@ FOnGivenActiveGameplayEffectRemoved& UAbilitySystemComponent::OnAnyGameplayEffec
 	return ActiveGameplayEffects.OnActiveGameplayEffectRemovedDelegate;
 }
 
+void UAbilitySystemComponent::UpdateTagMap_Internal(const FGameplayTagContainer& Container, const int32 CountDelta)
+{
+	// For removal, reorder calls so that FillParentTags is only called once
+	if (CountDelta > 0)
+	{
+		for (auto TagIt = Container.CreateConstIterator(); TagIt; ++TagIt)
+		{
+			const FGameplayTag& Tag = *TagIt;
+			if (GameplayTagCountContainer.UpdateTagCount(Tag, CountDelta))
+			{
+				OnTagUpdated(Tag, true);
+			}
+		}
+	}
+	else if (CountDelta < 0)
+	{
+		// Defer FillParentTags and calling delegates until all Tags have been removed
+		TArray<FGameplayTag> RemovedTags;
+		RemovedTags.Reserve(Container.Num()); // pre-allocate max number (if all are removed)
+		TArray<FDeferredTagChangeDelegate> DeferredTagChangeDelegates;
+
+		for (auto TagIt = Container.CreateConstIterator(); TagIt; ++TagIt)
+		{
+			const FGameplayTag& Tag = *TagIt;
+			if (GameplayTagCountContainer.UpdateTagCount_DeferredParentRemoval(Tag, CountDelta, DeferredTagChangeDelegates))
+			{
+				RemovedTags.Add(Tag);
+			}
+		}
+
+		// now do the work that was deferred
+		if (RemovedTags.Num() > 0)
+		{
+			GameplayTagCountContainer.FillParentTags();
+		}
+
+		for (FDeferredTagChangeDelegate& Delegate : DeferredTagChangeDelegates)
+		{
+			Delegate.Execute();
+		}
+
+		// Notify last in case OnTagUpdated queries this container
+		for (FGameplayTag& Tag : RemovedTags)
+		{
+			OnTagUpdated(Tag, false);
+		}
+	}
+}
+
 FOnGameplayEffectTagCountChanged& UAbilitySystemComponent::RegisterGameplayTagEvent(FGameplayTag Tag, EGameplayTagEventType::Type EventType)
 {
 	return GameplayTagCountContainer.RegisterGameplayTagEvent(Tag, EventType);
diff --git a/Engine/Plugins/Runtime/GameplayAbilities/Source/GameplayAbilities/Private/GameplayEffectTypes.cpp b/Engine/Plugins/Runtime/GameplayAbilities/Source/GameplayAbilities/Private/GameplayEffectTypes.cpp
--- a/Engine/Plugins/Runtime/GameplayAbilities/Source/GameplayAbilities/Private/GameplayEffectTypes.cpp
+++ b/Engine/Plugins/Runtime/GameplayAbilities/Source/GameplayAbilities/Private/GameplayEffectTypes.cpp
@@ -418,7 +418,7 @@ void FGameplayTagCountContainer::Reset()
 	OnAnyTagChangeDelegate.Clear();
 }
 
-bool FGameplayTagCountContainer::UpdateTagMap_Internal(const FGameplayTag& Tag, int32 CountDelta, bool bDeferParentTagsOnRemove)
+bool FGameplayTagCountContainer::UpdateExplicitTags(const FGameplayTag& Tag, const int32 CountDelta, const bool bDeferParentTagsOnRemove)
 {
 	const bool bTagAlreadyExplicitlyExists = ExplicitTags.HasTagExact(Tag);
 
@@ -455,6 +455,11 @@ bool FGameplayTagCountContainer::UpdateTagMap_Internal(const FGameplayTag& Tag,
 		ExplicitTags.RemoveTag(Tag, bDeferParentTagsOnRemove);
 	}
 
+	return true;
+}
+
+bool FGameplayTagCountContainer::GatherTagChangeDelegates(const FGameplayTag& Tag, const int32 CountDelta, TArray<FDeferredTagChangeDelegate>& TagChangeDelegates)
+{
 	// Check if change delegates are required to fire for the tag or any of its parents based on the count change
 	FGameplayTagContainer TagAndParentsContainer = Tag.GetGameplayTagParents();
 	bool CreatedSignificantChange = false;
@@ -462,7 +467,7 @@ bool FGameplayTagCountContainer::UpdateTagMap_Internal(const FGameplayTag& Tag,
 	{
 		const FGameplayTag& CurTag = *CompleteTagIt;
 
-		// Get the current count of the specified tag. NOTE: Stored as a reference, so subsequent changes propogate to the map.
+		// Get the current count of the specified tag. NOTE: Stored as a reference, so subsequent changes propagate to the map.
 		int32& TagCountRef = GameplayTagCountMap.FindOrAdd(CurTag);
 
 		const int32 OldCount = TagCountRef;
@@ -472,24 +477,33 @@ bool FGameplayTagCountContainer::UpdateTagMap_Internal(const FGameplayTag& Tag,
 		TagCountRef = NewTagCount;
 
 		// If a significant change (new addition or total removal) occurred, trigger related delegates
-		bool SignificantChange = (OldCount == 0 || NewTagCount == 0);
+		const bool SignificantChange = (OldCount == 0 || NewTagCount == 0);
 		CreatedSignificantChange |= SignificantChange;
 		if (SignificantChange)
 		{
-			OnAnyTagChangeDelegate.Broadcast(CurTag, NewTagCount);
+			TagChangeDelegates.AddDefaulted();
+			TagChangeDelegates.Last().BindLambda([Delegate = OnAnyTagChangeDelegate, CurTag, NewTagCount]()
+			{
+				Delegate.Broadcast(CurTag, NewTagCount);
+			});
 		}
 
 		FDelegateInfo* DelegateInfo = GameplayTagEventMap.Find(CurTag);
 		if (DelegateInfo)
 		{
-			// Prior to calling OnAnyChange delegate, copy our OnNewOrRemove delegate, since things listening to OnAnyChange could add or remove 
-			// to this map causing our pointer to become invalid.
-			FOnGameplayEffectTagCountChanged OnNewOrRemoveLocalCopy = DelegateInfo->OnNewOrRemove;
+			TagChangeDelegates.AddDefaulted();
+			TagChangeDelegates.Last().BindLambda([Delegate = DelegateInfo->OnAnyChange, CurTag, NewTagCount]()
+			{
+				Delegate.Broadcast(CurTag, NewTagCount);
+			});
 
-			DelegateInfo->OnAnyChange.Broadcast(CurTag, NewTagCount);
 			if (SignificantChange)
 			{
-				OnNewOrRemoveLocalCopy.Broadcast(CurTag, NewTagCount);
+				TagChangeDelegates.AddDefaulted();
+				TagChangeDelegates.Last().BindLambda([Delegate = DelegateInfo->OnNewOrRemove, CurTag, NewTagCount]()
+				{
+					Delegate.Broadcast(CurTag, NewTagCount);
+				});
 			}
 		}
 	}
@@ -497,6 +511,33 @@ bool FGameplayTagCountContainer::UpdateTagMap_Internal(const FGameplayTag& Tag,
 	return CreatedSignificantChange;
 }
 
+bool FGameplayTagCountContainer::UpdateTagMap_Internal(const FGameplayTag& Tag, int32 CountDelta)
+{
+	if (!UpdateExplicitTags(Tag, CountDelta, false))
+	{
+		return false;
+	}
+
+	TArray<FDeferredTagChangeDelegate> DeferredTagChangeDelegates;
+	bool bSignificantChange = GatherTagChangeDelegates(Tag, CountDelta, DeferredTagChangeDelegates);
+	for (FDeferredTagChangeDelegate& Delegate : DeferredTagChangeDelegates)
+	{
+		Delegate.Execute();
+	}
+
+	return bSignificantChange;
+}
+
+bool FGameplayTagCountContainer::UpdateTagMapDeferredParentRemoval_Internal(const FGameplayTag& Tag, int32 CountDelta, TArray<FDeferredTagChangeDelegate>& DeferredTagChangeDelegates)
+{
+	if (!UpdateExplicitTags(Tag, CountDelta, true))
+	{
+		return false;
+	}
+
+	return GatherTagChangeDelegates(Tag, CountDelta, DeferredTagChangeDelegates);
+}
+
 FGameplayTagBlueprintPropertyMap::FGameplayTagBlueprintPropertyMap()
 {
 }
diff --git a/Engine/Plugins/Runtime/GameplayAbilities/Source/GameplayAbilities/Public/AbilitySystemComponent.h b/Engine/Plugins/Runtime/GameplayAbilities/Source/GameplayAbilities/Public/AbilitySystemComponent.h
--- a/Engine/Plugins/Runtime/GameplayAbilities/Source/GameplayAbilities/Public/AbilitySystemComponent.h
+++ b/Engine/Plugins/Runtime/GameplayAbilities/Source/GameplayAbilities/Public/AbilitySystemComponent.h
@@ -508,43 +508,9 @@ class GAMEPLAYABILITIES_API UAbilitySystemComponent : public UGameplayTasksCompo
 	/** Update the number of instances of a given tag and calls callback */
 	FORCEINLINE void UpdateTagMap(const FGameplayTagContainer& Container, int32 CountDelta)
 	{
-		// For removal, reorder calls so that FillParentTags is only called once
-		if (CountDelta > 0)
+		if (!Container.IsEmpty())
 		{
-			for (auto TagIt = Container.CreateConstIterator(); TagIt; ++TagIt)
-			{
-				const FGameplayTag& Tag = *TagIt;
-				if (GameplayTagCountContainer.UpdateTagCount(Tag, CountDelta))
-				{
-					OnTagUpdated(Tag, true);
-				}
-			}
-		}
-		else if (CountDelta < 0)
-		{
-			// Defer FillParentTags until all Tags have been removed
-			TArray<FGameplayTag> RemovedTags;
-			RemovedTags.Reserve(Container.Num()); // pre-allocate max number (if all are removed)
-
-			for (auto TagIt = Container.CreateConstIterator(); TagIt; ++TagIt)
-			{
-				const FGameplayTag& Tag = *TagIt;
-				if (GameplayTagCountContainer.UpdateTagCount(Tag, CountDelta, true))
-				{
-					RemovedTags.Add(Tag);
-				}
-			}
-			
-			if (RemovedTags.Num() > 0)
-			{
-				GameplayTagCountContainer.FillParentTags();
-			}
-
-			// Notify last in case OnTagUpdated queries this container
-			for (FGameplayTag& Tag : RemovedTags)
-			{
-				OnTagUpdated(Tag, false);
-			}
+			UpdateTagMap_Internal(Container, CountDelta);
 		}
 	}
 
@@ -1568,6 +1534,8 @@ protected:
 	const UAttributeSet*	GetAttributeSubobjectChecked(const TSubclassOf<UAttributeSet> AttributeClass) const;
 	const UAttributeSet*	GetOrCreateAttributeSubobject(TSubclassOf<UAttributeSet> AttributeClass);
 
+	void UpdateTagMap_Internal(const FGameplayTagContainer& Container, int32 CountDelta);
+
 	friend struct FActiveGameplayEffect;
 	friend struct FActiveGameplayEffectAction;
 	friend struct FActiveGameplayEffectsContainer;
diff --git a/Engine/Plugins/Runtime/GameplayAbilities/Source/GameplayAbilities/Public/GameplayEffectTypes.h b/Engine/Plugins/Runtime/GameplayAbilities/Source/GameplayAbilities/Public/GameplayEffectTypes.h
--- a/Engine/Plugins/Runtime/GameplayAbilities/Source/GameplayAbilities/Public/GameplayEffectTypes.h
+++ b/Engine/Plugins/Runtime/GameplayAbilities/Source/GameplayAbilities/Public/GameplayEffectTypes.h
@@ -1050,6 +1050,7 @@ namespace EGameplayCueEvent
 }
 
 DECLARE_DELEGATE_OneParam(FOnGameplayAttributeEffectExecuted, struct FGameplayModifierEvaluatedData&);
+DECLARE_DELEGATE(FDeferredTagChangeDelegate);
 
 DECLARE_MULTICAST_DELEGATE_TwoParams(FOnGameplayEffectTagCountChanged, const FGameplayTag, int32);
 DECLARE_MULTICAST_DELEGATE_OneParam(FOnGivenActiveGameplayEffectRemoved, const FActiveGameplayEffect&);
@@ -1197,15 +1198,21 @@ struct GAMEPLAYABILITIES_API FGameplayTagCountContainer
 		if (CountDelta != 0)
 		{
 			bool bUpdatedAny = false;
+			TArray<FDeferredTagChangeDelegate> DeferredTagChangeDelegates;
 			for (auto TagIt = Container.CreateConstIterator(); TagIt; ++TagIt)
 			{
-				bUpdatedAny |= UpdateTagMap_Internal(*TagIt, CountDelta, true);
+				bUpdatedAny |= UpdateTagMapDeferredParentRemoval_Internal(*TagIt, CountDelta, DeferredTagChangeDelegates);
 			}
 
 			if (bUpdatedAny && CountDelta < 0)
 			{
 				ExplicitTags.FillParentTags();
 			}
+
+			for (FDeferredTagChangeDelegate& Delegate : DeferredTagChangeDelegates)
+			{
+				Delegate.Execute();
+			}
 		}
 	}
 	
@@ -1214,15 +1221,34 @@ struct GAMEPLAYABILITIES_API FGameplayTagCountContainer
 	 * 
 	 * @param Tag						Tag to update
 	 * @param CountDelta				Delta of the tag count to apply
-	 * @param bDeferParentTagsOnRemove	Skip calling FillParentTags for performance (must be handled by calling code)
 	 * 
 	 * @return True if tag was *either* added or removed. (E.g., we had the tag and now dont. or didnt have the tag and now we do. We didn't just change the count (1 count -> 2 count would return false).
 	 */
-	FORCEINLINE bool UpdateTagCount(const FGameplayTag& Tag, int32 CountDelta, bool bDeferParentTagsOnRemove=false)
+	FORCEINLINE bool UpdateTagCount(const FGameplayTag& Tag, int32 CountDelta)
+	{
+		if (CountDelta != 0)
+		{
+			return UpdateTagMap_Internal(Tag, CountDelta);
+		}
+
+		return false;
+	}
+
+	/**
+	 * Update the specified tag by the specified delta, potentially causing an additional or removal from the explicit tag list.
+	 * Calling code MUST call FillParentTags followed by executing the returned delegates.
+	 * 
+	 * @param Tag						Tag to update
+	 * @param CountDelta				Delta of the tag count to apply
+	 * @param DeferredTagChangeDelegates		Delegates to be called after this code runs
+	 * 
+	 * @return True if tag was *either* added or removed. (E.g., we had the tag and now dont. or didnt have the tag and now we do. We didn't just change the count (1 count -> 2 count would return false).
+	 */
+	FORCEINLINE bool UpdateTagCount_DeferredParentRemoval(const FGameplayTag& Tag, int32 CountDelta, TArray<FDeferredTagChangeDelegate>& DeferredTagChangeDelegates)
 	{
 		if (CountDelta != 0)
 		{
-			return UpdateTagMap_Internal(Tag, CountDelta, bDeferParentTagsOnRemove);
+			return UpdateTagMapDeferredParentRemoval_Internal(Tag, CountDelta, DeferredTagChangeDelegates);
 		}
 
 		return false;
@@ -1333,7 +1359,16 @@ private:
 	FGameplayTagContainer ExplicitTags;
 
 	/** Internal helper function to adjust the explicit tag list & corresponding maps/delegates/etc. as necessary */
-	bool UpdateTagMap_Internal(const FGameplayTag& Tag, int32 CountDelta, bool bDeferParentTagsOnRemove=false);
+	bool UpdateTagMap_Internal(const FGameplayTag& Tag, int32 CountDelta);
+
+	/** Internal helper function to adjust the explicit tag list & corresponding maps/delegates/etc. as necessary. This does not call FillParentTags or any of the tag change delegates. These delegates are returned and must be executed by the caller. */
+	bool UpdateTagMapDeferredParentRemoval_Internal(const FGameplayTag& Tag, int32 CountDelta, TArray<FDeferredTagChangeDelegate>& DeferredTagChangeDelegates);
+
+	/** Internal helper function to adjust the explicit tag list & corresponding map. */
+	bool UpdateExplicitTags(const FGameplayTag& Tag, int32 CountDelta, bool bDeferParentTagsOnRemove);
+
+	/** Internal helper function to collect the delegates that need to be called when Tag has its count changed by CountDelta. */
+	bool GatherTagChangeDelegates(const FGameplayTag& Tag, int32 CountDelta, TArray<FDeferredTagChangeDelegate>& TagChangeDelegates);
 };
 
 
