Added a new bool to ability specs to track when we have delayed giving an ability that we intend to activate once and then remove. When we add pending abilities we now check this bool and do the appropriate action.

diff --git a/Engine/Plugins/Runtime/GameplayAbilities/Source/GameplayAbilities/Private/AbilitySystemComponent_Abilities.cpp b/Engine/Plugins/Runtime/GameplayAbilities/Source/GameplayAbilities/Private/AbilitySystemComponent_Abilities.cpp
--- a/Engine/Plugins/Runtime/GameplayAbilities/Source/GameplayAbilities/Private/AbilitySystemComponent_Abilities.cpp
+++ b/Engine/Plugins/Runtime/GameplayAbilities/Source/GameplayAbilities/Private/AbilitySystemComponent_Abilities.cpp
@@ -257,7 +257,7 @@ FGameplayAbilitySpecHandle UAbilitySystemComponent::GiveAbility(const FGameplayA
 	return OwnedSpec.Handle;
 }
 
-FGameplayAbilitySpecHandle UAbilitySystemComponent::GiveAbilityAndActivateOnce(const FGameplayAbilitySpec& Spec)
+FGameplayAbilitySpecHandle UAbilitySystemComponent::GiveAbilityAndActivateOnce(FGameplayAbilitySpec& Spec)
 {
 	check(Spec.Ability);
 
@@ -275,6 +275,8 @@ FGameplayAbilitySpecHandle UAbilitySystemComponent::GiveAbilityAndActivateOnce(c
 		return FGameplayAbilitySpecHandle();
 	}
 
+	Spec.bActivateOnce = true;
+
 	FGameplayAbilitySpecHandle AddedAbilityHandle = GiveAbility(Spec);
 
 	FGameplayAbilitySpec* FoundSpec = FindAbilitySpecFromHandle(AddedAbilityHandle);
@@ -584,7 +586,14 @@ void UAbilitySystemComponent::DecrementAbilityListLock()
 		TArray<FGameplayAbilitySpec, TInlineAllocator<2> > LocalPendingAdds(MoveTemp(AbilityPendingAdds));
 		for (FGameplayAbilitySpec& Spec : LocalPendingAdds)
 		{
-			GiveAbility(Spec);
+			if (Spec.bActivateOnce)
+			{
+				GiveAbilityAndActivateOnce(Spec);
+			}
+			else
+			{
+				GiveAbility(Spec);
+			}
 		}
 
 		TArray<FGameplayAbilitySpecHandle, TInlineAllocator<2> > LocalPendingRemoves(MoveTemp(AbilityPendingRemoves));
diff --git a/Engine/Plugins/Runtime/GameplayAbilities/Source/GameplayAbilities/Public/AbilitySystemComponent.h b/Engine/Plugins/Runtime/GameplayAbilities/Source/GameplayAbilities/Public/AbilitySystemComponent.h
--- a/Engine/Plugins/Runtime/GameplayAbilities/Source/GameplayAbilities/Public/AbilitySystemComponent.h
+++ b/Engine/Plugins/Runtime/GameplayAbilities/Source/GameplayAbilities/Public/AbilitySystemComponent.h
@@ -810,7 +810,7 @@ class GAMEPLAYABILITIES_API UAbilitySystemComponent : public UGameplayTasksCompo
 	FGameplayAbilitySpecHandle GiveAbility(const FGameplayAbilitySpec& AbilitySpec);
 
 	/** Grants an ability and attempts to activate it exactly one time, which will cause it to be removed. Only valid on the server! */
-	FGameplayAbilitySpecHandle GiveAbilityAndActivateOnce(const FGameplayAbilitySpec& AbilitySpec);
+	FGameplayAbilitySpecHandle GiveAbilityAndActivateOnce(FGameplayAbilitySpec& AbilitySpec);
 
 	/** Wipes all 'given' abilities. */
 	void ClearAllAbilities();
diff --git a/Engine/Plugins/Runtime/GameplayAbilities/Source/GameplayAbilities/Public/GameplayAbilitySpec.h b/Engine/Plugins/Runtime/GameplayAbilities/Source/GameplayAbilities/Public/GameplayAbilitySpec.h
--- a/Engine/Plugins/Runtime/GameplayAbilities/Source/GameplayAbilities/Public/GameplayAbilitySpec.h
+++ b/Engine/Plugins/Runtime/GameplayAbilities/Source/GameplayAbilities/Public/GameplayAbilitySpec.h
@@ -266,6 +266,10 @@ struct GAMEPLAYABILITIES_API FGameplayAbilitySpec : public FFastArraySerializerI
 	UPROPERTY(NotReplicated)
 	uint8 PendingRemove:1;
 
+	/** This ability should be activated once when it is granted. */
+	UPROPERTY(NotReplicated)
+	uint8 bActivateOnce : 1;
+
 	/** Activation state of this ability. This is not replicated since it needs to be overwritten locally on clients during prediction. */
 	UPROPERTY(NotReplicated)
 	FGameplayAbilityActivationInfo	ActivationInfo;
