Check for World existing before dereferencing the pointer to avoid crashes that could occur here.

diff --git a/Engine/Plugins/Runtime/GameplayAbilities/Source/GameplayAbilities/Private/Abilities/Tasks/AbilityTask_Repeat.cpp b/Engine/Plugins/Runtime/GameplayAbilities/Source/GameplayAbilities/Private/Abilities/Tasks/AbilityTask_Repeat.cpp
--- a/Engine/Plugins/Runtime/GameplayAbilities/Source/GameplayAbilities/Private/Abilities/Tasks/AbilityTask_Repeat.cpp
+++ b/Engine/Plugins/Runtime/GameplayAbilities/Source/GameplayAbilities/Private/Abilities/Tasks/AbilityTask_Repeat.cpp
@@ -41,7 +41,11 @@ void UAbilityTask_Repeat::Activate()
 		PerformAction();
 		if (ActionCounter < ActionPerformancesDesired)
 		{
-			GetWorld()->GetTimerManager().SetTimer(TimerHandle_PerformAction, this, &UAbilityTask_Repeat::PerformAction, TimeBetweenActions, true);
+			UWorld* World = GetWorld();
+			if (World)
+			{
+				World->GetTimerManager().SetTimer(TimerHandle_PerformAction, this, &UAbilityTask_Repeat::PerformAction, TimeBetweenActions, true);
+			}
 		}
 	}
 	else
@@ -56,7 +60,11 @@ void UAbilityTask_Repeat::Activate()
 
 void UAbilityTask_Repeat::OnDestroy(bool AbilityIsEnding)
 {
-	GetWorld()->GetTimerManager().ClearTimer(TimerHandle_PerformAction);
+	UWorld* World = GetWorld();
+	if (World)
+	{
+		World->GetTimerManager().ClearTimer(TimerHandle_PerformAction);
+	}
 
 	Super::OnDestroy(AbilityIsEnding);
 }
