Make a copy of the gameplay event delegates when handling gameplay events. This avoids problems that may occur if this container gets modified by one of the delegates.

diff --git a/Engine/Plugins/Runtime/GameplayAbilities/Source/GameplayAbilities/Private/AbilitySystemComponent_Abilities.cpp b/Engine/Plugins/Runtime/GameplayAbilities/Source/GameplayAbilities/Private/AbilitySystemComponent_Abilities.cpp
--- a/Engine/Plugins/Runtime/GameplayAbilities/Source/GameplayAbilities/Private/AbilitySystemComponent_Abilities.cpp
+++ b/Engine/Plugins/Runtime/GameplayAbilities/Source/GameplayAbilities/Private/AbilitySystemComponent_Abilities.cpp
@@ -1935,10 +1935,10 @@ int32 UAbilitySystemComponent::HandleGameplayEvent(FGameplayTag EventTag, const
 		Delegate->Broadcast(Payload);
 	}
 
-	// Iterate with integers in case it changes due to callbacks
-	for (int32 Index = 0; Index < GameplayEventTagContainerDelegates.Num(); Index++)
+	// Make a copy in case it changes due to callbacks
+	TArray<TPair<FGameplayTagContainer, FGameplayEventTagMulticastDelegate>> LocalGameplayEventTagContainerDelegates = GameplayEventTagContainerDelegates;
+	for (TPair<FGameplayTagContainer, FGameplayEventTagMulticastDelegate>& SearchPair : LocalGameplayEventTagContainerDelegates)
 	{
-		TPair<FGameplayTagContainer, FGameplayEventTagMulticastDelegate>& SearchPair = GameplayEventTagContainerDelegates[Index];
 		if (SearchPair.Key.IsEmpty() || EventTag.MatchesAny(SearchPair.Key))
 		{
 			SearchPair.Value.Broadcast(EventTag, Payload);
@@ -2311,7 +2311,7 @@ void UAbilitySystemComponent::TargetConfirm()
 			}
 		}
 	}
-	SpawnedTargetActors = LeftoverTargetActors;		//These actors declined to confirm targeting, or are allowed to fire multiple times, so keep contact with them.
+	SpawnedTargetActors = MoveTemp(LeftoverTargetActors);		//These actors declined to confirm targeting, or are allowed to fire multiple times, so keep contact with them.
 }
 
 void UAbilitySystemComponent::TargetCancel()
