Added a guard for null GE definition to prevent potential crashes.

diff --git a/Engine/Plugins/Runtime/GameplayAbilities/Source/GameplayAbilities/Private/GameplayEffect.cpp b/Engine/Plugins/Runtime/GameplayAbilities/Source/GameplayAbilities/Private/GameplayEffect.cpp
--- a/Engine/Plugins/Runtime/GameplayAbilities/Source/GameplayAbilities/Private/GameplayEffect.cpp
+++ b/Engine/Plugins/Runtime/GameplayAbilities/Source/GameplayAbilities/Private/GameplayEffect.cpp
@@ -1717,6 +1717,12 @@ void FActiveGameplayEffect::CheckOngoingTagRequirements(const FGameplayTagContai
 
 bool FActiveGameplayEffect::CheckRemovalTagRequirements(const FGameplayTagContainer& OwnerTags, FActiveGameplayEffectsContainer& OwningContainer) const
 {
+	if (Spec.Def == nullptr)
+	{
+		ABILITY_LOG(Error, TEXT("FActiveGameplayEffect::CheckRemovalTagRequirements called with no UGameplayEffect def. (%s)"), *Spec.GetEffectContext().ToString());
+		return false;
+	}
+
 	if (!Spec.Def->RemovalTagRequirements.IsEmpty())
 	{
 		if (Spec.Def->RemovalTagRequirements.RequirementsMet(OwnerTags))
