Fix uninitialized member
#rnx

diff --git a/Engine/Plugins/Runtime/GameplayAbilities/Source/GameplayAbilities/Private/GameplayAbilityTypes.cpp b/Engine/Plugins/Runtime/GameplayAbilities/Source/GameplayAbilities/Private/GameplayAbilityTypes.cpp
--- a/Engine/Plugins/Runtime/GameplayAbilities/Source/GameplayAbilities/Private/GameplayAbilityTypes.cpp
+++ b/Engine/Plugins/Runtime/GameplayAbilities/Source/GameplayAbilities/Private/GameplayAbilityTypes.cpp
@@ -206,6 +206,7 @@ FGameplayAbilitySpec::FGameplayAbilitySpec(UGameplayAbility* InAbility, int32 In
 	, InputPressed(false)
 	, RemoveAfterActivation(false)
 	, PendingRemove(false)
+	, bActivateOnce(false) 
 {
 	Handle.GenerateNewHandle();
 }
@@ -219,6 +220,7 @@ FGameplayAbilitySpec::FGameplayAbilitySpec(TSubclassOf<UGameplayAbility> InAbili
 	, InputPressed(false)
 	, RemoveAfterActivation(false)
 	, PendingRemove(false)
+	, bActivateOnce(false)
 {
 	Handle.GenerateNewHandle();
 }
@@ -231,6 +233,7 @@ FGameplayAbilitySpec::FGameplayAbilitySpec(FGameplayAbilitySpecDef& InDef, int32
 	, InputPressed(false)
 	, RemoveAfterActivation(false)
 	, PendingRemove(false)
+	, bActivateOnce(false)
 {
 	Handle.GenerateNewHandle();
 	InDef.AssignedHandle = Handle;
diff --git a/Engine/Plugins/Runtime/GameplayAbilities/Source/GameplayAbilities/Public/GameplayAbilitySpec.h b/Engine/Plugins/Runtime/GameplayAbilities/Source/GameplayAbilities/Public/GameplayAbilitySpec.h
--- a/Engine/Plugins/Runtime/GameplayAbilities/Source/GameplayAbilities/Public/GameplayAbilitySpec.h
+++ b/Engine/Plugins/Runtime/GameplayAbilities/Source/GameplayAbilities/Public/GameplayAbilitySpec.h
@@ -218,7 +218,7 @@ struct GAMEPLAYABILITIES_API FGameplayAbilitySpec : public FFastArraySerializerI
 	GENERATED_USTRUCT_BODY()
 
 	FGameplayAbilitySpec()
-		: Ability(nullptr), Level(1), InputID(INDEX_NONE), SourceObject(nullptr), ActiveCount(0), InputPressed(false), RemoveAfterActivation(false), PendingRemove(false)
+		: Ability(nullptr), Level(1), InputID(INDEX_NONE), SourceObject(nullptr), ActiveCount(0), InputPressed(false), RemoveAfterActivation(false), PendingRemove(false), bActivateOnce(false)
 	{ }
 
 	/** Version that takes an ability class */
