FGameplayEffectSpec::SetLevel can require the context (depending on data setup) or it will crash. Avoid this by calling SetContext before calling SetLevel during setup. As a result we also need to manually assign the level because SetContext requires it to ran correctly.

diff --git a/Engine/Plugins/Runtime/GameplayAbilities/Source/GameplayAbilities/Private/GameplayEffect.cpp b/Engine/Plugins/Runtime/GameplayAbilities/Source/GameplayAbilities/Private/GameplayEffect.cpp
--- a/Engine/Plugins/Runtime/GameplayAbilities/Source/GameplayAbilities/Private/GameplayEffect.cpp
+++ b/Engine/Plugins/Runtime/GameplayAbilities/Source/GameplayAbilities/Private/GameplayEffect.cpp
@@ -798,8 +798,11 @@ void FGameplayEffectSpec::Initialize(const UGameplayEffect* InDef, const FGamepl
 {
 	Def = InDef;
 	check(Def);	
-	SetLevel(InLevel);
+	// SetContext requires the level to be set before it runs
+	// however, there are code paths in SetLevel that can potentially (depends on game data setup) require the context to be set
+	Level = InLevel;
 	SetContext(InEffectContext);
+	SetLevel(InLevel);
 
 	// Init our ModifierSpecs
 	Modifiers.SetNum(Def->Modifiers.Num());
