Make a copy of the spawned target actor array before iterating over it because some callbacks may modify the array.

diff --git a/Engine/Plugins/Runtime/GameplayAbilities/Source/GameplayAbilities/Private/AbilitySystemComponent_Abilities.cpp b/Engine/Plugins/Runtime/GameplayAbilities/Source/GameplayAbilities/Private/AbilitySystemComponent_Abilities.cpp
--- a/Engine/Plugins/Runtime/GameplayAbilities/Source/GameplayAbilities/Private/AbilitySystemComponent_Abilities.cpp
+++ b/Engine/Plugins/Runtime/GameplayAbilities/Source/GameplayAbilities/Private/AbilitySystemComponent_Abilities.cpp
@@ -2291,8 +2291,10 @@ void UAbilitySystemComponent::LocalInputCancel()
 
 void UAbilitySystemComponent::TargetConfirm()
 {
-	TArray<AGameplayAbilityTargetActor*> LeftoverTargetActors;
-	for (AGameplayAbilityTargetActor* TargetActor : SpawnedTargetActors)
+	// Callbacks may modify the spawned target actor array so iterate over a copy instead
+	TArray<AGameplayAbilityTargetActor*> LocalTargetActors = SpawnedTargetActors;
+	SpawnedTargetActors.Reset();
+	for (AGameplayAbilityTargetActor* TargetActor : LocalTargetActors)
 	{
 		if (TargetActor)
 		{
@@ -2301,30 +2303,30 @@ void UAbilitySystemComponent::TargetConfirm()
 				//TODO: There might not be any cases where this bool is false
 				if (!TargetActor->bDestroyOnConfirmation)
 				{
-					LeftoverTargetActors.Add(TargetActor);
+					SpawnedTargetActors.Add(TargetActor);
 				}
 				TargetActor->ConfirmTargeting();
 			}
 			else
 			{
-				LeftoverTargetActors.Add(TargetActor);
+				SpawnedTargetActors.Add(TargetActor);
 			}
 		}
 	}
-	SpawnedTargetActors = MoveTemp(LeftoverTargetActors);		//These actors declined to confirm targeting, or are allowed to fire multiple times, so keep contact with them.
 }
 
 void UAbilitySystemComponent::TargetCancel()
 {
-	for (AGameplayAbilityTargetActor* TargetActor : SpawnedTargetActors)
+	// Callbacks may modify the spawned target actor array so iterate over a copy instead
+	TArray<AGameplayAbilityTargetActor*> LocalTargetActors = SpawnedTargetActors;
+	SpawnedTargetActors.Reset();
+	for (AGameplayAbilityTargetActor* TargetActor : LocalTargetActors)
 	{
 		if (TargetActor)
 		{
 			TargetActor->CancelTargeting();
 		}
 	}
-
-	SpawnedTargetActors.Empty();
 }
 
 // --------------------------------------------------------------------------
