Rename UAbilitySystemComponent::InternalServerTryActiveAbility to UAbilitySystemComponent::InternalServerTryActivateAbility.

PR #5682: GameplayAbilitySystem: Fixed some more typos (Contributed by [REDACTED_PERSON])

diff --git a/Engine/Plugins/Runtime/GameplayAbilities/Source/GameplayAbilities/Private/AbilitySystemComponent_Abilities.cpp b/Engine/Plugins/Runtime/GameplayAbilities/Source/GameplayAbilities/Private/AbilitySystemComponent_Abilities.cpp
--- a/Engine/Plugins/Runtime/GameplayAbilities/Source/GameplayAbilities/Private/AbilitySystemComponent_Abilities.cpp
+++ b/Engine/Plugins/Runtime/GameplayAbilities/Source/GameplayAbilities/Private/AbilitySystemComponent_Abilities.cpp
@@ -1440,7 +1440,7 @@ bool UAbilitySystemComponent::InternalTryActivateAbility(FGameplayAbilitySpecHan
 
 void UAbilitySystemComponent::ServerTryActivateAbility_Implementation(FGameplayAbilitySpecHandle Handle, bool InputPressed, FPredictionKey PredictionKey)
 {
-	InternalServerTryActiveAbility(Handle, InputPressed, PredictionKey, nullptr);
+	InternalServerTryActivateAbility(Handle, InputPressed, PredictionKey, nullptr);
 }
 
 bool UAbilitySystemComponent::ServerTryActivateAbility_Validate(FGameplayAbilitySpecHandle Handle, bool InputPressed, FPredictionKey PredictionKey)
@@ -1450,7 +1450,7 @@ bool UAbilitySystemComponent::ServerTryActivateAbility_Validate(FGameplayAbility
 
 void UAbilitySystemComponent::ServerTryActivateAbilityWithEventData_Implementation(FGameplayAbilitySpecHandle Handle, bool InputPressed, FPredictionKey PredictionKey, FGameplayEventData TriggerEventData)
 {
-	InternalServerTryActiveAbility(Handle, InputPressed, PredictionKey, &TriggerEventData);
+	InternalServerTryActivateAbility(Handle, InputPressed, PredictionKey, &TriggerEventData);
 }
 
 bool UAbilitySystemComponent::ServerTryActivateAbilityWithEventData_Validate(FGameplayAbilitySpecHandle Handle, bool InputPressed, FPredictionKey PredictionKey, FGameplayEventData TriggerEventData)
@@ -1476,7 +1476,7 @@ void UAbilitySystemComponent::ClientTryActivateAbility_Implementation(FGameplayA
 	InternalTryActivateAbility(Handle);
 }
 
-void UAbilitySystemComponent::InternalServerTryActiveAbility(FGameplayAbilitySpecHandle Handle, bool InputPressed, const FPredictionKey& PredictionKey, const FGameplayEventData* TriggerEventData)
+void UAbilitySystemComponent::InternalServerTryActivateAbility(FGameplayAbilitySpecHandle Handle, bool InputPressed, const FPredictionKey& PredictionKey, const FGameplayEventData* TriggerEventData)
 {
 #if WITH_SERVER_CODE
 #if !(UE_BUILD_SHIPPING || UE_BUILD_TEST)
@@ -1494,7 +1494,7 @@ void UAbilitySystemComponent::InternalServerTryActiveAbility(FGameplayAbilitySpe
 	if (!Spec)
 	{
 		// Can potentially happen in race conditions where client tries to activate ability that is removed server side before it is received.
-		ABILITY_LOG(Display, TEXT("InternalServerTryActiveAbility. Rejecting ClientActivation of ability with invalid SpecHandle!"));
+		ABILITY_LOG(Display, TEXT("InternalServerTryActivateAbility. Rejecting ClientActivation of ability with invalid SpecHandle!"));
 		ClientActivateAbilityFailed(Handle, PredictionKey.Current);
 		return;
 	}
@@ -1522,7 +1522,7 @@ void UAbilitySystemComponent::InternalServerTryActiveAbility(FGameplayAbilitySpe
 	}
 	else
 	{
-		ABILITY_LOG(Display, TEXT("InternalServerTryActiveAbility. Rejecting ClientActivation of %s. InternalTryActivateAbility failed: %s"), *GetNameSafe(Spec->Ability), *InternalTryActivateAbilityFailureTags.ToStringSimple() );
+		ABILITY_LOG(Display, TEXT("InternalServerTryActivateAbility. Rejecting ClientActivation of %s. InternalTryActivateAbility failed: %s"), *GetNameSafe(Spec->Ability), *InternalTryActivateAbilityFailureTags.ToStringSimple() );
 		ClientActivateAbilityFailed(Handle, PredictionKey.Current);
 		Spec->InputPressed = false;
 
diff --git a/Engine/Plugins/Runtime/GameplayAbilities/Source/GameplayAbilities/Private/GameplayModMagnitudeCalculation.cpp b/Engine/Plugins/Runtime/GameplayAbilities/Source/GameplayAbilities/Private/GameplayModMagnitudeCalculation.cpp
--- a/Engine/Plugins/Runtime/GameplayAbilities/Source/GameplayAbilities/Private/GameplayModMagnitudeCalculation.cpp
+++ b/Engine/Plugins/Runtime/GameplayAbilities/Source/GameplayAbilities/Private/GameplayModMagnitudeCalculation.cpp
@@ -34,7 +34,7 @@ bool UGameplayModMagnitudeCalculation::GetCapturedAttributeMagnitude(const FGame
 	}
 	if (CaptureSpec->AttemptCalculateAttributeMagnitude(EvaluationParameters, Magnitude) == false)
 	{
-		ABILITY_LOG(Error, TEXT("GetCapturedAttributeMagnitude unable to calculate Health attribute magnitude."));
+		ABILITY_LOG(Error, TEXT("GetCapturedAttributeMagnitude unable to calculate attribute magnitude."));
 		return false;
 	}
 
diff --git a/Engine/Plugins/Runtime/GameplayAbilities/Source/GameplayAbilities/Public/AbilitySystemComponent.h b/Engine/Plugins/Runtime/GameplayAbilities/Source/GameplayAbilities/Public/AbilitySystemComponent.h
--- a/Engine/Plugins/Runtime/GameplayAbilities/Source/GameplayAbilities/Public/AbilitySystemComponent.h
+++ b/Engine/Plugins/Runtime/GameplayAbilities/Source/GameplayAbilities/Public/AbilitySystemComponent.h
@@ -1422,7 +1422,7 @@ protected:
 	void	ClientActivateAbilitySucceedWithEventData(FGameplayAbilitySpecHandle AbilityToActivate, FPredictionKey PredictionKey, FGameplayEventData TriggerEventData);
 
 	/** Implementation of ServerTryActivateAbility */
-	virtual void InternalServerTryActiveAbility(FGameplayAbilitySpecHandle AbilityToActivate, bool InputPressed, const FPredictionKey& PredictionKey, const FGameplayEventData* TriggerEventData);
+	virtual void InternalServerTryActivateAbility(FGameplayAbilitySpecHandle AbilityToActivate, bool InputPressed, const FPredictionKey& PredictionKey, const FGameplayEventData* TriggerEventData);
 
 	/** Called when a prediction key that played a montage is rejected */
 	void OnPredictiveMontageRejected(UAnimMontage* PredictiveMontage);
diff --git a/Engine/Plugins/Runtime/GameplayAbilities/Source/GameplayAbilities/Public/AttributeSet.h b/Engine/Plugins/Runtime/GameplayAbilities/Source/GameplayAbilities/Public/AttributeSet.h
--- a/Engine/Plugins/Runtime/GameplayAbilities/Source/GameplayAbilities/Public/AttributeSet.h
+++ b/Engine/Plugins/Runtime/GameplayAbilities/Source/GameplayAbilities/Public/AttributeSet.h
@@ -104,7 +104,7 @@ struct GAMEPLAYABILITIES_API FGameplayAttribute
 		return CastChecked<UClass>(Attribute->GetOwner<UObject>());
 	}
 
-	/** Returns true if this is one of the special attributes defined on the bBilitySystemComponent itself */
+	/** Returns true if this is one of the special attributes defined on the AbilitySystemComponent itself */
 	bool IsSystemAttribute() const;
 
 	/** Returns true if the variable associated with Property is of type FGameplayAttributeData or one of its subclasses */
