FScalableFloat editing improvements
- Filtering supports multiple-term search (e.g., "Storm Da" will match StormDamage and DamagePerStorm but not CloudDamage)
- The filtered list of row names can be navigated with the keyboard
- The filter is now automatically reset when the search box is closed
- No longer shows the thumbnail for the curve table reference, decreasing the height of the property widget (this applies to all FCurveTableRowHandle, as the icon for them is never meaningful)

diff --git a/Engine/Plugins/Runtime/GameplayAbilities/Source/GameplayAbilitiesEditor/Private/AttributeDetails.cpp b/Engine/Plugins/Runtime/GameplayAbilities/Source/GameplayAbilitiesEditor/Private/AttributeDetails.cpp
--- a/Engine/Plugins/Runtime/GameplayAbilities/Source/GameplayAbilitiesEditor/Private/AttributeDetails.cpp
+++ b/Engine/Plugins/Runtime/GameplayAbilities/Source/GameplayAbilitiesEditor/Private/AttributeDetails.cpp
@@ -18,6 +18,7 @@
 #include "AbilitySystemGlobals.h"
 #include "Widgets/Input/SSearchBox.h"
 #include "Widgets/Views/SListView.h"
+#include "SListViewSelectorDropdownMenu.h"
 #include "Widgets/Input/STextComboBox.h"
 #include "AbilitySystemComponent.h"
 #include "SGameplayAttributeWidget.h"
@@ -318,6 +319,7 @@ void FScalableFloatDetails::CustomizeHeader( TSharedRef<class IPropertyHandle> S
 				+SHorizontalBox::Slot()
 				.FillWidth(0.23f)
 				.HAlign(HAlign_Fill)
+				.VAlign(VAlign_Center)
 				.Padding(2.f, 0.f, 2.f, 0.f)
 				[
 					SAssignNew(RowNameComboButton, SComboButton)
@@ -451,7 +453,12 @@ TSharedPtr<FString> FScalableFloatDetails::InitWidgetContent()
 	return InitialValue;
 }
 
-UCurveTable * FScalableFloatDetails::GetCurveTable()
+bool FScalableFloatDetails::OnIsSelectableOrNavigableInternal(TSharedPtr<FString> SelectedItem) const
+{
+	return SelectedItem.IsValid() ? DoesPassFilter(SelectedItem) : false;
+}
+
+UCurveTable* FScalableFloatDetails::GetCurveTable()
 {
 	UCurveTable* CurveTable = NULL;
 	CurveTableProperty->GetValue((UObject*&)CurveTable);
@@ -491,48 +498,65 @@ void FScalableFloatDetails::SetPreviewLevel(float NewLevel)
 
 TSharedRef<SWidget> FScalableFloatDetails::GetListContent()
 {
+	FilterTerms.Reset();
+
 	SAssignNew(RowNameComboListView, SListView<TSharedPtr<FString> >)
 		.ListItemsSource(&RowNames)
 		.OnSelectionChanged(this, &FScalableFloatDetails::OnSelectionChanged)
 		.OnGenerateRow(this, &FScalableFloatDetails::HandleRowNameComboBoxGenarateWidget)
-		.SelectionMode(ESelectionMode::Single);
+		.SelectionMode(ESelectionMode::Single)
+		.OnIsSelectableOrNavigable(this, &FScalableFloatDetails::OnIsSelectableOrNavigableInternal);
 
 	if (CurrentSelectedItem.IsValid())
 	{
 		RowNameComboListView->SetSelection(CurrentSelectedItem);
 	}
 
-	return SNew(SVerticalBox)
-		+ SVerticalBox::Slot()
-		.AutoHeight()
-		[
-			SNew(SSearchBox)
-			.OnTextChanged(this, &FScalableFloatDetails::OnFilterTextChanged)
-		]
-		+ SVerticalBox::Slot()
-		.FillHeight(1.f)
+	SearchBoxWidget = SNew(SSearchBox)
+		.OnTextChanged(this, &FScalableFloatDetails::OnFilterTextChanged)
+		.OnTextCommitted(this, &FScalableFloatDetails::OnFilterTextCommitted);
+	
+	RowNameComboButton->SetMenuContentWidgetToFocus(SearchBoxWidget);
+
+	return SNew(SListViewSelectorDropdownMenu<TSharedPtr<FString>>, SearchBoxWidget, RowNameComboListView)
 		[
-			RowNameComboListView.ToSharedRef()
+			SNew(SVerticalBox)
+			+ SVerticalBox::Slot()
+			.AutoHeight()
+			[
+				SearchBoxWidget.ToSharedRef()
+			]
+			+ SVerticalBox::Slot()
+			.FillHeight(1.f)
+			[
+				RowNameComboListView.ToSharedRef()
+			]
 		];
 }
 
 void FScalableFloatDetails::OnSelectionChanged(TSharedPtr<FString> SelectedItem, ESelectInfo::Type SelectInfo)
 {
-	if (SelectedItem.IsValid())
+	if (SelectInfo != ESelectInfo::OnNavigation)
 	{
-		CurrentSelectedItem = SelectedItem;
-		FName NewValue = FName(**SelectedItem);
-		RowNameProperty->SetValue(NewValue);
+		if (SelectedItem.IsValid())
+		{
+			CurrentSelectedItem = SelectedItem;
+			
+			const FName NewValue = FName(**SelectedItem);
+			RowNameProperty->SetValue(NewValue);
 
-		// Close the combo
-		RowNameComboButton->SetIsOpen(false);
+			RowNameComboButton->SetIsOpen(false);
+		}
 	}
 }
 
 TSharedRef<ITableRow> FScalableFloatDetails::HandleRowNameComboBoxGenarateWidget(TSharedPtr<FString> InItem, const TSharedRef<STableViewBase>& OwnerTable)
 {
+	const EVisibility WidgetVisibility = DoesPassFilter(InItem) ? EVisibility::Visible : EVisibility::Collapsed;
+
 	return
 		SNew(STableRow<TSharedPtr<FString>>, OwnerTable)
+		.Visibility(WidgetVisibility)
 		[
 			SNew(STextBlock).Text(FText::FromString(*InItem))
 		];
@@ -567,10 +591,10 @@ FText FScalableFloatDetails::GetRowValuePreviewText() const
 {
 	TArray<const void*> RawPtrs;
 	CurveTableHandleProperty->AccessRawData(RawPtrs);
-	if ( RawPtrs.Num() == 1 && RawPtrs[0] != nullptr )
+	if (RawPtrs.Num() == 1 && RawPtrs[0] != nullptr)
 	{
 		const FCurveTableRowHandle& Curve = *reinterpret_cast<const FCurveTableRowHandle*>(RawPtrs[0]);
-		if ( Curve.CurveTable && Curve.RowName != NAME_None )
+		if (Curve.CurveTable && Curve.RowName != NAME_None)
 		{
 			float Value;
 			ValueProperty->GetValue(Value);
@@ -586,36 +610,73 @@ FText FScalableFloatDetails::GetRowValuePreviewText() const
 	return FText::GetEmpty();
 }
 
-/** Called by Slate when the filter box changes text. */
-void FScalableFloatDetails::OnFilterTextChanged(const FText& InFilterText)
+bool FScalableFloatDetails::DoesPassFilter(const TSharedPtr<FString>& TestStringPtr) const
 {
-	FString CurrentFilterText = InFilterText.ToString();
+	bool bFilterTextMatches = true;
+	const FString& TestString = *TestStringPtr.Get();
 
-	FName RowName;
-	const FPropertyAccess::Result RowResult = RowNameProperty->GetValue(RowName);
-	RowNames.Empty();
+	for (int32 FilterIndex = 0; FilterIndex < FilterTerms.Num() && bFilterTextMatches; ++FilterIndex)
+	{
+		const bool bMatchesTerm = TestString.Contains(FilterTerms[FilterIndex], ESearchCase::IgnoreCase);
+		bFilterTextMatches = bFilterTextMatches && bMatchesTerm;
+	}
 
-	/** Get the properties we wish to work with */
-	UCurveTable* CurveTable = GetCurveTable();
+	return bFilterTextMatches;
+}
 
-	if (CurveTable != NULL)
+void FScalableFloatDetails::OnFilterTextChanged(const FText& InFilterText)
+{
+	// Tokenize the search box text into a set of terms; all of them must be present to pass the filter
+	const FString TrimmedFilterText = FText::TrimPrecedingAndTrailing(InFilterText).ToString();
+	TrimmedFilterText.ParseIntoArray(FilterTerms, TEXT(" "), true);
+
+	TSharedPtr<FString> FirstItemToPassFilter;
+	for (int32 i = 0; i < RowNames.Num(); i++)
 	{
-		/** Extract all the row names from the RowMap */
-		for (TMap<FName, FRealCurve*>::TConstIterator Iterator(CurveTable->GetRowMap()); Iterator; ++Iterator)
+		TSharedPtr<ITableRow> Row = RowNameComboListView->WidgetFromItem(RowNames[i]);
+		if (Row)
 		{
-			/** Create a simple array of the row names */
-			FString RowString = Iterator.Key().ToString();
-			if (CurrentFilterText == TEXT("") || RowString.Contains(CurrentFilterText))
+			const bool bIncludeRow = DoesPassFilter(RowNames[i]);
+
+			Row->AsWidget()->SetVisibility(bIncludeRow ? EVisibility::Visible : EVisibility::Collapsed);
+
+			// See if the selection no longer passes the filter and clear it
+			if ((RowNames[i] == CurrentSelectedItem) && !bIncludeRow)
+			{
+				CurrentSelectedItem.Reset();
+				RowNameComboListView->SetSelection(CurrentSelectedItem, ESelectInfo::OnNavigation);
+			}
+
+			// Remember the first item that passed the filter in case we need to select that
+			if (bIncludeRow && !FirstItemToPassFilter.IsValid())
 			{
-				TSharedRef<FString> RowNameItem = MakeShareable(new FString(RowString));
-				RowNames.Add(RowNameItem);
+				FirstItemToPassFilter = RowNames[i];
 			}
 		}
 	}
+
+	if (!CurrentSelectedItem.IsValid() && FirstItemToPassFilter.IsValid())
+	{
+		CurrentSelectedItem = FirstItemToPassFilter;
+		RowNameComboListView->SetSelection(CurrentSelectedItem, ESelectInfo::OnNavigation);
+	}
+
 	RowNameComboListView->RequestListRefresh();
 }
 
-bool FScalableFloatDetails::IsEditable( ) const
+void FScalableFloatDetails::OnFilterTextCommitted(const FText& InText, ETextCommit::Type CommitInfo)
+{
+	if (CommitInfo == ETextCommit::OnEnter)
+	{
+		TArray<TSharedPtr<FString>> SelectedItems = RowNameComboListView->GetSelectedItems();
+		if (SelectedItems.Num() > 0)
+		{
+			RowNameComboListView->SetSelection(SelectedItems[0]);
+		}
+	}
+}
+
+bool FScalableFloatDetails::IsEditable() const
 {
 	return true;
 }
diff --git a/Engine/Plugins/Runtime/GameplayAbilities/Source/GameplayAbilitiesEditor/Public/AttributeDetails.h b/Engine/Plugins/Runtime/GameplayAbilities/Source/GameplayAbilitiesEditor/Public/AttributeDetails.h
--- a/Engine/Plugins/Runtime/GameplayAbilities/Source/GameplayAbilitiesEditor/Public/AttributeDetails.h
+++ b/Engine/Plugins/Runtime/GameplayAbilities/Source/GameplayAbilitiesEditor/Public/AttributeDetails.h
@@ -13,6 +13,7 @@
 class IDetailLayoutBuilder;
 class IPropertyHandle;
 class SComboButton;
+class SSearchBox;
 
 DECLARE_LOG_CATEGORY_EXTERN(LogAttributeDetails, Log, All);
 
@@ -94,23 +95,32 @@ protected:
 	FText GetRowValuePreviewText() const;
 	TSharedRef<ITableRow> HandleRowNameComboBoxGenarateWidget(TSharedPtr<FString> InItem, const TSharedRef<STableViewBase>& OwnerTable);
 	void OnFilterTextChanged(const FText& InFilterText);
-	
-	class UCurveTable * GetCurveTable();
+	void OnFilterTextCommitted(const FText& InText, ETextCommit::Type CommitInfo);
+
+	bool OnIsSelectableOrNavigableInternal(TSharedPtr<FString> SelectedItem) const;
+
+
+	class UCurveTable* GetCurveTable();
 
 	TSharedPtr<FString> InitWidgetContent();
 
 	void OnCurveTableChanged();
 
+	bool DoesPassFilter(const TSharedPtr<FString>& TestStringPtr) const;
+
 	TSharedPtr<SComboButton> RowNameComboButton;
 	TSharedPtr<SListView<TSharedPtr<FString> > > RowNameComboListView;
 	TSharedPtr<FString> CurrentSelectedItem;
 	TArray<TSharedPtr<FString> > RowNames;
+	TSharedPtr<SSearchBox> SearchBoxWidget;
 
 	TSharedPtr<IPropertyHandle> ValueProperty;
 	TSharedPtr<IPropertyHandle> CurveTableHandleProperty;
 	TSharedPtr<IPropertyHandle> CurveTableProperty;
 	TSharedPtr<IPropertyHandle> RowNameProperty;
 
+	TArray<FString> FilterTerms;
+
 	float PreviewLevel;
 	float MaxPreviewLevel;
 };
diff --git a/Engine/Source/Runtime/Engine/Classes/Engine/CurveTable.h b/Engine/Source/Runtime/Engine/Classes/Engine/CurveTable.h
--- a/Engine/Source/Runtime/Engine/Classes/Engine/CurveTable.h
+++ b/Engine/Source/Runtime/Engine/Classes/Engine/CurveTable.h
@@ -237,7 +237,7 @@ struct ENGINE_API FCurveTableRowHandle
 	{ }
 
 	/** Pointer to table we want a row from */
-	UPROPERTY(EditAnywhere, BlueprintReadWrite, Category=CurveTableRowHandle)
+	UPROPERTY(EditAnywhere, BlueprintReadWrite, Category=CurveTableRowHandle, meta=(DisplayThumbnail="false"))
 	const class UCurveTable*	CurveTable;
 
 	/** Name of row in the table that we want */
