Ability system global target data cache;

* Add UAbilitySystemGlobals::InitTargetDataScriptStructCache().  This allows other modules to re-initialize the target data script cache if the module defines a new target data class.  This was needed since the cache is first initialized before any plugin modules are loaded.

diff --git a/Engine/Plugins/Runtime/GameplayAbilities/Source/GameplayAbilities/Private/AbilitySystemGlobals.cpp b/Engine/Plugins/Runtime/GameplayAbilities/Source/GameplayAbilities/Private/AbilitySystemGlobals.cpp
--- a/Engine/Plugins/Runtime/GameplayAbilities/Source/GameplayAbilities/Private/AbilitySystemGlobals.cpp
+++ b/Engine/Plugins/Runtime/GameplayAbilities/Source/GameplayAbilities/Private/AbilitySystemGlobals.cpp
@@ -47,7 +47,7 @@ void UAbilitySystemGlobals::InitGlobalData()
 	GetGameplayTagResponseTable();
 	InitGlobalTags();
 
-	TargetDataStructCache.InitForType(FGameplayAbilityTargetData::StaticStruct());
+	InitTargetDataScriptStructCache();
 
 	// Register for PreloadMap so cleanup can occur on map transitions
 	FCoreUObjectDelegates::PreLoadMap.AddUObject(this, &UAbilitySystemGlobals::HandlePreLoadMap);
@@ -231,6 +231,11 @@ UFunction* UAbilitySystemGlobals::GetGameplayCueFunction(const FGameplayTag& Chi
 	return nullptr;
 }
 
+void UAbilitySystemGlobals::InitTargetDataScriptStructCache()
+{
+	TargetDataStructCache.InitForType(FGameplayAbilityTargetData::StaticStruct());
+}
+
 // --------------------------------------------------------------------
 
 void UAbilitySystemGlobals::InitGameplayCueParameters(FGameplayCueParameters& CueParameters, const FGameplayEffectSpecForRPC& Spec)
@@ -508,6 +513,8 @@ void UAbilitySystemGlobals::NonShipping_ApplyGlobalAbilityScaler_Duration(float&
 
 void FNetSerializeScriptStructCache::InitForType(UScriptStruct* InScriptStruct)
 {
+	ScriptStructs.Reset();
+
 	// Find all script structs of this type and add them to the list
 	// (not sure of a better way to do this but it should only happen once at startup)
 	for (TObjectIterator<UScriptStruct> It; It; ++It)
diff --git a/Engine/Plugins/Runtime/GameplayAbilities/Source/GameplayAbilities/Public/AbilitySystemGlobals.h b/Engine/Plugins/Runtime/GameplayAbilities/Source/GameplayAbilities/Public/AbilitySystemGlobals.h
--- a/Engine/Plugins/Runtime/GameplayAbilities/Source/GameplayAbilities/Public/AbilitySystemGlobals.h
+++ b/Engine/Plugins/Runtime/GameplayAbilities/Source/GameplayAbilities/Public/AbilitySystemGlobals.h
@@ -248,6 +248,8 @@ class GAMEPLAYABILITIES_API UAbilitySystemGlobals : public UObject
 		}
 	}
 
+	void InitTargetDataScriptStructCache();
+
 	/** Initialize GameplayCue Parameters */
 	virtual void InitGameplayCueParameters(FGameplayCueParameters& CueParameters, const FGameplayEffectSpecForRPC &Spec);
 	virtual void InitGameplayCueParameters_GESpec(FGameplayCueParameters& CueParameters, const FGameplayEffectSpec &Spec);
