GameplayAbility "IsLocallyControlled()" and "HasAuthority()" are now callable from blueprint.

diff --git a/Engine/Plugins/Runtime/GameplayAbilities/Source/GameplayAbilities/Private/Abilities/GameplayAbility.cpp b/Engine/Plugins/Runtime/GameplayAbilities/Source/GameplayAbilities/Private/Abilities/GameplayAbility.cpp
--- a/Engine/Plugins/Runtime/GameplayAbilities/Source/GameplayAbilities/Private/Abilities/GameplayAbility.cpp
+++ b/Engine/Plugins/Runtime/GameplayAbilities/Source/GameplayAbilities/Private/Abilities/GameplayAbility.cpp
@@ -1607,6 +1607,11 @@ bool UGameplayAbility::HasAuthority(const FGameplayAbilityActivationInfo* Activa
 	return (ActivationInfo->ActivationMode == EGameplayAbilityActivationMode::Authority);
 }
 
+bool UGameplayAbility::K2_HasAuthority() const
+{
+	return HasAuthority(&CurrentActivationInfo);
+}
+
 bool UGameplayAbility::HasAuthorityOrPredictionKey(const FGameplayAbilityActorInfo* ActorInfo, const FGameplayAbilityActivationInfo* ActivationInfo) const
 {
 	return ActorInfo->AbilitySystemComponent->HasAuthorityOrPredictionKey(ActivationInfo);
diff --git a/Engine/Plugins/Runtime/GameplayAbilities/Source/GameplayAbilities/Public/Abilities/GameplayAbility.h b/Engine/Plugins/Runtime/GameplayAbilities/Source/GameplayAbilities/Public/Abilities/GameplayAbility.h
--- a/Engine/Plugins/Runtime/GameplayAbilities/Source/GameplayAbilities/Public/Abilities/GameplayAbility.h
+++ b/Engine/Plugins/Runtime/GameplayAbilities/Source/GameplayAbilities/Public/Abilities/GameplayAbility.h
@@ -237,11 +237,15 @@ public:
 	bool IsForRemoteClient() const;
 
 	/** True if the owning actor is locally controlled, true in single player */
+	UFUNCTION(BlueprintCallable, BlueprintPure = false, Category = Ability, Meta = (ExpandBoolAsExecs = "ReturnValue"))
 	bool IsLocallyControlled() const;
 
 	/** True if this is the server or single player */
 	bool HasAuthority(const FGameplayAbilityActivationInfo* ActivationInfo) const;
 
+	UFUNCTION(BlueprintCallable, BlueprintPure = false, Category = Ability, DisplayName = "HasAuthority", Meta = (ExpandBoolAsExecs = "ReturnValue"))
+	bool K2_HasAuthority() const;
+
 	/** True if we are authority or we have a valid prediciton key that is expected to work */
 	bool HasAuthorityOrPredictionKey(const FGameplayAbilityActorInfo* ActorInfo, const FGameplayAbilityActivationInfo* ActivationInfo) const;
 
