Fix a bug where stacking GEs that did not reset the duration on additional instances of the effect being applied and with set by caller durations would only have the duration correctly set for the first instance on the stack. All other GE specs in the stack would have a duration of 1 second.

diff --git a/Engine/Plugins/Runtime/GameplayAbilities/Source/GameplayAbilities/Private/GameplayEffect.cpp b/Engine/Plugins/Runtime/GameplayAbilities/Source/GameplayAbilities/Private/GameplayEffect.cpp
--- a/Engine/Plugins/Runtime/GameplayAbilities/Source/GameplayAbilities/Private/GameplayEffect.cpp
+++ b/Engine/Plugins/Runtime/GameplayAbilities/Source/GameplayAbilities/Private/GameplayEffect.cpp
@@ -2941,49 +2941,46 @@ FActiveGameplayEffect* FActiveGameplayEffectsContainer::ApplyGameplayEffectSpec(
 	// Register Source and Target non snapshot capture delegates here
 	AppliedEffectSpec.CapturedRelevantAttributes.RegisterLinkedAggregatorCallbacks(AppliedActiveGE->Handle);
 	
-	if (bSetDuration)
+	// Re-calculate the duration, as it could rely on target captured attributes
+	float DefCalcDuration = 0.f;
+	if (AppliedEffectSpec.AttemptCalculateDurationFromDef(DefCalcDuration))
 	{
-		// Re-calculate the duration, as it could rely on target captured attributes
-		float DefCalcDuration = 0.f;
-		if (AppliedEffectSpec.AttemptCalculateDurationFromDef(DefCalcDuration))
-		{
-			AppliedEffectSpec.SetDuration(DefCalcDuration, false);
-		}
-		else if (AppliedEffectSpec.Def->DurationMagnitude.GetMagnitudeCalculationType() == EGameplayEffectMagnitudeCalculation::SetByCaller)
-		{
-			AppliedEffectSpec.Def->DurationMagnitude.AttemptCalculateMagnitude(AppliedEffectSpec, AppliedEffectSpec.Duration);
-		}
+		AppliedEffectSpec.SetDuration(DefCalcDuration, false);
+	}
+	else if (AppliedEffectSpec.Def->DurationMagnitude.GetMagnitudeCalculationType() == EGameplayEffectMagnitudeCalculation::SetByCaller)
+	{
+		AppliedEffectSpec.Def->DurationMagnitude.AttemptCalculateMagnitude(AppliedEffectSpec, AppliedEffectSpec.Duration);
+	}
 
-		const float DurationBaseValue = AppliedEffectSpec.GetDuration();
+	const float DurationBaseValue = AppliedEffectSpec.GetDuration();
 
-		// Calculate Duration mods if we have a real duration
-		if (DurationBaseValue > 0.f)
-		{
-			float FinalDuration = AppliedEffectSpec.CalculateModifiedDuration();
+	// Calculate Duration mods if we have a real duration
+	if (DurationBaseValue > 0.f)
+	{
+		float FinalDuration = AppliedEffectSpec.CalculateModifiedDuration();
 
-			// We cannot mod ourselves into an instant or infinite duration effect
-			if (FinalDuration <= 0.f)
-			{
-				ABILITY_LOG(Error, TEXT("GameplayEffect %s Duration was modified to %.2f. Clamping to 0.1s duration."), *AppliedEffectSpec.Def->GetName(), FinalDuration);
-				FinalDuration = 0.1f;
-			}
+		// We cannot mod ourselves into an instant or infinite duration effect
+		if (FinalDuration <= 0.f)
+		{
+			ABILITY_LOG(Error, TEXT("GameplayEffect %s Duration was modified to %.2f. Clamping to 0.1s duration."), *AppliedEffectSpec.Def->GetName(), FinalDuration);
+			FinalDuration = 0.1f;
+		}
 
-			AppliedEffectSpec.SetDuration(FinalDuration, true);
+		AppliedEffectSpec.SetDuration(FinalDuration, true);
 
-			// ABILITY_LOG(Warning, TEXT("SetDuration for %s. Base: %.2f, Final: %.2f"), *NewEffect.Spec.Def->GetName(), DurationBaseValue, FinalDuration);
+		// ABILITY_LOG(Warning, TEXT("SetDuration for %s. Base: %.2f, Final: %.2f"), *NewEffect.Spec.Def->GetName(), DurationBaseValue, FinalDuration);
 
-			// Register duration callbacks with the timer manager
-			if (Owner)
+		// Register duration callbacks with the timer manager
+		if (Owner && bSetDuration)
+		{
+			FTimerManager& TimerManager = Owner->GetWorld()->GetTimerManager();
+			FTimerDelegate Delegate = FTimerDelegate::CreateUObject(Owner, &UAbilitySystemComponent::CheckDurationExpired, AppliedActiveGE->Handle);
+			TimerManager.SetTimer(AppliedActiveGE->DurationHandle, Delegate, FinalDuration, false);
+			if (!ensureMsgf(AppliedActiveGE->DurationHandle.IsValid(), TEXT("Invalid Duration Handle after attempting to set duration for GE %s @ %.2f"), 
+				*AppliedActiveGE->GetDebugString(), FinalDuration))
 			{
-				FTimerManager& TimerManager = Owner->GetWorld()->GetTimerManager();
-				FTimerDelegate Delegate = FTimerDelegate::CreateUObject(Owner, &UAbilitySystemComponent::CheckDurationExpired, AppliedActiveGE->Handle);
-				TimerManager.SetTimer(AppliedActiveGE->DurationHandle, Delegate, FinalDuration, false);
-				if (!ensureMsgf(AppliedActiveGE->DurationHandle.IsValid(), TEXT("Invalid Duration Handle after attempting to set duration for GE %s @ %.2f"), 
-					*AppliedActiveGE->GetDebugString(), FinalDuration))
-				{
-					// Force this off next frame
-					TimerManager.SetTimerForNextTick(Delegate);
-				}
+				// Force this off next frame
+				TimerManager.SetTimerForNextTick(Delegate);
 			}
 		}
 	}
diff --git a/Engine/Plugins/Runtime/GameplayAbilities/Source/GameplayAbilities/Private/GameplayEffectTests.cpp b/Engine/Plugins/Runtime/GameplayAbilities/Source/GameplayAbilities/Private/GameplayEffectTests.cpp
--- a/Engine/Plugins/Runtime/GameplayAbilities/Source/GameplayAbilities/Private/GameplayEffectTests.cpp
+++ b/Engine/Plugins/Runtime/GameplayAbilities/Source/GameplayAbilities/Private/GameplayEffectTests.cpp
@@ -227,6 +227,93 @@ public: // the tests
 		// TODO: test that the effect is no longer applied
 	}
 
+	void Test_StackLimit()
+	{
+		const float Duration = 10.0f;
+		const float HalfDuration = Duration / 2.f;
+		const float ChangePerGE = 5.f;
+		const uint32 StackLimit = 2;
+		const float StartingAttributeValue = DestComponent->GetSet<UAbilitySystemTestAttributeSet>()->StackingAttribute1;
+
+		// Apply one copy of the stacking GE
+		CONSTRUCT_CLASS(UGameplayEffect, StackingEffect);
+		AddModifier(StackingEffect, GET_FIELD_CHECKED(UAbilitySystemTestAttributeSet, StackingAttribute1), EGameplayModOp::Additive, FScalableFloat(ChangePerGE));
+		StackingEffect->DurationPolicy = EGameplayEffectDurationType::HasDuration;
+		StackingEffect->DurationMagnitude = FGameplayEffectModifierMagnitude(FScalableFloat(Duration));
+		StackingEffect->StackLimitCount = StackLimit;
+		StackingEffect->StackingType = EGameplayEffectStackingType::AggregateByTarget;
+		StackingEffect->StackDurationRefreshPolicy = EGameplayEffectStackingDurationPolicy::NeverRefresh;
+		StackingEffect->StackExpirationPolicy = EGameplayEffectStackingExpirationPolicy::ClearEntireStack;
+
+		// Apply the GE StackLimit + 1 times
+		for (uint32 Idx = 0; Idx <= StackLimit; ++Idx)
+		{
+			SourceComponent->ApplyGameplayEffectToTarget(StackingEffect, DestComponent, 1.f);
+		}
+
+		Test->TestEqual(TEXT("Stacking GEs"), DestComponent->GetSet<UAbilitySystemTestAttributeSet>()->StackingAttribute1, StartingAttributeValue + (StackLimit * ChangePerGE));
+	}
+
+	void Test_SetByCallerStackingDuration()
+	{
+		const float Duration = 10.0f;
+		const float HalfDuration = Duration / 2.f;
+		const float ChangePerGE = 5.f;
+		const uint32 StackLimit = 2;
+		const float StartingAttributeValue = DestComponent->GetSet<UAbilitySystemTestAttributeSet>()->StackingAttribute1;
+
+		const FName DurationName(TEXT("Duration"));
+		FSetByCallerFloat SetByCallerDuration;
+		SetByCallerDuration.DataName = DurationName;
+
+		// Apply one copy of the stacking GE
+		CONSTRUCT_CLASS(UGameplayEffect, StackingEffect);
+		AddModifier(StackingEffect, GET_FIELD_CHECKED(UAbilitySystemTestAttributeSet, StackingAttribute1), EGameplayModOp::Additive, FScalableFloat(ChangePerGE));
+		StackingEffect->DurationPolicy = EGameplayEffectDurationType::HasDuration;
+		StackingEffect->DurationMagnitude = FGameplayEffectModifierMagnitude(SetByCallerDuration);
+		StackingEffect->StackLimitCount = StackLimit;
+		StackingEffect->StackingType = EGameplayEffectStackingType::AggregateByTarget;
+		StackingEffect->StackDurationRefreshPolicy = EGameplayEffectStackingDurationPolicy::NeverRefresh;
+		StackingEffect->StackExpirationPolicy = EGameplayEffectStackingExpirationPolicy::RemoveSingleStackAndRefreshDuration;
+
+		// create a spec, set the magnitude and apply the GE
+		{
+			FGameplayEffectSpec	Spec(StackingEffect, FGameplayEffectContextHandle(), 1.f);
+			Spec.SetSetByCallerMagnitude(DurationName, Duration);
+			SourceComponent->ApplyGameplayEffectSpecToTarget(Spec, DestComponent, FPredictionKey());
+		}
+
+		// Tick to partway through the GE's duration and apply a second copy of the GE
+		TickWorld(HalfDuration);
+
+		Test->TestEqual(TEXT("Stacking GEs"), DestComponent->GetSet<UAbilitySystemTestAttributeSet>()->StackingAttribute1, StartingAttributeValue + ChangePerGE);
+
+		// apply second copy of GE
+		{
+			FGameplayEffectSpec	Spec(StackingEffect, FGameplayEffectContextHandle(), 1.f);
+			Spec.SetSetByCallerMagnitude(DurationName, Duration);
+			SourceComponent->ApplyGameplayEffectSpecToTarget(Spec, DestComponent, FPredictionKey());
+		}
+
+		Test->TestEqual(TEXT("Stacking GEs"), DestComponent->GetSet<UAbilitySystemTestAttributeSet>()->StackingAttribute1, StartingAttributeValue + (2 * ChangePerGE));
+
+		// Tick to just after the first GE should have expired
+		TickWorld(HalfDuration + 0.1f);
+
+		// we should have removed one copy and still have the second copy
+		Test->TestEqual(TEXT("Stacking GEs"), DestComponent->GetSet<UAbilitySystemTestAttributeSet>()->StackingAttribute1, StartingAttributeValue + ChangePerGE);
+
+		// check again near the end of the remaining GE's duration
+		TickWorld(Duration - 0.2f);
+
+		Test->TestEqual(TEXT("Stacking GEs"), DestComponent->GetSet<UAbilitySystemTestAttributeSet>()->StackingAttribute1, StartingAttributeValue + ChangePerGE);
+
+		// Tick to a point just after where the remaining GE should have been removed
+		TickWorld(0.2f);
+
+		Test->TestEqual(TEXT("Stacking GEs"), DestComponent->GetSet<UAbilitySystemTestAttributeSet>()->StackingAttribute1, StartingAttributeValue);
+	}
+
 private: // test helpers
 
 	template<typename MODIFIER_T>
@@ -283,6 +370,8 @@ public:
 		ADD_TEST(Test_InstantDamageRemap);
 		ADD_TEST(Test_ManaBuff);
 		ADD_TEST(Test_PeriodicDamage);
+		ADD_TEST(Test_StackLimit);
+		ADD_TEST(Test_SetByCallerStackingDuration);
 	}
 
 	virtual uint32 GetTestFlags() const override { return EAutomationTestFlags::EditorContext | EAutomationTestFlags::EngineFilter; }
