Add "Enable Gravity" parameter to ApplyRootMotionConstantForce ability task node for override constant forces that want similar behavior as animation root motion

#rnx

diff --git a/Engine/Plugins/Runtime/GameplayAbilities/Source/GameplayAbilities/Private/Abilities/Tasks/AbilityTask_ApplyRootMotionConstantForce.cpp b/Engine/Plugins/Runtime/GameplayAbilities/Source/GameplayAbilities/Private/Abilities/Tasks/AbilityTask_ApplyRootMotionConstantForce.cpp
--- a/Engine/Plugins/Runtime/GameplayAbilities/Source/GameplayAbilities/Private/Abilities/Tasks/AbilityTask_ApplyRootMotionConstantForce.cpp
+++ b/Engine/Plugins/Runtime/GameplayAbilities/Source/GameplayAbilities/Private/Abilities/Tasks/AbilityTask_ApplyRootMotionConstantForce.cpp
@@ -24,7 +24,8 @@ UAbilityTask_ApplyRootMotionConstantForce* UAbilityTask_ApplyRootMotionConstantF
 	UCurveFloat* StrengthOverTime,
 	ERootMotionFinishVelocityMode VelocityOnFinishMode, 
 	FVector SetVelocityOnFinish, 
-	float ClampVelocityOnFinish
+	float ClampVelocityOnFinish,
+	bool bEnableGravity
 )
 {
 	UAbilitySystemGlobals::NonShipping_ApplyGlobalAbilityScaler_Duration(Duration);
@@ -40,6 +41,7 @@ UAbilityTask_ApplyRootMotionConstantForce* UAbilityTask_ApplyRootMotionConstantF
 	MyTask->FinishVelocityMode = VelocityOnFinishMode;
 	MyTask->FinishSetVelocity = SetVelocityOnFinish;
 	MyTask->FinishClampVelocity = ClampVelocityOnFinish;
+	MyTask->bEnableGravity = bEnableGravity;
 	MyTask->SharedInitAndApply();
 
 	return MyTask;
@@ -66,6 +68,10 @@ void UAbilityTask_ApplyRootMotionConstantForce::SharedInitAndApply()
 			ConstantForce->FinishVelocityParams.Mode = FinishVelocityMode;
 			ConstantForce->FinishVelocityParams.SetVelocity = FinishSetVelocity;
 			ConstantForce->FinishVelocityParams.ClampVelocity = FinishClampVelocity;
+			if (bEnableGravity)
+			{
+				ConstantForce->Settings.SetFlag(ERootMotionSourceSettingsFlags::IgnoreZAccumulate);
+			}
 			RootMotionSourceID = MovementComponent->ApplyRootMotionSource(ConstantForce);
 
 			if (Ability)
@@ -128,6 +134,7 @@ void UAbilityTask_ApplyRootMotionConstantForce::GetLifetimeReplicatedProps(TArra
 	DOREPLIFETIME(UAbilityTask_ApplyRootMotionConstantForce, Duration);
 	DOREPLIFETIME(UAbilityTask_ApplyRootMotionConstantForce, bIsAdditive);
 	DOREPLIFETIME(UAbilityTask_ApplyRootMotionConstantForce, StrengthOverTime);
+	DOREPLIFETIME(UAbilityTask_ApplyRootMotionConstantForce, bEnableGravity);
 }
 
 void UAbilityTask_ApplyRootMotionConstantForce::PreDestroyFromReplication()
diff --git a/Engine/Plugins/Runtime/GameplayAbilities/Source/GameplayAbilities/Public/Abilities/Tasks/AbilityTask_ApplyRootMotionConstantForce.h b/Engine/Plugins/Runtime/GameplayAbilities/Source/GameplayAbilities/Public/Abilities/Tasks/AbilityTask_ApplyRootMotionConstantForce.h
--- a/Engine/Plugins/Runtime/GameplayAbilities/Source/GameplayAbilities/Public/Abilities/Tasks/AbilityTask_ApplyRootMotionConstantForce.h
+++ b/Engine/Plugins/Runtime/GameplayAbilities/Source/GameplayAbilities/Public/Abilities/Tasks/AbilityTask_ApplyRootMotionConstantForce.h
@@ -39,7 +39,8 @@ class GAMEPLAYABILITIES_API UAbilityTask_ApplyRootMotionConstantForce : public U
 		UCurveFloat* StrengthOverTime,
 		ERootMotionFinishVelocityMode VelocityOnFinishMode,
 		FVector SetVelocityOnFinish,
-		float ClampVelocityOnFinish
+		float ClampVelocityOnFinish,
+		bool bEnableGravity
 	);
 
 	/** Tick function for this task, if bTickingTask == true */
@@ -74,4 +75,8 @@ protected:
 	 */
 	UPROPERTY(Replicated)
 	UCurveFloat* StrengthOverTime;
+
+	UPROPERTY(Replicated)
+	bool bEnableGravity;
+
 };
