Blueprints: Prevent gameplay attribute pins from displaying theiir default value when the pin has connections

diff --git a/Engine/Plugins/Runtime/GameplayAbilities/Source/GameplayAbilitiesEditor/Private/SGameplayAttributeGraphPin.cpp b/Engine/Plugins/Runtime/GameplayAbilities/Source/GameplayAbilitiesEditor/Private/SGameplayAttributeGraphPin.cpp
--- a/Engine/Plugins/Runtime/GameplayAbilities/Source/GameplayAbilitiesEditor/Private/SGameplayAttributeGraphPin.cpp
+++ b/Engine/Plugins/Runtime/GameplayAbilities/Source/GameplayAbilitiesEditor/Private/SGameplayAttributeGraphPin.cpp
@@ -34,6 +34,8 @@ TSharedRef<SWidget>	SGameplayAttributeGraphPin::GetDefaultValueWidget()
 			SNew(SGameplayAttributeWidget)
 			.OnAttributeChanged(this, &SGameplayAttributeGraphPin::OnAttributeChanged)
 			.DefaultProperty(DefaultAttribute.GetUProperty())
+			.Visibility(this, &SGraphPin::GetDefaultValueVisibility)
+			.IsEnabled(this, &SGameplayAttributeGraphPin::GetDefaultValueIsEnabled)
 		];
 }
 
diff --git a/Engine/Plugins/Runtime/GameplayAbilities/Source/GameplayAbilitiesEditor/Private/SGameplayAttributeGraphPin.h b/Engine/Plugins/Runtime/GameplayAbilities/Source/GameplayAbilitiesEditor/Private/SGameplayAttributeGraphPin.h
--- a/Engine/Plugins/Runtime/GameplayAbilities/Source/GameplayAbilitiesEditor/Private/SGameplayAttributeGraphPin.h
+++ b/Engine/Plugins/Runtime/GameplayAbilities/Source/GameplayAbilitiesEditor/Private/SGameplayAttributeGraphPin.h
@@ -24,5 +24,8 @@ public:
 	FProperty* LastSelectedProperty;
 
 private:
-
+	bool GetDefaultValueIsEnabled() const
+	{
+		return !GraphPinObj->bDefaultValueIsReadOnly;
+	}
 };
