Because of the level of precision an animation section can be setup with, had to bring back the full serialization of the montage position over a less precise packed version.

#rnx

diff --git a/Engine/Plugins/Runtime/GameplayAbilities/Source/GameplayAbilities/Private/GameplayAbilityTypes.cpp b/Engine/Plugins/Runtime/GameplayAbilities/Source/GameplayAbilities/Private/GameplayAbilityTypes.cpp
--- a/Engine/Plugins/Runtime/GameplayAbilities/Source/GameplayAbilities/Private/GameplayAbilityTypes.cpp
+++ b/Engine/Plugins/Runtime/GameplayAbilities/Source/GameplayAbilities/Private/GameplayAbilityTypes.cpp
@@ -388,17 +388,11 @@ bool FGameplayAbilityRepAnimMontage::NetSerialize(FArchive& Ar, class UPackageMa
 		SectionIdToPlay = 0;
 		SkipPositionCorrection = false;
 
-		if (Ar.IsSaving())
-		{
-			uint32 PackedPosition = static_cast<uint32>(Position * 100.0f);
-			Ar.SerializeIntPacked(PackedPosition);
-		}
-		else if (Ar.IsLoading())
-		{
-			uint32 PackedPosition = 0;
-			Ar.SerializeIntPacked(PackedPosition);
-			Position = static_cast<float>(PackedPosition) / 100.0f;
-		}
+		// @note: section frames have such a high amount of precision they use, when
+		// removing some of the position precision and packing it into a uint32 caused
+		// issues where ability code would pick the end of a previous section instead of
+		// the start of a new section. For now serializing the full position again.
+		Ar << Position;
 	}
 	else
 	{
