Preventing problems related to bad data getting stored in SpawnedAttributes.

Made SpawnedAttributes transient so it won't save data that can become stale and incorrect.
Added null checks to prevent any currently saved stale data from propagating.
Added null checks to everywhere we iterate over SpawnedAttributes. This is necessary despite the other changes due to replication.
Deprecated AddDefaultSubobjectSet and replaced it with AddAttributeSetSubobject with comments that should make correct use clearer.

diff --git a/Engine/Plugins/Runtime/GameplayAbilities/Source/GameplayAbilities/Private/AbilitySystemComponent.cpp b/Engine/Plugins/Runtime/GameplayAbilities/Source/GameplayAbilities/Private/AbilitySystemComponent.cpp
--- a/Engine/Plugins/Runtime/GameplayAbilities/Source/GameplayAbilities/Private/AbilitySystemComponent.cpp
+++ b/Engine/Plugins/Runtime/GameplayAbilities/Source/GameplayAbilities/Private/AbilitySystemComponent.cpp
@@ -129,11 +129,16 @@ void UAbilitySystemComponent::GetAllAttributes(OUT TArray<FGameplayAttribute>& A
 {
 	for (UAttributeSet* Set : SpawnedAttributes)
 	{
-		for ( TFieldIterator<FProperty> It(Set->GetClass()); It; ++It)
+		if (!Set)
+		{
+			continue;
+		}
+
+		for (TFieldIterator<FProperty> It(Set->GetClass()); It; ++It)
 		{
 			if (FFloatProperty* FloatProperty = CastField<FFloatProperty>(*It))
 			{
-				Attributes.Push( FGameplayAttribute(FloatProperty) );
+				Attributes.Push(FGameplayAttribute(FloatProperty));
 			}
 			else if (FGameplayAttribute::IsGameplayAttributeDataProperty(*It))
 			{
@@ -2458,6 +2463,11 @@ void UAbilitySystemComponent::Debug_Internal(FAbilitySystemComponentDebugInfo& I
 		}
 		for (UAttributeSet* Set : SpawnedAttributes)
 		{
+			if (!Set)
+			{
+				continue;
+			}
+
 			for (TFieldIterator<FProperty> It(Set->GetClass()); It; ++It)
 			{
 				FGameplayAttribute	Attribute(*It);
diff --git a/Engine/Plugins/Runtime/GameplayAbilities/Source/GameplayAbilities/Private/AbilitySystemComponent_Abilities.cpp b/Engine/Plugins/Runtime/GameplayAbilities/Source/GameplayAbilities/Private/AbilitySystemComponent_Abilities.cpp
--- a/Engine/Plugins/Runtime/GameplayAbilities/Source/GameplayAbilities/Private/AbilitySystemComponent_Abilities.cpp
+++ b/Engine/Plugins/Runtime/GameplayAbilities/Source/GameplayAbilities/Private/AbilitySystemComponent_Abilities.cpp
@@ -54,6 +54,15 @@ void UAbilitySystemComponent::InitializeComponent()
 	AActor *Owner = GetOwner();
 	InitAbilityActorInfo(Owner, Owner);	// Default init to our outer owner
 
+	// cleanup any bad data that may have gotten into SpawnedAttributes
+	for (int32 Idx = SpawnedAttributes.Num()-1; Idx >= 0; --Idx)
+	{
+		if (SpawnedAttributes[Idx] == nullptr)
+		{
+			SpawnedAttributes.RemoveAt(Idx);
+		}
+	}
+
 	TArray<UObject*> ChildObjects;
 	GetObjectsWithOuter(Owner, ChildObjects, false, RF_NoFlags, EInternalObjectFlags::PendingKill);
 	for (UObject* Obj : ChildObjects)
diff --git a/Engine/Plugins/Runtime/GameplayAbilities/Source/GameplayAbilities/Private/AbilitySystemDebugHUD.cpp b/Engine/Plugins/Runtime/GameplayAbilities/Source/GameplayAbilities/Private/AbilitySystemDebugHUD.cpp
--- a/Engine/Plugins/Runtime/GameplayAbilities/Source/GameplayAbilities/Private/AbilitySystemDebugHUD.cpp
+++ b/Engine/Plugins/Runtime/GameplayAbilities/Source/GameplayAbilities/Private/AbilitySystemDebugHUD.cpp
@@ -99,8 +99,9 @@ void AAbilitySystemDebugHUD::DrawDebugAbilitySystemComponent(UAbilitySystemCompo
 	for (const UAttributeSet* Set : Component->SpawnedAttributes)
 	{
 		if (!Set)
+		{
 			continue;
-		check(Set);
+		}
 
 		// Draw Attribute Set
 		DrawWithBackground(Font, FString::Printf(TEXT("%s (%d)"), *Set->GetName(), Set->IsDefaultSubobject()), Color, EAlignHorizontal::Left, X, EAlignVertical::Top, Y);
diff --git a/Engine/Plugins/Runtime/GameplayAbilities/Source/GameplayAbilities/Private/AttributeSet.cpp b/Engine/Plugins/Runtime/GameplayAbilities/Source/GameplayAbilities/Private/AttributeSet.cpp
--- a/Engine/Plugins/Runtime/GameplayAbilities/Source/GameplayAbilities/Private/AttributeSet.cpp
+++ b/Engine/Plugins/Runtime/GameplayAbilities/Source/GameplayAbilities/Private/AttributeSet.cpp
@@ -721,6 +721,10 @@ void FAttributeSetInitterDiscreteLevels::InitAttributeSetDefaults(UAbilitySystem
 	const FAttributeSetDefaults& SetDefaults = Collection->LevelData[Level - 1];
 	for (const UAttributeSet* Set : AbilitySystemComponent->SpawnedAttributes)
 	{
+		if (!Set)
+		{
+			continue;
+		}
 		const FAttributeDefaultValueList* DefaultDataList = SetDefaults.DataMap.Find(Set->GetClass());
 		if (DefaultDataList)
 		{
@@ -768,6 +772,11 @@ void FAttributeSetInitterDiscreteLevels::ApplyAttributeDefault(UAbilitySystemCom
 	const FAttributeSetDefaults& SetDefaults = Collection->LevelData[Level - 1];
 	for (const UAttributeSet* Set : AbilitySystemComponent->SpawnedAttributes)
 	{
+		if (!Set)
+		{
+			continue;
+		}
+
 		const FAttributeDefaultValueList* DefaultDataList = SetDefaults.DataMap.Find(Set->GetClass());
 		if (DefaultDataList)
 		{
diff --git a/Engine/Plugins/Runtime/GameplayAbilities/Source/GameplayAbilities/Public/AbilitySystemComponent.h b/Engine/Plugins/Runtime/GameplayAbilities/Source/GameplayAbilities/Public/AbilitySystemComponent.h
--- a/Engine/Plugins/Runtime/GameplayAbilities/Source/GameplayAbilities/Public/AbilitySystemComponent.h
+++ b/Engine/Plugins/Runtime/GameplayAbilities/Source/GameplayAbilities/Public/AbilitySystemComponent.h
@@ -121,9 +121,19 @@ class GAMEPLAYABILITIES_API UAbilitySystemComponent : public UGameplayTasksCompo
 		return (T*)GetOrCreateAttributeSubobject(T::StaticClass());
 	}
 
-	/** Adds a new AttributeSet that is a DSO (created by called in their CStor) */
 	template <class T>
+	UE_DEPRECATED(4.25, "AddDefaultSubobjectSet is deprecated. Use AddAttributeSetSubobject instead.")
 	const T*	AddDefaultSubobjectSet(T* Subobject)
+	{
+		return AddAttributeSetSubobject(Subobject);
+	}
+
+	/** 
+	 * Manually add a new attribute set that is a subobject of this ability system component.
+	 * All subobjects of this component are automatically added during initialization.
+	 */
+	template <class T>
+	const T* AddAttributeSetSubobject(T* Subobject)
 	{
 		SpawnedAttributes.AddUnique(Subobject);
 		return Subobject;
@@ -151,7 +161,7 @@ class GAMEPLAYABILITIES_API UAbilitySystemComponent : public UGameplayTasksCompo
 	TArray<FAttributeDefaults>	DefaultStartingData;
 
 	/** List of attribute sets */
-	UPROPERTY(Replicated)
+	UPROPERTY(Replicated, Transient)
 	TArray<UAttributeSet*>	SpawnedAttributes;
 
 	/** The linked Anim Instance that this component will play montages in. Use NAME_None for the main anim instance. */
