When root motion movement related ability tasks end they should return the movement component's movement mode to the movement mode it was in before the task started.

diff --git a/Engine/Plugins/Runtime/GameplayAbilities/Source/GameplayAbilities/Private/Abilities/Tasks/AbilityTask_ApplyRootMotionMoveToActorForce.cpp b/Engine/Plugins/Runtime/GameplayAbilities/Source/GameplayAbilities/Private/Abilities/Tasks/AbilityTask_ApplyRootMotionMoveToActorForce.cpp
--- a/Engine/Plugins/Runtime/GameplayAbilities/Source/GameplayAbilities/Private/Abilities/Tasks/AbilityTask_ApplyRootMotionMoveToActorForce.cpp
+++ b/Engine/Plugins/Runtime/GameplayAbilities/Source/GameplayAbilities/Private/Abilities/Tasks/AbilityTask_ApplyRootMotionMoveToActorForce.cpp
@@ -351,7 +351,7 @@ void UAbilityTask_ApplyRootMotionMoveToActorForce::OnDestroy(bool AbilityIsEndin
 
 		if (bSetNewMovementMode)
 		{
-			MovementComponent->SetMovementMode(NewMovementMode);
+			MovementComponent->SetMovementMode(PreviousMovementMode);
 		}
 	}
 
diff --git a/Engine/Plugins/Runtime/GameplayAbilities/Source/GameplayAbilities/Private/Abilities/Tasks/AbilityTask_ApplyRootMotionMoveToForce.cpp b/Engine/Plugins/Runtime/GameplayAbilities/Source/GameplayAbilities/Private/Abilities/Tasks/AbilityTask_ApplyRootMotionMoveToForce.cpp
--- a/Engine/Plugins/Runtime/GameplayAbilities/Source/GameplayAbilities/Private/Abilities/Tasks/AbilityTask_ApplyRootMotionMoveToForce.cpp
+++ b/Engine/Plugins/Runtime/GameplayAbilities/Source/GameplayAbilities/Private/Abilities/Tasks/AbilityTask_ApplyRootMotionMoveToForce.cpp
@@ -169,7 +169,7 @@ void UAbilityTask_ApplyRootMotionMoveToForce::OnDestroy(bool AbilityIsEnding)
 
 		if (bSetNewMovementMode)
 		{
-			MovementComponent->SetMovementMode(NewMovementMode);
+			MovementComponent->SetMovementMode(PreviousMovementMode);
 		}
 	}
 
