Remove an unnecessary pointer and variable access.

diff --git a/Engine/Plugins/Runtime/GameplayAbilities/Source/GameplayAbilities/Private/GameplayEffectTypes.cpp b/Engine/Plugins/Runtime/GameplayAbilities/Source/GameplayAbilities/Private/GameplayEffectTypes.cpp
--- a/Engine/Plugins/Runtime/GameplayAbilities/Source/GameplayAbilities/Private/GameplayEffectTypes.cpp
+++ b/Engine/Plugins/Runtime/GameplayAbilities/Source/GameplayAbilities/Private/GameplayEffectTypes.cpp
@@ -327,7 +327,6 @@ bool FGameplayEffectContextHandle::NetSerialize(FArchive& Ar, class UPackageMap*
 			}
 		}
 
-		void* ContainerPtr = Data.Get();
 		UScriptStruct* ScriptStruct = Data->GetScriptStruct();
 
 		if (ScriptStruct->StructFlags & STRUCT_NetSerializeNative)
