Upgrade AbilitySystemComponent, GameplayAbility, AttributeSet and some derived classes over to Push Model.

diff --git a/Engine/Plugins/Runtime/GameplayAbilities/Source/GameplayAbilities/GameplayAbilities.Build.cs b/Engine/Plugins/Runtime/GameplayAbilities/Source/GameplayAbilities/GameplayAbilities.Build.cs
--- a/Engine/Plugins/Runtime/GameplayAbilities/Source/GameplayAbilities/GameplayAbilities.Build.cs
+++ b/Engine/Plugins/Runtime/GameplayAbilities/Source/GameplayAbilities/GameplayAbilities.Build.cs
@@ -13,6 +13,7 @@ namespace UnrealBuildTool.Rules
 				{
 					"Core",
 					"CoreUObject",
+					"NetCore",
 					"Engine",
 					"GameplayTags",
 					"GameplayTasks",
diff --git a/Engine/Plugins/Runtime/GameplayAbilities/Source/GameplayAbilities/Private/Abilities/Tasks/AbilityTask_WaitGameplayEffectApplied.cpp b/Engine/Plugins/Runtime/GameplayAbilities/Source/GameplayAbilities/Private/Abilities/Tasks/AbilityTask_WaitGameplayEffectApplied.cpp
--- a/Engine/Plugins/Runtime/GameplayAbilities/Source/GameplayAbilities/Private/Abilities/Tasks/AbilityTask_WaitGameplayEffectApplied.cpp
+++ b/Engine/Plugins/Runtime/GameplayAbilities/Source/GameplayAbilities/Private/Abilities/Tasks/AbilityTask_WaitGameplayEffectApplied.cpp
@@ -22,7 +22,7 @@ void UAbilityTask_WaitGameplayEffectApplied::OnApplyGameplayEffectCallback(UAbil
 {
 	bool PassedComparison = false;
 
-	AActor* AvatarActor = Target ? Target->AvatarActor : nullptr;
+	AActor* AvatarActor = Target ? Target->GetAvatarActor_Direct() : nullptr;
 
 	if (!Filter.FilterPassesForActor(AvatarActor))
 	{
diff --git a/Engine/Plugins/Runtime/GameplayAbilities/Source/GameplayAbilities/Private/AbilitySystemComponent.cpp b/Engine/Plugins/Runtime/GameplayAbilities/Source/GameplayAbilities/Private/AbilitySystemComponent.cpp
--- a/Engine/Plugins/Runtime/GameplayAbilities/Source/GameplayAbilities/Private/AbilitySystemComponent.cpp
+++ b/Engine/Plugins/Runtime/GameplayAbilities/Source/GameplayAbilities/Private/AbilitySystemComponent.cpp
@@ -15,6 +15,7 @@
 #include "Engine/ActorChannel.h"
 #include "GameplayEffectCustomApplicationRequirement.h"
 #include "TimerManager.h"
+#include "Net/Core/PushModel/PushModel.h"
 
 DEFINE_LOG_CATEGORY(LogAbilitySystemComponent);
 
@@ -92,7 +93,7 @@ const UAttributeSet* UAbilitySystemComponent::GetOrCreateAttributeSubobject(TSub
 		if (!MyAttributes)
 		{
 			UAttributeSet* Attributes = NewObject<UAttributeSet>(OwningActor, AttributeClass);
-			SpawnedAttributes.AddUnique(Attributes);
+			GetSpawnedAttributes_Mutable().AddUnique(Attributes);
 			MyAttributes = Attributes;
 			bIsNetDirty = true;
 		}
@@ -110,7 +111,7 @@ const UAttributeSet* UAbilitySystemComponent::GetAttributeSubobjectChecked(const
 
 const UAttributeSet* UAbilitySystemComponent::GetAttributeSubobject(const TSubclassOf<UAttributeSet> AttributeClass) const
 {
-	for (const UAttributeSet* Set : SpawnedAttributes)
+	for (const UAttributeSet* Set : GetSpawnedAttributes())
 	{
 		if (Set && Set->IsA(AttributeClass))
 		{
@@ -127,7 +128,7 @@ bool UAbilitySystemComponent::HasAttributeSetForAttribute(FGameplayAttribute Att
 
 void UAbilitySystemComponent::GetAllAttributes(OUT TArray<FGameplayAttribute>& Attributes)
 {
-	for (UAttributeSet* Set : SpawnedAttributes)
+	for (const UAttributeSet* Set : GetSpawnedAttributes())
 	{
 		if (!Set)
 		{
@@ -171,7 +172,12 @@ void UAbilitySystemComponent::OnRegister()
 	ActiveGameplayCues.SetOwner(this);	
 	MinimalReplicationGameplayCues.bMinimalReplication = true;
 	MinimalReplicationGameplayCues.SetOwner(this);
+
+PRAGMA_DISABLE_DEPRECATION_WARNINGS
+	// This field is not replicated (MinimalReplicationTags has a custom serializer),
+	// so we don't need to mark it dirty.
 	MinimalReplicationTags.Owner = this;
+PRAGMA_ENABLE_DEPRECATION_WARNINGS
 
 	/** Allocate an AbilityActorInfo. Note: this goes through a global function and is a SharedPtr so projects can make their own AbilityActorInfo */
 	if(!AbilityActorInfo.IsValid())
@@ -750,13 +756,13 @@ FActiveGameplayEffectHandle UAbilitySystemComponent::ApplyGameplayEffectSpecToSe
 			// Log results of applied GE spec
 			if (UE_LOG_ACTIVE(VLogAbilitySystem, Log))
 			{
-				ABILITY_VLOG(OwnerActor, Log, TEXT("Applied %s"), *OurCopyOfSpec->Def->GetFName().ToString());
+				ABILITY_VLOG(GetOwnerActor(), Log, TEXT("Applied %s"), *OurCopyOfSpec->Def->GetFName().ToString());
 
 				for (const FGameplayModifierInfo& Modifier : Spec.Def->Modifiers)
 				{
 					float Magnitude = 0.f;
 					Modifier.ModifierMagnitude.AttemptCalculateMagnitude(Spec, Magnitude);
-					ABILITY_VLOG(OwnerActor, Log, TEXT("         %s: %s %f"), *Modifier.Attribute.GetName(), *EGameplayModOpToString(Modifier.ModifierOp), Magnitude);
+					ABILITY_VLOG(GetOwnerActor(), Log, TEXT("         %s: %s %f"), *Modifier.Attribute.GetName(), *EGameplayModOpToString(Modifier.ModifierOp), Magnitude);
 				}
 			}
 		}
@@ -911,13 +917,13 @@ void UAbilitySystemComponent::ExecuteGameplayEffect(FGameplayEffectSpec &Spec, F
 
 	if (UE_LOG_ACTIVE(VLogAbilitySystem, Log))
 	{
-		ABILITY_VLOG(OwnerActor, Log, TEXT("Executed %s"), *Spec.Def->GetFName().ToString());
+		ABILITY_VLOG(GetOwnerActor(), Log, TEXT("Executed %s"), *Spec.Def->GetFName().ToString());
 		
 		for (FGameplayModifierInfo Modifier : Spec.Def->Modifiers)
 		{
 			float Magnitude = 0.f;
 			Modifier.ModifierMagnitude.AttemptCalculateMagnitude(Spec, Magnitude);
-			ABILITY_VLOG(OwnerActor, Log, TEXT("         %s: %s %f"), *Modifier.Attribute.GetName(), *EGameplayModOpToString(Modifier.ModifierOp), Magnitude);
+			ABILITY_VLOG(GetOwnerActor(), Log, TEXT("         %s: %s %f"), *Modifier.Attribute.GetName(), *EGameplayModOpToString(Modifier.ModifierOp), Magnitude);
 		}
 	}
 	bIsNetDirty = true;
@@ -1071,8 +1077,8 @@ void UAbilitySystemComponent::OnImmunityBlockGameplayEffect(const FGameplayEffec
 
 void UAbilitySystemComponent::InitDefaultGameplayCueParameters(FGameplayCueParameters& Parameters)
 {
-	Parameters.Instigator = OwnerActor;
-	Parameters.EffectCauser = AvatarActor;
+	Parameters.Instigator = GetOwnerActor();
+	Parameters.EffectCauser = GetAvatarActor_Direct();
 }
 
 bool UAbilitySystemComponent::IsReadyForGameplayCues()
@@ -1486,25 +1492,34 @@ int32 UAbilitySystemComponent::RemoveActiveEffects(const FGameplayEffectQuery& Q
 
 void UAbilitySystemComponent::GetLifetimeReplicatedProps(TArray< FLifetimeProperty > & OutLifetimeProps) const
 {	
-	DOREPLIFETIME(UAbilitySystemComponent, SpawnedAttributes);
-	DOREPLIFETIME(UAbilitySystemComponent, ActiveGameplayEffects);
-	DOREPLIFETIME(UAbilitySystemComponent, ActiveGameplayCues);
-	
-	DOREPLIFETIME_CONDITION(UAbilitySystemComponent, ActivatableAbilities, COND_ReplayOrOwner);
-	DOREPLIFETIME_CONDITION(UAbilitySystemComponent, BlockedAbilityBindings, COND_OwnerOnly);
+	// Fast Arrays don't use push model, but there's no harm in marking them with it.
+	// The flag will just be ignored.
 
-	DOREPLIFETIME(UAbilitySystemComponent, OwnerActor);
-	DOREPLIFETIME(UAbilitySystemComponent, AvatarActor);
+	FDoRepLifetimeParams Params;
+	Params.bIsPushBased = true;
 
-	DOREPLIFETIME_CONDITION(UAbilitySystemComponent, ReplicatedPredictionKeyMap, COND_OwnerOnly);
-	DOREPLIFETIME(UAbilitySystemComponent, RepAnimMontageInfo);
-	
-	DOREPLIFETIME_CONDITION(UAbilitySystemComponent, MinimalReplicationGameplayCues, COND_SkipOwner);
-	DOREPLIFETIME_CONDITION(UAbilitySystemComponent, MinimalReplicationTags, COND_SkipOwner);
+	DOREPLIFETIME_WITH_PARAMS_FAST(UAbilitySystemComponent, SpawnedAttributes, Params);
+	DOREPLIFETIME_WITH_PARAMS_FAST(UAbilitySystemComponent, ActiveGameplayEffects, Params);
+	DOREPLIFETIME_WITH_PARAMS_FAST(UAbilitySystemComponent, ActiveGameplayCues, Params);
+	DOREPLIFETIME_WITH_PARAMS_FAST(UAbilitySystemComponent, RepAnimMontageInfo, Params);
+	DOREPLIFETIME_WITH_PARAMS_FAST(UAbilitySystemComponent, OwnerActor, Params);
+	DOREPLIFETIME_WITH_PARAMS_FAST(UAbilitySystemComponent, AvatarActor, Params);
+
+	Params.Condition = COND_ReplayOrOwner;
+	DOREPLIFETIME_WITH_PARAMS_FAST(UAbilitySystemComponent, ActivatableAbilities, Params);
+
+	Params.Condition = COND_ReplayOnly;
+	DOREPLIFETIME_WITH_PARAMS_FAST(UAbilitySystemComponent, ClientDebugStrings, Params);
+	DOREPLIFETIME_WITH_PARAMS_FAST(UAbilitySystemComponent, ServerDebugStrings, Params);
 	
-	DOREPLIFETIME_CONDITION(UAbilitySystemComponent, ClientDebugStrings, COND_ReplayOnly);
-	DOREPLIFETIME_CONDITION(UAbilitySystemComponent, ServerDebugStrings, COND_ReplayOnly);
+	Params.Condition = COND_OwnerOnly;
+	DOREPLIFETIME_WITH_PARAMS_FAST(UAbilitySystemComponent, BlockedAbilityBindings, Params);
+	DOREPLIFETIME_WITH_PARAMS_FAST(UAbilitySystemComponent, ReplicatedPredictionKeyMap, Params);
 	
+	Params.Condition = COND_SkipOwner;
+	DOREPLIFETIME_WITH_PARAMS_FAST(UAbilitySystemComponent, MinimalReplicationGameplayCues, Params);
+	DOREPLIFETIME_WITH_PARAMS_FAST(UAbilitySystemComponent, MinimalReplicationTags, Params);
+
 	Super::GetLifetimeReplicatedProps(OutLifetimeProps);
 }
 
@@ -1520,9 +1535,9 @@ void UAbilitySystemComponent::ForceReplication()
 
 void UAbilitySystemComponent::ForceAvatarReplication()
 {
-	if (AvatarActor)
+	if (AActor* LocalAvatarActor = GetAvatarActor_Direct())
 	{
-		AvatarActor->ForceNetUpdate();
+		LocalAvatarActor->ForceNetUpdate();
 	}
 }
 
@@ -1530,7 +1545,7 @@ bool UAbilitySystemComponent::ReplicateSubobjects(class UActorChannel *Channel,
 {
 	bool WroteSomething = Super::ReplicateSubobjects(Channel, Bunch, RepFlags);
 
-	for (const UAttributeSet* Set : SpawnedAttributes)
+	for (const UAttributeSet* Set : GetSpawnedAttributes())
 	{
 		if (Set && !Set->IsPendingKill())
 		{
@@ -1551,7 +1566,7 @@ bool UAbilitySystemComponent::ReplicateSubobjects(class UActorChannel *Channel,
 
 void UAbilitySystemComponent::GetSubobjectsWithStableNamesForNetworking(TArray<UObject*>& Objs)
 {
-	for (const UAttributeSet* Set : SpawnedAttributes)
+	for (const UAttributeSet* Set : GetSpawnedAttributes())
 	{
 		if (Set && Set->IsNameStableForNetworking())
 		{
@@ -1593,7 +1608,7 @@ IAbilitySystemReplicationProxyInterface* UAbilitySystemComponent::GetReplication
 	if (ReplicationProxyEnabled)
 	{
 		// Note the expectation is that when the avatar actor is null (e.g during a respawn) that we do return null and calling code handles this (by probably not replicating whatever it was going to)
-		return Cast<IAbilitySystemReplicationProxyInterface>(AvatarActor);
+		return Cast<IAbilitySystemReplicationProxyInterface>(GetAvatarActor_Direct());
 	}
 
 	return Cast<IAbilitySystemReplicationProxyInterface>(this);
@@ -1720,18 +1735,22 @@ void UAbilitySystemComponent::DebugCyclicAggregatorBroadcasts(FAggregator* Aggre
 
 void UAbilitySystemComponent::OnRep_ClientDebugString()
 {
+	const TArray<FString>& LocalClientStrings = GetClientDebugStrings();
+
 	ABILITY_LOG(Display, TEXT(" "));
-	ABILITY_LOG(Display, TEXT("Received Client AbilitySystem Debug information: (%d lines)"), ClientDebugStrings.Num());
-	for (FString& Str : ClientDebugStrings)
+	ABILITY_LOG(Display, TEXT("Received Client AbilitySystem Debug information: (%d lines)"), LocalClientStrings.Num());
+	for (const FString& Str : LocalClientStrings)
 	{
 		ABILITY_LOG(Display, TEXT("%s"), *Str);
 	}
 }
 void UAbilitySystemComponent::OnRep_ServerDebugString()
 {
+	const TArray<FString>& LocalServerStrings = GetClientDebugStrings();
+
 	ABILITY_LOG(Display, TEXT(" "));
-	ABILITY_LOG(Display, TEXT("Server AbilitySystem Debug information: (%d lines)"), ClientDebugStrings.Num());
-	for (FString& Str : ServerDebugStrings)
+	ABILITY_LOG(Display, TEXT("Server AbilitySystem Debug information: (%d lines)"), LocalServerStrings.Num());
+	for (const FString& Str : LocalServerStrings)
 	{
 		ABILITY_LOG(Display, TEXT("%s"), *Str);
 	}
@@ -1796,7 +1815,7 @@ void UAbilitySystemComponent::ServerPrintDebug_RequestWithStrings_Implementation
 		ABILITY_LOG(Display, TEXT("%s"), *Str);
 	}
 
-	ClientDebugStrings = Strings;
+	SetClientDebugStrings(TArray<FString>(Strings));
 	ServerPrintDebug_Request_Implementation();
 }
 
@@ -1818,7 +1837,7 @@ void UAbilitySystemComponent::ServerPrintDebug_Request_Implementation()
 
 	Debug_Internal(DebugInfo);
 
-	ServerDebugStrings = DebugInfo.Strings;
+	SetServerDebugStrings(MoveTemp(DebugInfo.Strings));
 
 	ClientPrintDebug_Response(DebugInfo.Strings, DebugInfo.GameFlags);
 }
@@ -2173,11 +2192,14 @@ void UAbilitySystemComponent::Debug_Internal(FAbilitySystemComponentDebugInfo& I
 		{
 			DebugTitle += TEXT("GAMEPLAYEFFECTS ");
 		}
+		AActor* LocalAvatarActor = GetAvatarActor_Direct();
+		AActor* LocalOwnerActor = GetOwnerActor();
+
 		// Avatar info
-		if (AvatarActor)
+		if (LocalAvatarActor)
 		{
-			const ENetRole AvatarRole = AvatarActor->GetLocalRole();
-			DebugTitle += FString::Printf(TEXT("for avatar %s "), *AvatarActor->GetName());
+			const ENetRole AvatarRole = LocalAvatarActor->GetLocalRole();
+			DebugTitle += FString::Printf(TEXT("for avatar %s "), *LocalAvatarActor->GetName());
 			if (AvatarRole == ROLE_AutonomousProxy)
 			{
 				DebugTitle += TEXT("(local player) ");
@@ -2192,10 +2214,10 @@ void UAbilitySystemComponent::Debug_Internal(FAbilitySystemComponentDebugInfo& I
 			}
 		}
 		// Owner info
-		if (OwnerActor && OwnerActor != AvatarActor)
+		if (LocalOwnerActor && LocalOwnerActor != LocalAvatarActor)
 		{
-			const ENetRole OwnerRole = OwnerActor->GetLocalRole();
-			DebugTitle += FString::Printf(TEXT("for owner %s "), *OwnerActor->GetName());
+			const ENetRole OwnerRole = LocalOwnerActor->GetLocalRole();
+			DebugTitle += FString::Printf(TEXT("for owner %s "), *LocalOwnerActor->GetName());
 			if (OwnerRole == ROLE_AutonomousProxy)
 			{
 				DebugTitle += TEXT("(autonomous) ");
@@ -2389,7 +2411,7 @@ void UAbilitySystemComponent::Debug_Internal(FAbilitySystemComponentDebugInfo& I
 
 				if (ActiveGE.Spec.Def->StackingType == EGameplayEffectStackingType::AggregateBySource)
 				{
-					StackString = FString::Printf(TEXT("(Stacks: %d. From: %s) "), ActiveGE.Spec.StackCount, *GetNameSafe(ActiveGE.Spec.GetContext().GetInstigatorAbilitySystemComponent()->AvatarActor));
+					StackString = FString::Printf(TEXT("(Stacks: %d. From: %s) "), ActiveGE.Spec.StackCount, *GetNameSafe(ActiveGE.Spec.GetContext().GetInstigatorAbilitySystemComponent()->GetAvatarActor_Direct()));
 				}
 				else
 				{
@@ -2461,7 +2483,7 @@ void UAbilitySystemComponent::Debug_Internal(FAbilitySystemComponentDebugInfo& I
 		{
 			Info.Canvas->SetDrawColor(FColor::White);
 		}
-		for (UAttributeSet* Set : SpawnedAttributes)
+		for (UAttributeSet* Set : GetSpawnedAttributes())
 		{
 			if (!Set)
 			{
@@ -2500,13 +2522,14 @@ void UAbilitySystemComponent::Debug_Internal(FAbilitySystemComponentDebugInfo& I
 			FString StatusText;
 			FColor AbilityTextColor = FColor(128, 128, 128);
 			FGameplayTagContainer FailureTags;
+			const TArray<uint8>& LocalBlockedAbilityBindings = GetBlockedAbilityBindings();
 
 			if (AbilitySpec.IsActive())
 			{
 				StatusText = FString::Printf(TEXT(" (Active %d)"), AbilitySpec.ActiveCount);
 				AbilityTextColor = FColor::Yellow;
 			}
-			else if (BlockedAbilityBindings.IsValidIndex(AbilitySpec.InputID) && BlockedAbilityBindings[AbilitySpec.InputID])
+			else if (LocalBlockedAbilityBindings.IsValidIndex(AbilitySpec.InputID) && LocalBlockedAbilityBindings[AbilitySpec.InputID])
 			{
 				StatusText = TEXT(" (InputBlocked)");
 				AbilityTextColor = FColor::Red;
@@ -2614,4 +2637,126 @@ void UAbilitySystemComponent::Debug_Internal(FAbilitySystemComponentDebugInfo& I
 	Info.YL = MaxCharHeight;
 }
 
+PRAGMA_DISABLE_DEPRECATION_WARNINGS
+void UAbilitySystemComponent::SetOwnerActor(AActor* NewOwnerActor)
+{
+	MARK_PROPERTY_DIRTY_FROM_NAME(UAbilitySystemComponent, OwnerActor, this);
+	OwnerActor = NewOwnerActor;
+}
+
+AActor* UAbilitySystemComponent::GetOwnerActor() const
+{
+	return OwnerActor;
+}
+
+void UAbilitySystemComponent::SetAvatarActor_Direct(AActor* NewAvatarActor)
+{
+	MARK_PROPERTY_DIRTY_FROM_NAME(UAbilitySystemComponent, AvatarActor, this);
+	AvatarActor = NewAvatarActor;
+}
+
+AActor* UAbilitySystemComponent::GetAvatarActor_Direct() const
+{
+	return AvatarActor;
+}
+
+void UAbilitySystemComponent::SetSpawnedAttributes(const TArray<UAttributeSet*>& NewSpawnedAttributes)
+{
+	GetSpawnedAttributes_Mutable() = NewSpawnedAttributes;
+}
+
+TArray<UAttributeSet*>& UAbilitySystemComponent::GetSpawnedAttributes_Mutable()
+{
+	MARK_PROPERTY_DIRTY_FROM_NAME(UAbilitySystemComponent, SpawnedAttributes, this);
+	return SpawnedAttributes;
+}
+
+const TArray<UAttributeSet*>& UAbilitySystemComponent::GetSpawnedAttributes() const
+{
+	return SpawnedAttributes;
+}
+
+void UAbilitySystemComponent::SetClientDebugStrings(TArray<FString>&& NewClientDebugStrings)
+{
+	GetClientDebugStrings_Mutable() = NewClientDebugStrings;
+}
+
+TArray<FString>& UAbilitySystemComponent::GetClientDebugStrings_Mutable()
+{
+	MARK_PROPERTY_DIRTY_FROM_NAME(UAbilitySystemComponent, ClientDebugStrings, this);
+	return ClientDebugStrings;
+}
+
+const TArray<FString>& UAbilitySystemComponent::GetClientDebugStrings() const
+{
+	MARK_PROPERTY_DIRTY_FROM_NAME(UAbilitySystemComponent, ClientDebugStrings, this);
+	return ClientDebugStrings;
+}
+
+void UAbilitySystemComponent::SetServerDebugStrings(TArray<FString>&& NewServerDebugStrings)
+{
+	GetServerDebugStrings_Mutable() = NewServerDebugStrings;
+}
+
+TArray<FString>& UAbilitySystemComponent::GetServerDebugStrings_Mutable()
+{
+	MARK_PROPERTY_DIRTY_FROM_NAME(UAbilitySystemComponent, ServerDebugStrings, this);
+	return ServerDebugStrings;
+}
+
+const TArray<FString>& UAbilitySystemComponent::GetServerDebugStrings() const
+{
+	return ServerDebugStrings;
+}
+
+void UAbilitySystemComponent::SetRepAnimMontageInfo(const FGameplayAbilityRepAnimMontage& NewRepAnimMontageInfo)
+{
+	GetRepAnimMontageInfo_Mutable() = NewRepAnimMontageInfo;
+}
+
+FGameplayAbilityRepAnimMontage& UAbilitySystemComponent::GetRepAnimMontageInfo_Mutable()
+{
+	MARK_PROPERTY_DIRTY_FROM_NAME(UAbilitySystemComponent, RepAnimMontageInfo, this);
+	return RepAnimMontageInfo;
+}
+
+const FGameplayAbilityRepAnimMontage& UAbilitySystemComponent::GetRepAnimMontageInfo() const
+{
+	return RepAnimMontageInfo;
+}
+
+void UAbilitySystemComponent::SetBlockedAbilityBindings(const TArray<uint8>& NewBlockedAbilityBindings)
+{
+	GetBlockedAbilityBindings_Mutable() = NewBlockedAbilityBindings;
+}
+
+TArray<uint8>& UAbilitySystemComponent::GetBlockedAbilityBindings_Mutable()
+{
+	MARK_PROPERTY_DIRTY_FROM_NAME(UAbilitySystemComponent, BlockedAbilityBindings, this);
+	return BlockedAbilityBindings;
+}
+
+const TArray<uint8>& UAbilitySystemComponent::GetBlockedAbilityBindings() const
+{
+	return BlockedAbilityBindings;
+}
+
+void UAbilitySystemComponent::SetMinimalReplicationTags(const FMinimalReplicationTagCountMap& NewMinimalReplicationTags)
+{
+	GetMinimalReplicationTags_Mutable() = NewMinimalReplicationTags;
+}
+
+FMinimalReplicationTagCountMap& UAbilitySystemComponent::GetMinimalReplicationTags_Mutable()
+{
+	MARK_PROPERTY_DIRTY_FROM_NAME(UAbilitySystemComponent, MinimalReplicationTags, this);
+	return MinimalReplicationTags;
+}
+
+const FMinimalReplicationTagCountMap& UAbilitySystemComponent::GetMinimalReplicationTags() const
+{
+	return MinimalReplicationTags;
+}
+
+PRAGMA_ENABLE_DEPRECATION_WARNINGS
+
 #undef LOCTEXT_NAMESPACE
diff --git a/Engine/Plugins/Runtime/GameplayAbilities/Source/GameplayAbilities/Private/AbilitySystemComponent_Abilities.cpp b/Engine/Plugins/Runtime/GameplayAbilities/Source/GameplayAbilities/Private/AbilitySystemComponent_Abilities.cpp
--- a/Engine/Plugins/Runtime/GameplayAbilities/Source/GameplayAbilities/Private/AbilitySystemComponent_Abilities.cpp
+++ b/Engine/Plugins/Runtime/GameplayAbilities/Source/GameplayAbilities/Private/AbilitySystemComponent_Abilities.cpp
@@ -55,22 +55,24 @@ void UAbilitySystemComponent::InitializeComponent()
 	InitAbilityActorInfo(Owner, Owner);	// Default init to our outer owner
 
 	// cleanup any bad data that may have gotten into SpawnedAttributes
-	for (int32 Idx = SpawnedAttributes.Num()-1; Idx >= 0; --Idx)
+	TArray<UAttributeSet*>& SpawnedAttributesRef = GetSpawnedAttributes_Mutable();
+	for (int32 Idx = SpawnedAttributesRef.Num()-1; Idx >= 0; --Idx)
 	{
-		if (SpawnedAttributes[Idx] == nullptr)
+		if (SpawnedAttributesRef[Idx] == nullptr)
 		{
-			SpawnedAttributes.RemoveAt(Idx);
+			SpawnedAttributesRef.RemoveAt(Idx);
 		}
 	}
 
 	TArray<UObject*> ChildObjects;
 	GetObjectsWithOuter(Owner, ChildObjects, false, RF_NoFlags, EInternalObjectFlags::PendingKill);
+
 	for (UObject* Obj : ChildObjects)
 	{
 		UAttributeSet* Set = Cast<UAttributeSet>(Obj);
 		if (Set)  
 		{
-			SpawnedAttributes.AddUnique(Set);
+			SpawnedAttributesRef.AddUnique(Set);
 			bIsNetDirty = true;
 		}
 	}
@@ -93,7 +95,7 @@ void UAbilitySystemComponent::OnComponentDestroyed(bool bDestroyingHierarchy)
 	// in this case. This can happen when the owning actor's level is removed and later re-added
 	// to the world, since EndPlay (and therefore UninitializeComponents) will be called on
 	// the owning actor when its level is removed.
-	for (UAttributeSet* Set : SpawnedAttributes)
+	for (UAttributeSet* Set : GetSpawnedAttributes())
 	{
 		if (Set)
 		{
@@ -117,7 +119,7 @@ void UAbilitySystemComponent::TickComponent(float DeltaTime, enum ELevelTick Tic
 
 	Super::TickComponent(DeltaTime, TickType, ThisTickFunction);
 
-	for (UAttributeSet* AttributeSet : SpawnedAttributes)
+	for (UAttributeSet* AttributeSet : GetSpawnedAttributes())
 	{
 		ITickableAttributeSetInterface* TickableSet = Cast<ITickableAttributeSetInterface>(AttributeSet);
 		if (TickableSet)
@@ -135,11 +137,11 @@ void UAbilitySystemComponent::InitAbilityActorInfo(AActor* InOwnerActor, AActor*
 
 	AbilityActorInfo->InitFromActor(InOwnerActor, InAvatarActor, this);
 
-	OwnerActor = InOwnerActor;
+	SetOwnerActor(InOwnerActor);
 
 	// caching the previous value of the actor so we can check against it but then setting the value to the new because it may get used
-	const AActor* PrevAvatarActor = AvatarActor;
-	AvatarActor = InAvatarActor;
+	const AActor* PrevAvatarActor = GetAvatarActor_Direct();
+	SetAvatarActor_Direct(InAvatarActor);
 
 	// if the avatar actor was null but won't be after this, we want to run the deferred gameplaycues that may not have run in NetDeltaSerialize
 	// Conversely, if the ability actor was previously null, then the effects would not run in the NetDeltaSerialize. As such we want to run them now.
@@ -168,7 +170,7 @@ void UAbilitySystemComponent::InitAbilityActorInfo(AActor* InOwnerActor, AActor*
 	LocalAnimMontageInfo = FGameplayAbilityLocalAnimMontage();
 	if (IsOwnerActorAuthoritative())
 	{
-		RepAnimMontageInfo = FGameplayAbilityRepAnimMontage();
+		SetRepAnimMontageInfo(FGameplayAbilityRepAnimMontage());
 	}
 
 	if (bPendingMontageRep)
@@ -179,7 +181,7 @@ void UAbilitySystemComponent::InitAbilityActorInfo(AActor* InOwnerActor, AActor*
 
 bool UAbilitySystemComponent::GetShouldTick() const 
 {
-	const bool bHasReplicatedMontageInfoToUpdate = (IsOwnerActorAuthoritative() && RepAnimMontageInfo.IsStopped == false);
+	const bool bHasReplicatedMontageInfoToUpdate = (IsOwnerActorAuthoritative() && GetRepAnimMontageInfo().IsStopped == false);
 	
 	if (bHasReplicatedMontageInfoToUpdate)
 	{
@@ -189,7 +191,7 @@ bool UAbilitySystemComponent::GetShouldTick() const
 	bool bResult = Super::GetShouldTick();	
 	if (bResult == false)
 	{
-		for (const UAttributeSet* AttributeSet : SpawnedAttributes)
+		for (const UAttributeSet* AttributeSet : GetSpawnedAttributes())
 		{
 			const ITickableAttributeSetInterface* TickableAttributeSet = Cast<const ITickableAttributeSetInterface>(AttributeSet);
 			if (TickableAttributeSet && TickableAttributeSet->ShouldTick())
@@ -206,26 +208,29 @@ bool UAbilitySystemComponent::GetShouldTick() const
 void UAbilitySystemComponent::SetAvatarActor(AActor* InAvatarActor)
 {
 	check(AbilityActorInfo.IsValid());
-	InitAbilityActorInfo(OwnerActor, InAvatarActor);
+	InitAbilityActorInfo(GetOwnerActor(), InAvatarActor);
 }
 
 void UAbilitySystemComponent::ClearActorInfo()
 {
 	check(AbilityActorInfo.IsValid());
 	AbilityActorInfo->ClearActorInfo();
-	OwnerActor = nullptr;
-	AvatarActor = nullptr;
+	SetOwnerActor(nullptr);
+	SetAvatarActor_Direct(nullptr);
 }
 
 void UAbilitySystemComponent::OnRep_OwningActor()
 {
 	check(AbilityActorInfo.IsValid());
 
-	if (OwnerActor != AbilityActorInfo->OwnerActor || AvatarActor != AbilityActorInfo->AvatarActor)
+	AActor* LocalOwnerActor = GetOwnerActor();
+	AActor* LocalAvatarActor = GetAvatarActor_Direct();
+
+	if (LocalOwnerActor != AbilityActorInfo->OwnerActor || LocalAvatarActor != AbilityActorInfo->AvatarActor)
 	{
-		if (OwnerActor != nullptr)
+		if (LocalOwnerActor != nullptr)
 		{
-			InitAbilityActorInfo(OwnerActor, AvatarActor);
+			InitAbilityActorInfo(LocalOwnerActor, LocalAvatarActor);
 		}
 		else
 		{
@@ -953,18 +958,20 @@ void UAbilitySystemComponent::UnBlockAbilitiesWithTags(const FGameplayTagContain
 void UAbilitySystemComponent::BlockAbilityByInputID(int32 InputID)
 {
 	bIsNetDirty = true;
-	if (InputID >= 0 && InputID < BlockedAbilityBindings.Num())
+	const TArray<uint8>& ConstBlockedAbilityBindings = GetBlockedAbilityBindings();
+	if (InputID >= 0 && InputID < ConstBlockedAbilityBindings.Num())
 	{
-		++BlockedAbilityBindings[InputID];
+		++GetBlockedAbilityBindings_Mutable()[InputID];
 	}
 }
 
 void UAbilitySystemComponent::UnBlockAbilityByInputID(int32 InputID)
 {
 	bIsNetDirty = true;
-	if (InputID >= 0 && InputID < BlockedAbilityBindings.Num() && BlockedAbilityBindings[InputID] > 0)
+	const TArray<uint8>& ConstBlockedAbilityBindings = GetBlockedAbilityBindings();
+	if (InputID >= 0 && InputID < ConstBlockedAbilityBindings.Num() && ConstBlockedAbilityBindings[InputID] > 0)
 	{
-		--BlockedAbilityBindings[InputID];
+		--GetBlockedAbilityBindings_Mutable()[InputID];
 	}
 }
 
@@ -1142,7 +1149,8 @@ bool UAbilitySystemComponent::TryActivateAbility(FGameplayAbilitySpecHandle Abil
 bool UAbilitySystemComponent::IsAbilityInputBlocked(int32 InputID) const
 {
 	// Check if this ability's input binding is currently blocked
-	if (InputID >= 0 && InputID < BlockedAbilityBindings.Num() && BlockedAbilityBindings[InputID] > 0)
+	const TArray<uint8>& ConstBlockedAbilityBindings = GetBlockedAbilityBindings();
+	if (InputID >= 0 && InputID < ConstBlockedAbilityBindings.Num() && ConstBlockedAbilityBindings[InputID] > 0)
 	{
 		return true;
 	}
@@ -1165,14 +1173,14 @@ bool UAbilitySystemComponent::InternalTryActivateAbility(FGameplayAbilitySpecHan
 
 	if (Handle.IsValid() == false)
 	{
-		ABILITY_LOG(Warning, TEXT("InternalTryActivateAbility called with invalid Handle! ASC: %s. AvatarActor: %s"), *GetPathName(), *GetNameSafe(AvatarActor));
+		ABILITY_LOG(Warning, TEXT("InternalTryActivateAbility called with invalid Handle! ASC: %s. AvatarActor: %s"), *GetPathName(), *GetNameSafe(GetAvatarActor_Direct()));
 		return false;
 	}
 
 	FGameplayAbilitySpec* Spec = FindAbilitySpecFromHandle(Handle);
 	if (!Spec)
 	{
-		ABILITY_LOG(Warning, TEXT("InternalTryActivateAbility called with a valid handle but no matching ability was found. Handle: %s ASC: %s. AvatarActor: %s"), *Handle.ToString(), *GetPathName(), *GetNameSafe(AvatarActor));
+		ABILITY_LOG(Warning, TEXT("InternalTryActivateAbility called with a valid handle but no matching ability was found. Handle: %s ASC: %s. AvatarActor: %s"), *Handle.ToString(), *GetPathName(), *GetNameSafe(GetAvatarActor_Direct()));
 		return false;
 	}
 
@@ -1196,9 +1204,9 @@ bool UAbilitySystemComponent::InternalTryActivateAbility(FGameplayAbilitySpecHan
 		NetMode = PC->GetLocalRole();
 	}
 	// Fallback to avataractor otherwise. Edge case: avatar "dies" and becomes torn off and ROLE_Authority. We don't want to use this case (use PC role instead).
-	else if (AvatarActor)
+	else if (AActor* LocalAvatarActor = GetAvatarActor_Direct())
 	{
-		NetMode = AvatarActor->GetLocalRole();
+		NetMode = LocalAvatarActor->GetLocalRole();
 	}
 
 	if (NetMode == ROLE_SimulatedProxy)
@@ -2029,8 +2037,8 @@ void UAbilitySystemComponent::MonitoredTagChanged(const FGameplayTag Tag, int32
 							FGameplayEventData EventData;
 							EventData.EventMagnitude = NewCount;
 							EventData.EventTag = EventTag;
-							EventData.Instigator = OwnerActor;
-							EventData.Target = OwnerActor;
+							EventData.Instigator = GetOwnerActor();
+							EventData.Target = EventData.Instigator;
 							// Try to activate it
 							InternalTryActivateAbility(Spec->Handle, FPredictionKey(), nullptr, nullptr, &EventData);
 
@@ -2140,7 +2148,7 @@ void UAbilitySystemComponent::SetBlockAbilityBindingsArray(FGameplayAbilityInput
 {
 	UEnum* EnumBinds = BindInfo.GetBindEnum();
 	bIsNetDirty = true;
-	BlockedAbilityBindings.SetNumZeroed(EnumBinds->NumEnums());
+	GetBlockedAbilityBindings_Mutable().SetNumZeroed(EnumBinds->NumEnums());
 }
 
 void UAbilitySystemComponent::AbilityLocalInputPressed(int32 InputID)
@@ -2393,15 +2401,17 @@ float UAbilitySystemComponent::PlayMontage(UGameplayAbility* InAnimatingAbility,
 			// Replicate to non owners
 			if (IsOwnerActorAuthoritative())
 			{
+				FGameplayAbilityRepAnimMontage& MutableRepAnimMontageInfo = GetRepAnimMontageInfo_Mutable();
+
 				// Those are static parameters, they are only set when the montage is played. They are not changed after that.
-				RepAnimMontageInfo.AnimMontage = NewAnimMontage;
-				RepAnimMontageInfo.ForcePlayBit = !bool(RepAnimMontageInfo.ForcePlayBit);
+				MutableRepAnimMontageInfo.AnimMontage = NewAnimMontage;
+				MutableRepAnimMontageInfo.ForcePlayBit = !bool(MutableRepAnimMontageInfo.ForcePlayBit);
 
-				RepAnimMontageInfo.SectionIdToPlay = 0;
-				if (RepAnimMontageInfo.AnimMontage && StartSectionName != NAME_None)
+				MutableRepAnimMontageInfo.SectionIdToPlay = 0;
+				if (MutableRepAnimMontageInfo.AnimMontage && StartSectionName != NAME_None)
 				{
 					// we add one so INDEX_NONE can be used in the on rep
-					RepAnimMontageInfo.SectionIdToPlay = RepAnimMontageInfo.AnimMontage->GetSectionIndex(StartSectionName) + 1;
+					MutableRepAnimMontageInfo.SectionIdToPlay = MutableRepAnimMontageInfo.AnimMontage->GetSectionIndex(StartSectionName) + 1;
 				}
 
 				// Update parameters that change during Montage life time.
@@ -2448,7 +2458,7 @@ void UAbilitySystemComponent::AnimMontage_UpdateReplicatedData()
 {
 	check(IsOwnerActorAuthoritative());
 
-	AnimMontage_UpdateReplicatedData(RepAnimMontageInfo);
+	AnimMontage_UpdateReplicatedData(GetRepAnimMontageInfo_Mutable());
 }
 
 void UAbilitySystemComponent::AnimMontage_UpdateReplicatedData(FGameplayAbilityRepAnimMontage& OutRepAnimMontageInfo)
@@ -2535,9 +2545,11 @@ void UAbilitySystemComponent::OnRep_ReplicatedAnimMontage()
 {
 	UWorld* World = GetWorld();
 
-	if (RepAnimMontageInfo.bSkipPlayRate)
+	const FGameplayAbilityRepAnimMontage& ConstRepAnimMontageInfo = GetRepAnimMontageInfo_Mutable();
+
+	if (ConstRepAnimMontageInfo.bSkipPlayRate)
 	{
-		RepAnimMontageInfo.PlayRate = 1.f;
+		GetRepAnimMontageInfo_Mutable().PlayRate = 1.f;
 	}
 
 	const bool bIsPlayingReplay = World && World->IsPlayingReplay();
@@ -2561,40 +2573,40 @@ void UAbilitySystemComponent::OnRep_ReplicatedAnimMontage()
 		{
 			ABILITY_LOG( Warning, TEXT("\n\nOnRep_ReplicatedAnimMontage, %s"), *GetNameSafe(this));
 			ABILITY_LOG( Warning, TEXT("\tAnimMontage: %s\n\tPlayRate: %f\n\tPosition: %f\n\tBlendTime: %f\n\tNextSectionID: %d\n\tIsStopped: %d\n\tForcePlayBit: %d"),
-				*GetNameSafe(RepAnimMontageInfo.AnimMontage), 
-				RepAnimMontageInfo.PlayRate, 
-				RepAnimMontageInfo.Position,
-				RepAnimMontageInfo.BlendTime,
-				RepAnimMontageInfo.NextSectionID, 
-				RepAnimMontageInfo.IsStopped, 
-				RepAnimMontageInfo.ForcePlayBit);
+				*GetNameSafe(ConstRepAnimMontageInfo.AnimMontage),
+				ConstRepAnimMontageInfo.PlayRate,
+				ConstRepAnimMontageInfo.Position,
+				ConstRepAnimMontageInfo.BlendTime,
+				ConstRepAnimMontageInfo.NextSectionID,
+				ConstRepAnimMontageInfo.IsStopped,
+				ConstRepAnimMontageInfo.ForcePlayBit);
 			ABILITY_LOG( Warning, TEXT("\tLocalAnimMontageInfo.AnimMontage: %s\n\tPosition: %f"),
 				*GetNameSafe(LocalAnimMontageInfo.AnimMontage), AnimInstance->Montage_GetPosition(LocalAnimMontageInfo.AnimMontage));
 		}
 
-		if( RepAnimMontageInfo.AnimMontage )
+		if(ConstRepAnimMontageInfo.AnimMontage )
 		{
 			// New Montage to play
-			const bool ReplicatedPlayBit = bool(RepAnimMontageInfo.ForcePlayBit);
-			if ((LocalAnimMontageInfo.AnimMontage != RepAnimMontageInfo.AnimMontage) || (LocalAnimMontageInfo.PlayBit != ReplicatedPlayBit))
+			const bool ReplicatedPlayBit = bool(ConstRepAnimMontageInfo.ForcePlayBit);
+			if ((LocalAnimMontageInfo.AnimMontage != ConstRepAnimMontageInfo.AnimMontage) || (LocalAnimMontageInfo.PlayBit != ReplicatedPlayBit))
 			{
 				LocalAnimMontageInfo.PlayBit = ReplicatedPlayBit;
-				PlayMontageSimulated(RepAnimMontageInfo.AnimMontage, RepAnimMontageInfo.PlayRate);
+				PlayMontageSimulated(ConstRepAnimMontageInfo.AnimMontage, ConstRepAnimMontageInfo.PlayRate);
 			}
 
 			if (LocalAnimMontageInfo.AnimMontage == nullptr)
 			{ 
-				ABILITY_LOG(Warning, TEXT("OnRep_ReplicatedAnimMontage: PlayMontageSimulated failed. Name: %s, AnimMontage: %s"), *GetNameSafe(this), *GetNameSafe(RepAnimMontageInfo.AnimMontage));
+				ABILITY_LOG(Warning, TEXT("OnRep_ReplicatedAnimMontage: PlayMontageSimulated failed. Name: %s, AnimMontage: %s"), *GetNameSafe(this), *GetNameSafe(ConstRepAnimMontageInfo.AnimMontage));
 				return;
 			}
 
 			// Play Rate has changed
-			if (AnimInstance->Montage_GetPlayRate(LocalAnimMontageInfo.AnimMontage) != RepAnimMontageInfo.PlayRate)
+			if (AnimInstance->Montage_GetPlayRate(LocalAnimMontageInfo.AnimMontage) != ConstRepAnimMontageInfo.PlayRate)
 			{
-				AnimInstance->Montage_SetPlayRate(LocalAnimMontageInfo.AnimMontage, RepAnimMontageInfo.PlayRate);
+				AnimInstance->Montage_SetPlayRate(LocalAnimMontageInfo.AnimMontage, ConstRepAnimMontageInfo.PlayRate);
 			}
 
-			const int32 SectionIdToPlay = (static_cast<int32>(RepAnimMontageInfo.SectionIdToPlay) - 1);
+			const int32 SectionIdToPlay = (static_cast<int32>(ConstRepAnimMontageInfo.SectionIdToPlay) - 1);
 			if (SectionIdToPlay != INDEX_NONE)
 			{
 				FName SectionNameToJumpTo = LocalAnimMontageInfo.AnimMontage->GetSectionName(SectionIdToPlay);
@@ -2606,20 +2618,20 @@ void UAbilitySystemComponent::OnRep_ReplicatedAnimMontage()
 
 			// Compressed Flags
 			const bool bIsStopped = AnimInstance->Montage_GetIsStopped(LocalAnimMontageInfo.AnimMontage);
-			const bool bReplicatedIsStopped = bool(RepAnimMontageInfo.IsStopped);
+			const bool bReplicatedIsStopped = bool(ConstRepAnimMontageInfo.IsStopped);
 
 			// Process stopping first, so we don't change sections and cause blending to pop.
 			if (bReplicatedIsStopped)
 			{
 				if (!bIsStopped)
 				{
-					CurrentMontageStop(RepAnimMontageInfo.BlendTime);
+					CurrentMontageStop(ConstRepAnimMontageInfo.BlendTime);
 				}
 			}
-			else if (!RepAnimMontageInfo.SkipPositionCorrection)
+			else if (!ConstRepAnimMontageInfo.SkipPositionCorrection)
 			{
-				const int32 RepSectionID = LocalAnimMontageInfo.AnimMontage->GetSectionIndexFromPosition(RepAnimMontageInfo.Position);
-				const int32 RepNextSectionID = int32(RepAnimMontageInfo.NextSectionID) - 1;
+				const int32 RepSectionID = LocalAnimMontageInfo.AnimMontage->GetSectionIndexFromPosition(ConstRepAnimMontageInfo.Position);
+				const int32 RepNextSectionID = int32(ConstRepAnimMontageInfo.NextSectionID) - 1;
 		
 				// And NextSectionID for the replicated SectionID.
 				if( RepSectionID != INDEX_NONE )
@@ -2645,25 +2657,25 @@ void UAbilitySystemComponent::OnRep_ReplicatedAnimMontage()
 				// Update Position. If error is too great, jump to replicated position.
 				const float CurrentPosition = AnimInstance->Montage_GetPosition(LocalAnimMontageInfo.AnimMontage);
 				const int32 CurrentSectionID = LocalAnimMontageInfo.AnimMontage->GetSectionIndexFromPosition(CurrentPosition);
-				const float DeltaPosition = RepAnimMontageInfo.Position - CurrentPosition;
+				const float DeltaPosition = ConstRepAnimMontageInfo.Position - CurrentPosition;
 
 				// Only check threshold if we are located in the same section. Different sections require a bit more work as we could be jumping around the timeline.
 				// And therefore DeltaPosition is not as trivial to determine.
-				if ((CurrentSectionID == RepSectionID) && (FMath::Abs(DeltaPosition) > MONTAGE_REP_POS_ERR_THRESH) && (RepAnimMontageInfo.IsStopped == 0))
+				if ((CurrentSectionID == RepSectionID) && (FMath::Abs(DeltaPosition) > MONTAGE_REP_POS_ERR_THRESH) && (ConstRepAnimMontageInfo.IsStopped == 0))
 				{
 					// fast forward to server position and trigger notifies
-					if (FAnimMontageInstance* MontageInstance = AnimInstance->GetActiveInstanceForMontage(RepAnimMontageInfo.AnimMontage))
+					if (FAnimMontageInstance* MontageInstance = AnimInstance->GetActiveInstanceForMontage(ConstRepAnimMontageInfo.AnimMontage))
 					{
 						// Skip triggering notifies if we're going backwards in time, we've already triggered them.
-						const float DeltaTime = !FMath::IsNearlyZero(RepAnimMontageInfo.PlayRate) ? (DeltaPosition / RepAnimMontageInfo.PlayRate) : 0.f;
+						const float DeltaTime = !FMath::IsNearlyZero(ConstRepAnimMontageInfo.PlayRate) ? (DeltaPosition / ConstRepAnimMontageInfo.PlayRate) : 0.f;
 						if (DeltaTime >= 0.f)
 						{
 							MontageInstance->UpdateWeight(DeltaTime);
-							MontageInstance->HandleEvents(CurrentPosition, RepAnimMontageInfo.Position, nullptr);
+							MontageInstance->HandleEvents(CurrentPosition, ConstRepAnimMontageInfo.Position, nullptr);
 							AnimInstance->TriggerAnimNotifies(DeltaTime);
 						}
 					}
-					AnimInstance->Montage_SetPosition(LocalAnimMontageInfo.AnimMontage, RepAnimMontageInfo.Position);
+					AnimInstance->Montage_SetPosition(LocalAnimMontageInfo.AnimMontage, ConstRepAnimMontageInfo.Position);
 				}
 			}
 		}
@@ -2714,11 +2726,13 @@ void UAbilitySystemComponent::CurrentMontageJumpToSection(FName SectionName)
 		AnimInstance->Montage_JumpToSection(SectionName, LocalAnimMontageInfo.AnimMontage);
 		if (IsOwnerActorAuthoritative())
 		{
-			RepAnimMontageInfo.SectionIdToPlay = 0;
-			if (RepAnimMontageInfo.AnimMontage)
+			FGameplayAbilityRepAnimMontage& MutableRepAnimMontageInfo = GetRepAnimMontageInfo_Mutable();
+
+			MutableRepAnimMontageInfo.SectionIdToPlay = 0;
+			if (MutableRepAnimMontageInfo.AnimMontage)
 			{
 				// we add one so INDEX_NONE can be used in the on rep
-				RepAnimMontageInfo.SectionIdToPlay = RepAnimMontageInfo.AnimMontage->GetSectionIndex(SectionName) + 1;
+				MutableRepAnimMontageInfo.SectionIdToPlay = MutableRepAnimMontageInfo.AnimMontage->GetSectionIndex(SectionName) + 1;
 			}
 
 			AnimMontage_UpdateReplicatedData();
@@ -2828,11 +2842,13 @@ void UAbilitySystemComponent::ServerCurrentMontageJumpToSectionName_Implementati
 			// Update replicated version for Simulated Proxies if we are on the server.
 			if (IsOwnerActorAuthoritative())
 			{
-				RepAnimMontageInfo.SectionIdToPlay = 0;
-				if (RepAnimMontageInfo.AnimMontage && SectionName != NAME_None)
+				FGameplayAbilityRepAnimMontage& MutableRepAnimMontageInfo = GetRepAnimMontageInfo_Mutable();
+
+				MutableRepAnimMontageInfo.SectionIdToPlay = 0;
+				if (MutableRepAnimMontageInfo.AnimMontage && SectionName != NAME_None)
 				{
 					// we add one so INDEX_NONE can be used in the on rep
-					RepAnimMontageInfo.SectionIdToPlay = RepAnimMontageInfo.AnimMontage->GetSectionIndex(SectionName) + 1;
+					MutableRepAnimMontageInfo.SectionIdToPlay = MutableRepAnimMontageInfo.AnimMontage->GetSectionIndex(SectionName) + 1;
 				}
 
 				AnimMontage_UpdateReplicatedData();
@@ -2952,7 +2968,7 @@ float UAbilitySystemComponent::GetCurrentMontageSectionTimeLeft() const
 
 void UAbilitySystemComponent::SetMontageRepAnimPositionMethod(ERepAnimPositionMethod InMethod)
 {
-	RepAnimMontageInfo.SetRepAnimPositionMethod(InMethod);
+	GetRepAnimMontageInfo_Mutable().SetRepAnimPositionMethod(InMethod);
 }
 
 bool UAbilitySystemComponent::IsAnimatingAbility(UGameplayAbility* InAbility) const
diff --git a/Engine/Plugins/Runtime/GameplayAbilities/Source/GameplayAbilities/Private/AbilitySystemDebugHUD.cpp b/Engine/Plugins/Runtime/GameplayAbilities/Source/GameplayAbilities/Private/AbilitySystemDebugHUD.cpp
--- a/Engine/Plugins/Runtime/GameplayAbilities/Source/GameplayAbilities/Private/AbilitySystemDebugHUD.cpp
+++ b/Engine/Plugins/Runtime/GameplayAbilities/Source/GameplayAbilities/Private/AbilitySystemDebugHUD.cpp
@@ -96,7 +96,7 @@ void AAbilitySystemDebugHUD::DrawDebugAbilitySystemComponent(UAbilitySystemCompo
 	DrawWithBackground(Font, String, Color, EAlignHorizontal::Left, X, EAlignVertical::Top, Y);
 
 
-	for (const UAttributeSet* Set : Component->SpawnedAttributes)
+	for (const UAttributeSet* Set : Component->GetSpawnedAttributes())
 	{
 		if (!Set)
 		{
diff --git a/Engine/Plugins/Runtime/GameplayAbilities/Source/GameplayAbilities/Private/AttributeSet.cpp b/Engine/Plugins/Runtime/GameplayAbilities/Source/GameplayAbilities/Private/AttributeSet.cpp
--- a/Engine/Plugins/Runtime/GameplayAbilities/Source/GameplayAbilities/Private/AttributeSet.cpp
+++ b/Engine/Plugins/Runtime/GameplayAbilities/Source/GameplayAbilities/Private/AttributeSet.cpp
@@ -15,6 +15,7 @@
 #include "AbilitySystemGlobals.h"
 #include "AbilitySystemComponent.h"
 #include "AbilitySystemTestAttributeSet.h"
+#include "Net/Core/PushModel/PushModel.h"
 
 #if WITH_EDITOR
 #include "EditorReimportHandler.h"
@@ -86,6 +87,8 @@ void FGameplayAttribute::SetNumericValueChecked(float& NewValue, class UAttribut
 		OldValue = *static_cast<float*>(ValuePtr);
 		Dest->PreAttributeChange(*this, NewValue);
 		NumericProperty->SetFloatingPointPropertyValue(ValuePtr, NewValue);
+
+		MARK_PROPERTY_DIRTY(Dest, NumericProperty);
 	}
 	else if (IsGameplayAttributeDataProperty(Attribute.Get()))
 	{
@@ -96,6 +99,8 @@ void FGameplayAttribute::SetNumericValueChecked(float& NewValue, class UAttribut
 		OldValue = DataPtr->GetCurrentValue();
 		Dest->PreAttributeChange(*this, NewValue);
 		DataPtr->SetCurrentValue(NewValue);
+
+		MARK_PROPERTY_DIRTY(Dest, StructProperty);
 	}
 	else
 	{
@@ -160,12 +165,40 @@ float FGameplayAttribute::GetNumericValueChecked(const UAttributeSet* Src) const
 	return 0.f;
 }
 
+const FGameplayAttributeData* FGameplayAttribute::GetGameplayAttributeData(const UAttributeSet* Src) const
+{
+	if (Src && IsGameplayAttributeDataProperty(Attribute.Get()))
+	{
+		FStructProperty* StructProperty = CastField<FStructProperty>(Attribute.Get());
+		check(StructProperty);
+		MARK_PROPERTY_DIRTY(Src, StructProperty);
+		return StructProperty->ContainerPtrToValuePtr<FGameplayAttributeData>(Src);
+	}
+
+	return nullptr;
+}
+
+const FGameplayAttributeData* FGameplayAttribute::GetGameplayAttributeDataChecked(const UAttributeSet* Src) const
+{
+	if (Src && IsGameplayAttributeDataProperty(Attribute.Get()))
+	{
+		FStructProperty* StructProperty = CastField<FStructProperty>(Attribute.Get());
+		check(StructProperty);
+		MARK_PROPERTY_DIRTY(Src, StructProperty);
+		return StructProperty->ContainerPtrToValuePtr<FGameplayAttributeData>(Src);
+	}
+
+	check(false);
+	return nullptr;
+}
+
 FGameplayAttributeData* FGameplayAttribute::GetGameplayAttributeData(UAttributeSet* Src) const
 {
 	if (Src && IsGameplayAttributeDataProperty(Attribute.Get()))
 	{
 		FStructProperty* StructProperty = CastField<FStructProperty>(Attribute.Get());
 		check(StructProperty);
+		MARK_PROPERTY_DIRTY(Src, StructProperty);
 		return StructProperty->ContainerPtrToValuePtr<FGameplayAttributeData>(Src);
 	}
 
@@ -178,6 +211,7 @@ FGameplayAttributeData* FGameplayAttribute::GetGameplayAttributeDataChecked(UAtt
 	{
 		FStructProperty* StructProperty = CastField<FStructProperty>(Attribute.Get());
 		check(StructProperty);
+		MARK_PROPERTY_DIRTY(Src, StructProperty);
 		return StructProperty->ContainerPtrToValuePtr<FGameplayAttributeData>(Src);
 	}
 
@@ -719,7 +753,7 @@ void FAttributeSetInitterDiscreteLevels::InitAttributeSetDefaults(UAbilitySystem
 	}
 
 	const FAttributeSetDefaults& SetDefaults = Collection->LevelData[Level - 1];
-	for (const UAttributeSet* Set : AbilitySystemComponent->SpawnedAttributes)
+	for (const UAttributeSet* Set : AbilitySystemComponent->GetSpawnedAttributes())
 	{
 		if (!Set)
 		{
@@ -770,7 +804,7 @@ void FAttributeSetInitterDiscreteLevels::ApplyAttributeDefault(UAbilitySystemCom
 	}
 
 	const FAttributeSetDefaults& SetDefaults = Collection->LevelData[Level - 1];
-	for (const UAttributeSet* Set : AbilitySystemComponent->SpawnedAttributes)
+	for (const UAttributeSet* Set : AbilitySystemComponent->GetSpawnedAttributes())
 	{
 		if (!Set)
 		{
diff --git a/Engine/Plugins/Runtime/GameplayAbilities/Source/GameplayAbilities/Private/GameplayEffect.cpp b/Engine/Plugins/Runtime/GameplayAbilities/Source/GameplayAbilities/Private/GameplayEffect.cpp
--- a/Engine/Plugins/Runtime/GameplayAbilities/Source/GameplayAbilities/Private/GameplayEffect.cpp
+++ b/Engine/Plugins/Runtime/GameplayAbilities/Source/GameplayAbilities/Private/GameplayEffect.cpp
@@ -2230,9 +2230,10 @@ void FActiveGameplayEffectsContainer::OnMagnitudeDependencyChange(FActiveGamepla
 					if (!bMarkedDirty)
 					{
 						bMarkedDirty = true;
-						if (Owner && Owner->OwnerActor && IsNetAuthority())
+						AActor* OwnerActor = Owner ? Owner->GetOwnerActor() : nullptr;
+						if (IsNetAuthority() && OwnerActor)
 						{
-							Owner->OwnerActor->FlushNetDormancy();
+							OwnerActor->FlushNetDormancy();
 						}
 						MarkItemDirty(*ActiveEffect);
 					}
@@ -2751,9 +2752,10 @@ FActiveGameplayEffect* FActiveGameplayEffectsContainer::ApplyGameplayEffectSpec(
 
 	bFoundExistingStackableGE = false;
 
-	if (Owner && Owner->OwnerActor && IsNetAuthority())
+	AActor* OwnerActor = Owner ? Owner->GetOwnerActor() : nullptr;
+	if (IsNetAuthority() && OwnerActor)
 	{
-		Owner->OwnerActor->FlushNetDormancy();
+		OwnerActor->FlushNetDormancy();
 	}
 
 	FActiveGameplayEffect* AppliedActiveGE = nullptr;
@@ -3047,7 +3049,8 @@ void FActiveGameplayEffectsContainer::InternalOnActiveGameplayEffectAdded(FActiv
 	SCOPE_CYCLE_UOBJECT(EffectDef, EffectDef);
 
 	GAMEPLAYEFFECT_SCOPE_LOCK();
-	UE_VLOG(Owner->OwnerActor ? Owner->OwnerActor : Owner->GetOuter(), LogGameplayEffects, Log, TEXT("Added: %s"), *GetNameSafe(EffectDef->GetClass()));
+	AActor* OwnerActor = Owner->GetOwnerActor();
+	UE_VLOG(OwnerActor ? OwnerActor : Owner->GetOuter(), LogGameplayEffects, Log, TEXT("Added: %s"), *GetNameSafe(EffectDef->GetClass()));
 
 	// Add our ongoing tag requirements to the dependency map. We will actually check for these tags below.
 	for (const FGameplayTag& Tag : EffectDef->OngoingTagRequirements.IgnoreTags)
@@ -3217,15 +3220,16 @@ bool FActiveGameplayEffectsContainer::RemoveActiveGameplayEffect(FActiveGameplay
 		FActiveGameplayEffect& Effect = *GetActiveGameplayEffect(ActiveGEIdx);
 		if (Effect.Handle == Handle && Effect.IsPendingRemove == false)
 		{
-			UE_VLOG(Owner->OwnerActor, LogGameplayEffects, Log, TEXT("Removed: %s"), *GetNameSafe(Effect.Spec.Def->GetClass()));
+			AActor* OwnerActor = Owner->GetOwnerActor();
+			UE_VLOG(OwnerActor, LogGameplayEffects, Log, TEXT("Removed: %s"), *GetNameSafe(Effect.Spec.Def->GetClass()));
 			if (UE_LOG_ACTIVE(VLogAbilitySystem, Log))
 			{
-				ABILITY_VLOG(Owner->OwnerActor, Log, TEXT("Removed %s"), *Effect.Spec.Def->GetFName().ToString());
+				ABILITY_VLOG(OwnerActor, Log, TEXT("Removed %s"), *Effect.Spec.Def->GetFName().ToString());
 				for (FGameplayModifierInfo Modifier : Effect.Spec.Def->Modifiers)
 				{
 					float Magnitude = 0.f;
 					Modifier.ModifierMagnitude.AttemptCalculateMagnitude(Effect.Spec, Magnitude);
-					ABILITY_VLOG(Owner->OwnerActor, Log, TEXT("         %s: %s %f"), *Modifier.Attribute.GetName(), *EGameplayModOpToString(Modifier.ModifierOp), Magnitude);
+					ABILITY_VLOG(OwnerActor, Log, TEXT("         %s: %s %f"), *Modifier.Attribute.GetName(), *EGameplayModOpToString(Modifier.ModifierOp), Magnitude);
 				}
 			}
 
@@ -3246,12 +3250,13 @@ void FActiveGameplayEffectsContainer::InternalExecutePeriodicGameplayEffect(FAct
 
 		if (UE_LOG_ACTIVE(VLogAbilitySystem, Log))
 		{
-			ABILITY_VLOG(Owner->OwnerActor, Log, TEXT("Executed Periodic Effect %s"), *ActiveEffect.Spec.Def->GetFName().ToString());
+			AActor* OwnerActor = Owner->GetOwnerActor();
+			ABILITY_VLOG(OwnerActor, Log, TEXT("Executed Periodic Effect %s"), *ActiveEffect.Spec.Def->GetFName().ToString());
 			for (FGameplayModifierInfo Modifier : ActiveEffect.Spec.Def->Modifiers)
 			{
 				float Magnitude = 0.f;
 				Modifier.ModifierMagnitude.AttemptCalculateMagnitude(ActiveEffect.Spec, Magnitude);
-				ABILITY_VLOG(Owner->OwnerActor, Log, TEXT("         %s: %s %f"), *Modifier.Attribute.GetName(), *EGameplayModOpToString(Modifier.ModifierOp), Magnitude);
+				ABILITY_VLOG(OwnerActor, Log, TEXT("         %s: %s %f"), *Modifier.Attribute.GetName(), *EGameplayModOpToString(Modifier.ModifierOp), Magnitude);
 			}
 		}
 
@@ -3331,9 +3336,10 @@ bool FActiveGameplayEffectsContainer::InternalRemoveActiveGameplayEffect(int32 I
 			Owner->GetWorld()->GetTimerManager().ClearTimer(Effect.PeriodHandle);
 		}
 
-		if (bIsNetAuthority && Owner->OwnerActor)
+		AActor* OwnerActor = Owner->GetOwnerActor();
+		if (bIsNetAuthority && OwnerActor)
 		{
-			Owner->OwnerActor->FlushNetDormancy();
+			OwnerActor->FlushNetDormancy();
 		}
 
 		// Remove this handle from the global map
@@ -3627,9 +3633,10 @@ void FActiveGameplayEffectsContainer::OnCustomMagnitudeExternalDependencyFired(T
 					{
 						// By default, a dormancy flush should be required here. If a calculation class has requested that
 						// non-net authorities can respond to external dependencies, the dormancy flush is skipped as a desired optimization
-						if (bRequiresDormancyFlush && Owner && Owner->OwnerActor)
+						AActor* OwnerActor = Owner ? Owner->GetOwnerActor() : nullptr;
+						if (bRequiresDormancyFlush && OwnerActor)
 						{
-							Owner->OwnerActor->FlushNetDormancy();
+							OwnerActor->FlushNetDormancy();
 						}
 
 						MarkItemDirty(Effect);
diff --git a/Engine/Plugins/Runtime/GameplayAbilities/Source/GameplayAbilities/Public/AbilitySystemComponent.h b/Engine/Plugins/Runtime/GameplayAbilities/Source/GameplayAbilities/Public/AbilitySystemComponent.h
--- a/Engine/Plugins/Runtime/GameplayAbilities/Source/GameplayAbilities/Public/AbilitySystemComponent.h
+++ b/Engine/Plugins/Runtime/GameplayAbilities/Source/GameplayAbilities/Public/AbilitySystemComponent.h
@@ -21,6 +21,7 @@
 #include "Abilities/GameplayAbilityTargetTypes.h"
 #include "Abilities/GameplayAbility.h"
 #include "AbilitySystemReplicationProxyInterface.h"
+#include "Net/Core/PushModel/PushModel.h"
 #include "AbilitySystemComponent.generated.h"
 
 class AGameplayAbilityTargetActor;
@@ -135,7 +136,7 @@ class GAMEPLAYABILITIES_API UAbilitySystemComponent : public UGameplayTasksCompo
 	template <class T>
 	const T* AddAttributeSetSubobject(T* Subobject)
 	{
-		SpawnedAttributes.AddUnique(Subobject);
+		GetSpawnedAttributes_Mutable().AddUnique(Subobject);
 		return Subobject;
 	}
 
@@ -161,9 +162,15 @@ class GAMEPLAYABILITIES_API UAbilitySystemComponent : public UGameplayTasksCompo
 	TArray<FAttributeDefaults>	DefaultStartingData;
 
 	/** List of attribute sets */
-	UPROPERTY(Replicated, Transient)
+	UE_DEPRECATED(4.26, "This will be made private in future engine versions. Use SetSpawnedAttributes, GetSpawnedAttributes, or GetSpawnedAttributes_Mutable instead.")
+	UPROPERTY(Replicated)
 	TArray<UAttributeSet*>	SpawnedAttributes;
 
+	void SetSpawnedAttributes(const TArray<UAttributeSet*>& NewAttributeSet);
+	TArray<UAttributeSet*>& GetSpawnedAttributes_Mutable();
+	const TArray<UAttributeSet*>& GetSpawnedAttributes() const;
+
+
 	/** The linked Anim Instance that this component will play montages in. Use NAME_None for the main anim instance. */
 	UPROPERTY(BlueprintReadWrite, EditAnywhere, Category = "Skills")
 	FName AffectedAnimInstanceTag; 
@@ -573,25 +580,25 @@ class GAMEPLAYABILITIES_API UAbilitySystemComponent : public UGameplayTasksCompo
 	 */
 	FORCEINLINE void AddMinimalReplicationGameplayTag(const FGameplayTag& GameplayTag)
 	{
-		MinimalReplicationTags.AddTag(GameplayTag);
+		GetMinimalReplicationTags_Mutable().AddTag(GameplayTag);
 		bIsNetDirty = true;
 	}
 
 	FORCEINLINE void AddMinimalReplicationGameplayTags(const FGameplayTagContainer& GameplayTags)
 	{
-		MinimalReplicationTags.AddTags(GameplayTags);
+		GetMinimalReplicationTags_Mutable().AddTags(GameplayTags);
 		bIsNetDirty = true;
 	}
 
 	FORCEINLINE void RemoveMinimalReplicationGameplayTag(const FGameplayTag& GameplayTag)
 	{
-		MinimalReplicationTags.RemoveTag(GameplayTag);
+		GetMinimalReplicationTags_Mutable().RemoveTag(GameplayTag);
 		bIsNetDirty = true;
 	}
 
 	FORCEINLINE void RemoveMinimalReplicationGameplayTags(const FGameplayTagContainer& GameplayTags)
 	{
-		MinimalReplicationTags.RemoveTags(GameplayTags);
+		GetMinimalReplicationTags_Mutable().RemoveTags(GameplayTags);
 		bIsNetDirty = true;
 	}
 
@@ -1039,12 +1046,22 @@ class GAMEPLAYABILITIES_API UAbilitySystemComponent : public UGameplayTasksCompo
 	void ClearDebugInstantEffects();
 #endif // ENABLE_VISUAL_LOG
 
+	UE_DEPRECATED(4.26, "This will be made private in future engine versions. Use SetClientDebugStrings, GetClientDebugStrings, or GetClientDebugStrings_Mutable instead.")
 	UPROPERTY(ReplicatedUsing=OnRep_ClientDebugString)
 	TArray<FString>	ClientDebugStrings;
 
+	void SetClientDebugStrings(TArray<FString>&& NewClientDebugStrings);
+	TArray<FString>& GetClientDebugStrings_Mutable();
+	const TArray<FString>& GetClientDebugStrings() const;
+
+	UE_DEPRECATED(4.26, "This will be made private in future engine versions. Use SetServerDebugStrings, GetServerDebugStrings, or GetServerDebugStrings_Mutable instead.")
 	UPROPERTY(ReplicatedUsing=OnRep_ServerDebugString)
 	TArray<FString>	ServerDebugStrings;
 
+	void SetServerDebugStrings(TArray<FString>&& NewServerDebugStrings);
+	TArray<FString>& GetServerDebugStrings_Mutable();
+	const TArray<FString>& GetServerDebugStrings() const;
+
 	UFUNCTION()
 	virtual void OnRep_ClientDebugString();
 
@@ -1204,12 +1221,20 @@ class GAMEPLAYABILITIES_API UAbilitySystemComponent : public UGameplayTasksCompo
 	// ----------------------------------------------------------------------------------------------------------------	
 
 	/** The actor that owns this component logically */
+	UE_DEPRECATED(4.26, "This will be made private in future engine versions. Use SetOwnerActor or GetOwnerActor instead.")
 	UPROPERTY(ReplicatedUsing = OnRep_OwningActor)
 	AActor* OwnerActor;
 
+	void SetOwnerActor(AActor* NewOwnerActor);
+	AActor* GetOwnerActor() const;
+
+	UE_DEPRECATED(4.26, "This will be made private in future engine versions. Use SetAvatarActor or GetAvatarActor instead.")
 	/** The actor that is the physical representation used for abilities. Can be NULL */
 	UPROPERTY(ReplicatedUsing = OnRep_OwningActor)
 	AActor* AvatarActor;
+
+	void SetAvatarActor_Direct(AActor* NewAvatarActor);
+	AActor* GetAvatarActor_Direct() const;
 	
 	UFUNCTION()
 	void OnRep_OwningActor();
@@ -1361,7 +1386,7 @@ protected:
 	 */
 
 	UPROPERTY(ReplicatedUsing=OnRep_ActivateAbilities, BlueprintReadOnly, Category = "Abilities")
-	FGameplayAbilitySpecContainer	ActivatableAbilities;
+	FGameplayAbilitySpecContainer ActivatableAbilities;
 
 	/** Maps from an ability spec to the target data. Used to track replicated data and callbacks */
 	FGameplayAbilityReplicatedDataContainer AbilityTargetDataMap;
@@ -1449,10 +1474,15 @@ protected:
 	/** Copy over playing flags for duplicate animation data */
 	void AnimMontage_UpdateForcedPlayFlags(FGameplayAbilityRepAnimMontage& OutRepAnimMontageInfo);
 
+	UE_DEPRECATED(4.26, "This will be made private in future engine versions. Use SetRepAnimMontageInfo, GetRepAnimMontageInfo, or GetRepAnimMontageInfo_Mutable instead.")
 	/** Data structure for replicating montage info to simulated clients */
 	UPROPERTY(ReplicatedUsing=OnRep_ReplicatedAnimMontage)
 	FGameplayAbilityRepAnimMontage RepAnimMontageInfo;
 
+	void SetRepAnimMontageInfo(const FGameplayAbilityRepAnimMontage& NewRepAnimMontageInfo);
+	FGameplayAbilityRepAnimMontage& GetRepAnimMontageInfo_Mutable();
+	const FGameplayAbilityRepAnimMontage& GetRepAnimMontageInfo() const;
+
 	/** Cached value of rather this is a simulated actor */
 	UPROPERTY()
 	bool bCachedIsNetSimulated;
@@ -1517,31 +1547,41 @@ protected:
 	
 	/** Contains all of the gameplay effects that are currently active on this component */
 	UPROPERTY(Replicated)
-	FActiveGameplayEffectsContainer	ActiveGameplayEffects;
+	FActiveGameplayEffectsContainer ActiveGameplayEffects;
 
 	/** List of all active gameplay cues, including ones applied manually */
 	UPROPERTY(Replicated)
-	FActiveGameplayCueContainer	ActiveGameplayCues;
+	FActiveGameplayCueContainer ActiveGameplayCues;
 
 	/** Replicated gameplaycues when in minimal replication mode. These are cues that would come normally come from ActiveGameplayEffects */
 	UPROPERTY(Replicated)
-	FActiveGameplayCueContainer	MinimalReplicationGameplayCues;
+	FActiveGameplayCueContainer MinimalReplicationGameplayCues;
 
 	/** Abilities with these tags are not able to be activated */
 	FGameplayTagCountContainer BlockedAbilityTags;
 
+	UE_DEPRECATED(4.26, "This will be made private in future engine versions. Use SetBlockedAbilityBindings, GetBlockedAbilityBindings, or GetBlockedAbilityBindings_Mutable instead.")
 	/** Tracks abilities that are blocked based on input binding. An ability is blocked if BlockedAbilityBindings[InputID] > 0 */
 	UPROPERTY(Transient, Replicated)
 	TArray<uint8> BlockedAbilityBindings;
 
+	void SetBlockedAbilityBindings(const TArray<uint8>& NewBlockedAbilityBindings);
+	TArray<uint8>& GetBlockedAbilityBindings_Mutable();
+	const TArray<uint8>& GetBlockedAbilityBindings() const;
+
 	void DebugCyclicAggregatorBroadcasts(struct FAggregator* Aggregator);
 	
 	/** Acceleration map for all gameplay tags (OwnedGameplayTags from GEs and explicit GameplayCueTags) */
 	FGameplayTagCountContainer GameplayTagCountContainer;
 
+	UE_DEPRECATED(4.26, "This will be made private in future engine versions. Use SetMinimalReplicationTags, GetMinimalReplicationTags, or GetMinimalReplicationTags_Mutable instead.")
 	UPROPERTY(Replicated)
 	FMinimalReplicationTagCountMap MinimalReplicationTags;
 
+	void SetMinimalReplicationTags(const FMinimalReplicationTagCountMap& NewMinimalReplicationTags);
+	FMinimalReplicationTagCountMap& GetMinimalReplicationTags_Mutable();
+	const FMinimalReplicationTagCountMap& GetMinimalReplicationTags() const;
+
 	void ResetTagMap();
 
 	void NotifyTagMap_StackCountChange(const FGameplayTagContainer& Container);
@@ -1577,7 +1617,5 @@ public:
 
 	/** PredictionKeys, see more info in GameplayPrediction.h. This has to come *last* in all replicated properties on the AbilitySystemComponent to ensure OnRep/callback order. */
 	UPROPERTY(Replicated)
-	FReplicatedPredictionKeyMap	ReplicatedPredictionKeyMap;
-
-	const FMinimalReplicationTagCountMap& GetMinimalReplicationTags() const { return MinimalReplicationTags; }
+	FReplicatedPredictionKeyMap ReplicatedPredictionKeyMap;
 };
diff --git a/Engine/Plugins/Runtime/GameplayAbilities/Source/GameplayAbilities/Public/AttributeSet.h b/Engine/Plugins/Runtime/GameplayAbilities/Source/GameplayAbilities/Public/AttributeSet.h
--- a/Engine/Plugins/Runtime/GameplayAbilities/Source/GameplayAbilities/Public/AttributeSet.h
+++ b/Engine/Plugins/Runtime/GameplayAbilities/Source/GameplayAbilities/Public/AttributeSet.h
@@ -118,6 +118,8 @@ struct GAMEPLAYABILITIES_API FGameplayAttribute
 	float GetNumericValueChecked(const UAttributeSet* Src) const;
 
 	/** Returns the AttributeData, will fail if this is a float attribute */
+	const FGameplayAttributeData* GetGameplayAttributeData(const UAttributeSet* Src) const;
+	const FGameplayAttributeData* GetGameplayAttributeDataChecked(const UAttributeSet* Src) const;
 	FGameplayAttributeData* GetGameplayAttributeData(UAttributeSet* Src) const;
 	FGameplayAttributeData* GetGameplayAttributeDataChecked(UAttributeSet* Src) const;
 	
diff --git a/Engine/Source/Runtime/GameplayTasks/Classes/GameplayTasksComponent.h b/Engine/Source/Runtime/GameplayTasks/Classes/GameplayTasksComponent.h
--- a/Engine/Source/Runtime/GameplayTasks/Classes/GameplayTasksComponent.h
+++ b/Engine/Source/Runtime/GameplayTasks/Classes/GameplayTasksComponent.h
@@ -83,9 +83,20 @@ protected:
 	FGameplayResourceSet CurrentlyClaimedResources;
 
 	/** Tasks that run on simulated proxies */
+	UE_DEPRECATED(4.26, "This will be made private in future versions. Use GetSimulatedTasks, GetSimulatedTasks_Mutable, or SetSimulatedTasks instead.")
 	UPROPERTY(ReplicatedUsing = OnRep_SimulatedTasks)
 	TArray<UGameplayTask*> SimulatedTasks;
 
+PRAGMA_DISABLE_DEPRECATION_WARNINGS
+	const TArray<UGameplayTask*>& GetSimulatedTasks()
+	{
+		return SimulatedTasks;
+	}
+PRAGMA_ENABLE_DEPRECATION_WARNINGS
+
+	TArray<UGameplayTask*>& GetSimulatedTasks_Mutable();
+	void SetSimulatedTasks(const TArray<UGameplayTask*>& NewSimulatedTasks);
+
 	UPROPERTY()
 	TArray<UGameplayTask*> TaskPriorityQueue;
 	
diff --git a/Engine/Source/Runtime/GameplayTasks/GameplayTasks.Build.cs b/Engine/Source/Runtime/GameplayTasks/GameplayTasks.Build.cs
--- a/Engine/Source/Runtime/GameplayTasks/GameplayTasks.Build.cs
+++ b/Engine/Source/Runtime/GameplayTasks/GameplayTasks.Build.cs
@@ -26,6 +26,7 @@ namespace UnrealBuildTool.Rules
 					"Core",
 					"CoreUObject",
 					"Engine",
+					"NetCore"
                     //"GameplayTags",
 					// ... add other public dependencies that you statically link with here ...
 				}
diff --git a/Engine/Source/Runtime/GameplayTasks/Private/GameplayTasksComponent.cpp b/Engine/Source/Runtime/GameplayTasks/Private/GameplayTasksComponent.cpp
--- a/Engine/Source/Runtime/GameplayTasks/Private/GameplayTasksComponent.cpp
+++ b/Engine/Source/Runtime/GameplayTasks/Private/GameplayTasksComponent.cpp
@@ -9,6 +9,7 @@
 #include "VisualLogger/VisualLogger.h"
 #include "GameplayTasksPrivate.h"
 #include "Logging/MessageLog.h"
+#include "Net/Core/PushModel/PushModel.h"
 
 #define LOCTEXT_NAMESPACE "GameplayTasksComponent"
 
@@ -76,8 +77,9 @@ void UGameplayTasksComponent::OnGameplayTaskActivated(UGameplayTask& Task)
 	
 	if (Task.IsSimulatedTask())
 	{
-		check(SimulatedTasks.Contains(&Task) == false);
-		SimulatedTasks.Add(&Task);
+		TArray<UGameplayTask*>& MutableSimulatedTasks = GetSimulatedTasks_Mutable();
+		check(MutableSimulatedTasks.Contains(&Task) == false);
+		MutableSimulatedTasks.Add(&Task);
 		bIsNetDirty = true;
 	}
 
@@ -122,7 +124,7 @@ void UGameplayTasksComponent::OnGameplayTaskDeactivated(UGameplayTask& Task)
 
 	if (Task.IsSimulatedTask())
 	{
-		SimulatedTasks.RemoveSingleSwap(&Task);
+		GetSimulatedTasks_Mutable().RemoveSingleSwap(&Task);
 		bIsNetDirty = true;
 	}
 
@@ -152,7 +154,11 @@ void UGameplayTasksComponent::GetLifetimeReplicatedProps(TArray< FLifetimeProper
 	// Intentionally not calling super: We do not want to replicate bActive which controls ticking. We sometimes need to tick on client predictively.
 	DISABLE_ALL_CLASS_REPLICATED_PROPERTIES(Super, EFieldIteratorFlags::IncludeSuper);
 
-	DOREPLIFETIME_CONDITION(UGameplayTasksComponent, SimulatedTasks, COND_SkipOwner);
+	FDoRepLifetimeParams Params;
+	Params.bIsPushBased = true;
+	Params.Condition = COND_SkipOwner;
+
+	DOREPLIFETIME_WITH_PARAMS_FAST(UGameplayTasksComponent, SimulatedTasks, Params);
 }
 
 bool UGameplayTasksComponent::ReplicateSubobjects(UActorChannel* Channel, class FOutBunch *Bunch, FReplicationFlags *RepFlags)
@@ -161,7 +167,7 @@ bool UGameplayTasksComponent::ReplicateSubobjects(UActorChannel* Channel, class
 	
 	if (!RepFlags->bNetOwner)
 	{
-		for (UGameplayTask* SimulatedTask : SimulatedTasks)
+		for (UGameplayTask* SimulatedTask : GetSimulatedTasks())
 		{
 			if (SimulatedTask && !SimulatedTask->IsPendingKill())
 			{
@@ -175,7 +181,7 @@ bool UGameplayTasksComponent::ReplicateSubobjects(UActorChannel* Channel, class
 
 void UGameplayTasksComponent::OnRep_SimulatedTasks()
 {
-	for (UGameplayTask* SimulatedTask : SimulatedTasks)
+	for (UGameplayTask* SimulatedTask : GetSimulatedTasks())
 	{
 		// Temp check 
 		if (SimulatedTask && SimulatedTask->IsTickingTask() && TickingTasks.Contains(SimulatedTask) == false)
@@ -675,6 +681,19 @@ EGameplayTaskRunResult UGameplayTasksComponent::RunGameplayTask(IGameplayTaskOwn
 	return EGameplayTaskRunResult::Error;
 }
 
+PRAGMA_DISABLE_DEPRECATION_WARNINGS
+TArray<UGameplayTask*>& UGameplayTasksComponent::GetSimulatedTasks_Mutable()
+{
+	MARK_PROPERTY_DIRTY_FROM_NAME(UGameplayTasksComponent, SimulatedTasks, this);
+	return SimulatedTasks;
+}
+PRAGMA_ENABLE_DEPRECATION_WARNINGS
+
+void UGameplayTasksComponent::SetSimulatedTasks(const TArray<UGameplayTask*>& NewSimulatedTasks)
+{
+	GetSimulatedTasks_Mutable() = NewSimulatedTasks;
+}
+
 //----------------------------------------------------------------------//
 // BP API
 //----------------------------------------------------------------------//
