AbilitySystemComponent_Abilities (crash);

There is a crash report about access violation and we were using a ptr without checking it.

-- Code --
AbilitySystemComponent_Abilities:
ClientActivateAbilitySucceedWithEventData_Implementation() - Checking LocalAbility before using it.

diff --git a/Engine/Plugins/Runtime/GameplayAbilities/Source/GameplayAbilities/Private/AbilitySystemComponent_Abilities.cpp b/Engine/Plugins/Runtime/GameplayAbilities/Source/GameplayAbilities/Private/AbilitySystemComponent_Abilities.cpp
--- a/Engine/Plugins/Runtime/GameplayAbilities/Source/GameplayAbilities/Private/AbilitySystemComponent_Abilities.cpp
+++ b/Engine/Plugins/Runtime/GameplayAbilities/Source/GameplayAbilities/Private/AbilitySystemComponent_Abilities.cpp
@@ -1798,7 +1798,7 @@ void UAbilitySystemComponent::ClientActivateAbilitySucceedWithEventData_Implemen
 			TArray<UGameplayAbility*> Instances = Spec->GetAbilityInstances();
 			for (UGameplayAbility* LocalAbility : Instances)
 			{
-				if (LocalAbility->GetCurrentActivationInfo().GetActivationPredictionKey() == PredictionKey)
+				if (LocalAbility != nullptr && LocalAbility->GetCurrentActivationInfo().GetActivationPredictionKey() == PredictionKey)
 				{
 					LocalAbility->ConfirmActivateSucceed();
 					found = true;
