Add MistmatchedTag case for bool tags to FScalableFloat so that refactors from bool to FScalableFloat can correctly read the previously set bool. Bool tags are different from other property tags  in that the serialized value is on the tag itself.

#tests Opened up blueprint and observed bool value was read and correctly converted to FScalableFloat with debugger. Saved and reloaded asset for for good measure.

diff --git a/Engine/Plugins/Runtime/GameplayAbilities/Source/GameplayAbilities/Private/AttributeSet.cpp b/Engine/Plugins/Runtime/GameplayAbilities/Source/GameplayAbilities/Private/AttributeSet.cpp
--- a/Engine/Plugins/Runtime/GameplayAbilities/Source/GameplayAbilities/Private/AttributeSet.cpp
+++ b/Engine/Plugins/Runtime/GameplayAbilities/Source/GameplayAbilities/Private/AttributeSet.cpp
@@ -547,6 +547,11 @@ bool FScalableFloat::SerializeFromMismatchedTag(const FPropertyTag& Tag, FStruct
 
 		return true;
 	}
+	else if (Tag.Type == NAME_BoolProperty)
+	{
+		*this = FScalableFloat(Tag.BoolVal ? 1.f : 0.f);
+		return true; 
+	}
 	return false;
 }
 
