FGameplayAbilitySpec dynamic ability tags;

* Added tag container 'DynamicAbilityTags' to FGameplayAbilitySpec.  These are optional ability tags that are replicated with the spec.  They are also captured as source tags by applied gameplay effects.

diff --git a/Engine/Plugins/Runtime/GameplayAbilities/Source/GameplayAbilities/Private/Abilities/GameplayAbility.cpp b/Engine/Plugins/Runtime/GameplayAbilities/Source/GameplayAbilities/Private/Abilities/GameplayAbility.cpp
--- a/Engine/Plugins/Runtime/GameplayAbilities/Source/GameplayAbilities/Private/Abilities/GameplayAbility.cpp
+++ b/Engine/Plugins/Runtime/GameplayAbilities/Source/GameplayAbilities/Private/Abilities/GameplayAbility.cpp
@@ -1055,7 +1055,10 @@ FGameplayEffectSpecHandle UGameplayAbility::MakeOutgoingGameplayEffectSpec(const
 
 void UGameplayAbility::ApplyAbilityTagsToGameplayEffectSpec(FGameplayEffectSpec& Spec, FGameplayAbilitySpec* AbilitySpec) const
 {
-	Spec.CapturedSourceTags.GetSpecTags().AppendTags(AbilityTags);
+	FGameplayTagContainer& CapturedSourceTags = Spec.CapturedSourceTags.GetSpecTags();
+
+	CapturedSourceTags.AppendTags(AbilityTags);
+	CapturedSourceTags.AppendTags(AbilitySpec->DynamicAbilityTags);
 
 	// Allow the source object of the ability to propagate tags along as well
 	if (AbilitySpec)
@@ -1066,7 +1069,7 @@ void UGameplayAbility::ApplyAbilityTagsToGameplayEffectSpec(FGameplayEffectSpec&
 			FGameplayTagContainer SourceObjTags;
 			SourceObjAsTagInterface->GetOwnedGameplayTags(SourceObjTags);
 
-			Spec.CapturedSourceTags.GetSpecTags().AppendTags(SourceObjTags);
+			CapturedSourceTags.AppendTags(SourceObjTags);
 		}
 
 		// Copy SetByCallerMagnitudes 
diff --git a/Engine/Plugins/Runtime/GameplayAbilities/Source/GameplayAbilities/Public/GameplayAbilitySpec.h b/Engine/Plugins/Runtime/GameplayAbilities/Source/GameplayAbilities/Public/GameplayAbilitySpec.h
--- a/Engine/Plugins/Runtime/GameplayAbilities/Source/GameplayAbilities/Public/GameplayAbilitySpec.h
+++ b/Engine/Plugins/Runtime/GameplayAbilities/Source/GameplayAbilities/Public/GameplayAbilitySpec.h
@@ -274,6 +274,10 @@ struct GAMEPLAYABILITIES_API FGameplayAbilitySpec : public FFastArraySerializerI
 	UPROPERTY(NotReplicated)
 	FGameplayAbilityActivationInfo	ActivationInfo;
 
+	/** Optional ability tags that are replicated.  These tags are also captured as source tags by applied gameplay effects. */
+	UPROPERTY()
+	FGameplayTagContainer DynamicAbilityTags;
+
 	/** Non replicating instances of this ability. */
 	UPROPERTY(NotReplicated)
 	TArray<UGameplayAbility*> NonReplicatedInstances;
