ShowDebug AbilitySystem options;

* Added "bUseDebugTargetFromHud" config property to AbilitySystemGlobals.  Set to true if you want "ShowDebug AbilitySystem" to use the hud's debug target instead of the ability system's debug target.

* Added "AbilitySystem.Debug.SetCategory" command.  This can be used to directly set what category to display when using "ShowDebug AbilitySystem".

diff --git a/Engine/Plugins/Runtime/GameplayAbilities/Source/GameplayAbilities/Private/AbilitySystemComponent.cpp b/Engine/Plugins/Runtime/GameplayAbilities/Source/GameplayAbilities/Private/AbilitySystemComponent.cpp
--- a/Engine/Plugins/Runtime/GameplayAbilities/Source/GameplayAbilities/Private/AbilitySystemComponent.cpp
+++ b/Engine/Plugins/Runtime/GameplayAbilities/Source/GameplayAbilities/Private/AbilitySystemComponent.cpp
@@ -1973,6 +1973,30 @@ static void CycleDebugCategory(UWorld* InWorld)
 	TargetInfo->DebugCategoryIndex = (TargetInfo->DebugCategoryIndex+1) % TargetInfo->DebugCategories.Num();
 }
 
+static void SetDebugCategory(const TArray<FString>& Args, UWorld* InWorld)
+{
+	if (Args.Num() < 1)
+	{
+		ABILITY_LOG(Error, TEXT("Missing category name parameter. Usage: AbilitySystem.Debug.SetCategory <CategoryName>"))
+		return;
+	}
+
+	FASCDebugTargetInfo* TargetInfo = GetDebugTargetInfo(InWorld);
+	check(TargetInfo);
+
+	for (int32 CategoryIndex = 0; CategoryIndex < TargetInfo->DebugCategories.Num(); ++CategoryIndex)
+	{
+		const FString CategoryString = TargetInfo->DebugCategories[CategoryIndex].ToString();
+		if (CategoryString.Equals(Args[0], ESearchCase::IgnoreCase))
+		{
+			TargetInfo->DebugCategoryIndex = CategoryIndex;
+			return;
+		}
+	}
+
+	ABILITY_LOG(Error, TEXT("Unable to match category name parameter [%s]. Usage: AbilitySystem.Debug.SetCategory <CategoryName>"), *Args[0]);
+}
+
 UAbilitySystemComponent* GetDebugTarget(FASCDebugTargetInfo* Info)
 {
 	// Return target if we already have one
@@ -2059,10 +2083,16 @@ FAutoConsoleCommandWithWorld AbilitySystemPrevDebugTargetCmd(
 
 FAutoConsoleCommandWithWorld AbilitySystemDebugNextCategoryCmd(
 	TEXT("AbilitySystem.Debug.NextCategory"),
-	TEXT("Targets previous AbilitySystemComponent in ShowDebug AbilitySystem"),
+	TEXT("Switches to the next ShowDebug AbilitySystem category"),
 	FConsoleCommandWithWorldDelegate::CreateStatic(CycleDebugCategory)
 	);
 
+FAutoConsoleCommandWithWorldAndArgs AbilitySystemDebugSetCategoryCmd(
+	TEXT("AbilitySystem.Debug.SetCategory"),
+	TEXT("Sets the ShowDebug AbilitySystem category. Usage: AbilitySystem.Debug.SetCategory <CategoryName>"),
+	FConsoleCommandWithWorldAndArgsDelegate::CreateStatic(SetDebugCategory)
+);
+
 void UAbilitySystemComponent::OnShowDebugInfo(AHUD* HUD, UCanvas* Canvas, const FDebugDisplayInfo& DisplayInfo, float& YL, float& YPos)
 {
 	if (DisplayInfo.IsDisplayOn(TEXT("AbilitySystem")))
@@ -2070,7 +2100,18 @@ void UAbilitySystemComponent::OnShowDebugInfo(AHUD* HUD, UCanvas* Canvas, const
 		UWorld* World = HUD->GetWorld();
 		FASCDebugTargetInfo* TargetInfo = GetDebugTargetInfo(World);
 	
-		if (UAbilitySystemComponent* ASC = GetDebugTarget(TargetInfo))
+		UAbilitySystemComponent* ASC = nullptr;
+
+		if (UAbilitySystemGlobals::Get().bUseDebugTargetFromHud)
+		{
+			ASC = UAbilitySystemGlobals::GetAbilitySystemComponentFromActor(HUD->GetCurrentDebugTargetActor());
+		}
+		else
+		{
+			ASC = GetDebugTarget(TargetInfo);
+		}
+
+		if (ASC)
 		{
 			TArray<FName> LocalDisplayNames;
 			LocalDisplayNames.Add( TargetInfo->DebugCategories[ TargetInfo->DebugCategoryIndex ] );
diff --git a/Engine/Plugins/Runtime/GameplayAbilities/Source/GameplayAbilities/Private/AbilitySystemGlobals.cpp b/Engine/Plugins/Runtime/GameplayAbilities/Source/GameplayAbilities/Private/AbilitySystemGlobals.cpp
--- a/Engine/Plugins/Runtime/GameplayAbilities/Source/GameplayAbilities/Private/AbilitySystemGlobals.cpp
+++ b/Engine/Plugins/Runtime/GameplayAbilities/Source/GameplayAbilities/Private/AbilitySystemGlobals.cpp
@@ -20,6 +20,8 @@ UAbilitySystemGlobals::UAbilitySystemGlobals(const FObjectInitializer& ObjectIni
 {
 	AbilitySystemGlobalsClassName = FSoftClassPath(TEXT("/Script/GameplayAbilities.AbilitySystemGlobals"));
 
+	bUseDebugTargetFromHud = false;
+
 	PredictTargetGameplayEffects = true;
 
 	MinimalReplicationTagCountBits = 5;
diff --git a/Engine/Plugins/Runtime/GameplayAbilities/Source/GameplayAbilities/Public/AbilitySystemGlobals.h b/Engine/Plugins/Runtime/GameplayAbilities/Source/GameplayAbilities/Public/AbilitySystemGlobals.h
--- a/Engine/Plugins/Runtime/GameplayAbilities/Source/GameplayAbilities/Public/AbilitySystemGlobals.h
+++ b/Engine/Plugins/Runtime/GameplayAbilities/Source/GameplayAbilities/Public/AbilitySystemGlobals.h
@@ -168,6 +168,10 @@ class GAMEPLAYABILITIES_API UAbilitySystemGlobals : public UObject
 	/** Global place to accumulate debug strings for ability system component. Used when we fill up client side debug string immediately, and then wait for server to send server strings */
 	TArray<FString>	AbilitySystemDebugStrings;
 
+	/** Set to true if you want the "ShowDebug AbilitySystem" cheat to use the hud's debug target instead of the ability system's debug target. */
+	UPROPERTY(config)
+	bool bUseDebugTargetFromHud;
+
 	/** Helper functions for applying global scaling to various ability system tasks. This isn't meant to be a shipping feature, but to help with debugging and interation via cvar AbilitySystem.GlobalAbilityScale */
 	static void NonShipping_ApplyGlobalAbilityScaler_Rate(float& Rate);
 	static void NonShipping_ApplyGlobalAbilityScaler_Duration(float& Duration);
