Added PostAttributeChange() and PostAttributeBaseChange() callbacks to UAttributeSet.  These callbacks mirror the existing pre-change callbacks.

diff --git a/Engine/Plugins/Runtime/GameplayAbilities/Source/GameplayAbilities/Private/AttributeSet.cpp b/Engine/Plugins/Runtime/GameplayAbilities/Source/GameplayAbilities/Private/AttributeSet.cpp
--- a/Engine/Plugins/Runtime/GameplayAbilities/Source/GameplayAbilities/Private/AttributeSet.cpp
+++ b/Engine/Plugins/Runtime/GameplayAbilities/Source/GameplayAbilities/Private/AttributeSet.cpp
@@ -87,6 +87,7 @@ void FGameplayAttribute::SetNumericValueChecked(float& NewValue, class UAttribut
 		OldValue = *static_cast<float*>(ValuePtr);
 		Dest->PreAttributeChange(*this, NewValue);
 		NumericProperty->SetFloatingPointPropertyValue(ValuePtr, NewValue);
+		Dest->PostAttributeChange(*this, OldValue, NewValue);
 
 		MARK_PROPERTY_DIRTY(Dest, NumericProperty);
 	}
@@ -99,6 +100,7 @@ void FGameplayAttribute::SetNumericValueChecked(float& NewValue, class UAttribut
 		OldValue = DataPtr->GetCurrentValue();
 		Dest->PreAttributeChange(*this, NewValue);
 		DataPtr->SetCurrentValue(NewValue);
+		Dest->PostAttributeChange(*this, OldValue, NewValue);
 
 		MARK_PROPERTY_DIRTY(Dest, StructProperty);
 	}
diff --git a/Engine/Plugins/Runtime/GameplayAbilities/Source/GameplayAbilities/Private/GameplayEffect.cpp b/Engine/Plugins/Runtime/GameplayAbilities/Source/GameplayAbilities/Private/GameplayEffect.cpp
--- a/Engine/Plugins/Runtime/GameplayAbilities/Source/GameplayAbilities/Private/GameplayEffect.cpp
+++ b/Engine/Plugins/Runtime/GameplayAbilities/Source/GameplayAbilities/Private/GameplayEffect.cpp
@@ -2583,6 +2583,9 @@ void FActiveGameplayEffectsContainer::SetAttributeBaseValue(FGameplayAttribute A
 		return;
 	}
 
+	float OldBaseValue = 0.0f;
+	bool bBaseValueSet = false;
+
 	Set->PreAttributeBaseChange(Attribute, NewBaseValue);
 
 	// if we're using the new attributes we should always update their base value
@@ -2594,21 +2597,35 @@ void FActiveGameplayEffectsContainer::SetAttributeBaseValue(FGameplayAttribute A
 		FGameplayAttributeData* DataPtr = StructProperty->ContainerPtrToValuePtr<FGameplayAttributeData>(const_cast<UAttributeSet*>(Set));
 		if (ensure(DataPtr))
 		{
+			OldBaseValue = DataPtr->GetBaseValue();
 			DataPtr->SetBaseValue(NewBaseValue);
+			bBaseValueSet = true;
 		}
 	}
 
 	FAggregatorRef* RefPtr = AttributeAggregatorMap.Find(Attribute);
 	if (RefPtr)
 	{
+		FAggregator* Aggregator = RefPtr->Get();
+		check(Aggregator);
+
 		// There is an aggregator for this attribute, so set the base value. The dirty callback chain
 		// will update the actual AttributeSet property value for us.		
-		RefPtr->Get()->SetBaseValue(NewBaseValue);
+		OldBaseValue = Aggregator->GetBaseValue();
+		Aggregator->SetBaseValue(NewBaseValue);
+		bBaseValueSet = true;
 	}
 	// if there is no aggregator set the current value (base == current in this case)
 	else
 	{
+		OldBaseValue = Owner->GetNumericAttribute(Attribute);
 		InternalUpdateNumericalAttribute(Attribute, NewBaseValue, nullptr);
+		bBaseValueSet = true;
+	}
+
+	if (bBaseValueSet)
+	{
+		Set->PostAttributeBaseChange(Attribute, OldBaseValue, NewBaseValue);
 	}
 }
 
diff --git a/Engine/Plugins/Runtime/GameplayAbilities/Source/GameplayAbilities/Public/AttributeSet.h b/Engine/Plugins/Runtime/GameplayAbilities/Source/GameplayAbilities/Public/AttributeSet.h
--- a/Engine/Plugins/Runtime/GameplayAbilities/Source/GameplayAbilities/Public/AttributeSet.h
+++ b/Engine/Plugins/Runtime/GameplayAbilities/Source/GameplayAbilities/Public/AttributeSet.h
@@ -216,6 +216,9 @@ public:
 	 */
 	virtual void PreAttributeChange(const FGameplayAttribute& Attribute, float& NewValue) { }
 
+	/** Called just after any modification happens to an attribute. */
+	virtual void PostAttributeChange(const FGameplayAttribute& Attribute, float OldValue, float NewValue) const { }
+
 	/**
 	 *	This is called just before any modification happens to an attribute's base value when an attribute aggregator exists.
 	 *	This function should enforce clamping (presuming you wish to clamp the base value along with the final value in PreAttributeChange)
@@ -224,6 +227,9 @@ public:
 	 */
 	virtual void PreAttributeBaseChange(const FGameplayAttribute& Attribute, float& NewValue) const { }
 
+	/** Called just after any modification happens to an attribute's base value when an attribute aggregator exists. */
+	virtual void PostAttributeBaseChange(const FGameplayAttribute& Attribute, float OldValue, float NewValue) const { }
+
 	/** Callback for when an FAggregator is created for an attribute in this set. Allows custom setup of FAggregator::EvaluationMetaData */
 	virtual void OnAttributeAggregatorCreated(const FGameplayAttribute& Attribute, FAggregator* NewAggregator) const { }
 
