Burst, BurstLatent, and Looping gameplay cue notifies;

* Added new gameplay cue notifies GameplayCueNotify_Burst, GameplayCueNotify_BurstLatent, and GameplayCueNotify_Looping.
    - Burst: A non-instanced gameplay cue notify for effects that are one-offs.
    - BurstLatent: An instanced gameplay cue notify for effects that are one-offs.
    - Looping: An instanced gameplay cue notify for continuous looping effects.

* These new gameplay cue notifies support the following effects: particle effects, sound effects, camera shakes, camera animations, lens effects, force feedback effects, and decals.  Each gameplay cue notify has a spawn condition and placement data that can be overridden per effect.

* Added new animation notifies AnimNotify_GameplayCue and AnimNotify_GameplayCueState for invoking non-replicated gameplay cues.

* Added functions for invoking non-replicated gameplay cues to the gameplay cue manager.

diff --git a/Engine/Plugins/Runtime/GameplayAbilities/Source/GameplayAbilities/Private/AnimNotify_GameplayCue.cpp b/Engine/Plugins/Runtime/GameplayAbilities/Source/GameplayAbilities/Private/AnimNotify_GameplayCue.cpp
new file mode 100644
--- /dev/null
+++ b/Engine/Plugins/Runtime/GameplayAbilities/Source/GameplayAbilities/Private/AnimNotify_GameplayCue.cpp
@@ -0,0 +1,183 @@
+// Copyright Epic Games, Inc. All Rights Reserved.
+
+#include "AnimNotify_GameplayCue.h"
+#include "Components/SkeletalMeshComponent.h"
+#include "AbilitySystemGlobals.h"
+#include "AbilitySystemInterface.h"
+#include "AbilitySystemComponent.h"
+#include "Animation/AnimMontage.h"
+#include "GameplayCueManager.h"
+
+
+typedef void (*GameplayCueFunc)(AActor* Target, const FGameplayTag GameplayCueTag, const FGameplayCueParameters& Parameters);
+
+
+static void ProcessGameplayCue(GameplayCueFunc Func, USkeletalMeshComponent* MeshComp, FGameplayTag GameplayCue, UAnimSequenceBase* Animation)
+{
+	if (MeshComp && GameplayCue.IsValid())
+	{
+		AActor* OwnerActor = MeshComp->GetOwner();
+
+#if WITH_EDITOR
+		if (GIsEditor && (OwnerActor == nullptr))
+		{
+			// Make preview work in anim preview window
+			UGameplayCueManager::PreviewComponent = MeshComp;
+			UGameplayCueManager::PreviewWorld = MeshComp->GetWorld();
+
+			if (UGameplayCueManager* GCM = UAbilitySystemGlobals::Get().GetGameplayCueManager())
+			{
+				GCM->LoadNotifyForEditorPreview(GameplayCue);
+			}
+		}
+#endif // #if WITH_EDITOR
+
+		FGameplayCueParameters Parameters;
+		Parameters.Instigator = OwnerActor;
+
+		// Try to get the ability level. This may not be able to find the ability level
+		// in cases where a blend out is happening or two abilities are trying to play animations
+		if (const IAbilitySystemInterface* ASI = Cast<IAbilitySystemInterface>(OwnerActor))
+		{
+			if (UAbilitySystemComponent* ASC = ASI->GetAbilitySystemComponent())
+			{
+				if (ASC->GetCurrentMontage() == Animation)
+				{
+					if (const UGameplayAbility* Ability = ASC->GetAnimatingAbility())
+					{
+						Parameters.AbilityLevel = Ability->GetAbilityLevel();
+					}
+				}
+
+				// Always use ASC's owner for instigator
+				Parameters.Instigator = ASC->GetOwner();
+			}
+		}
+
+		Parameters.TargetAttachComponent = MeshComp;
+
+		(*Func)(OwnerActor, GameplayCue, Parameters);
+	}
+
+#if WITH_EDITOR
+	if (GIsEditor)
+	{
+		UGameplayCueManager::PreviewComponent = nullptr;
+		UGameplayCueManager::PreviewWorld = nullptr;
+	}
+#endif // #if WITH_EDITOR
+}
+
+
+//////////////////////////////////////////////////////////////////////////
+// UAnimNotify_GameplayCue
+//////////////////////////////////////////////////////////////////////////
+UAnimNotify_GameplayCue::UAnimNotify_GameplayCue()
+{
+}
+
+void UAnimNotify_GameplayCue::Notify(USkeletalMeshComponent* MeshComp, UAnimSequenceBase* Animation)
+{
+	Super::Notify(MeshComp, Animation);
+
+	ProcessGameplayCue(&UGameplayCueManager::ExecuteGameplayCue_NonReplicated, MeshComp, GameplayCue.GameplayCueTag, Animation);
+}
+
+FString UAnimNotify_GameplayCue::GetNotifyName_Implementation() const
+{
+	FString DisplayName = TEXT("GameplayCue");
+
+	if (GameplayCue.GameplayCueTag.IsValid())
+	{
+		DisplayName = GameplayCue.GameplayCueTag.ToString();
+		DisplayName += TEXT(" (Burst)");
+	}
+
+	return DisplayName;
+}
+
+#if WITH_EDITOR
+bool UAnimNotify_GameplayCue::CanBePlaced(UAnimSequenceBase* Animation) const
+{
+	return (Animation && Animation->IsA(UAnimMontage::StaticClass()));
+}
+#endif // WITH_EDITOR
+
+
+//////////////////////////////////////////////////////////////////////////
+// UAnimNotify_GameplayCueState
+//////////////////////////////////////////////////////////////////////////
+UAnimNotify_GameplayCueState::UAnimNotify_GameplayCueState()
+{
+}
+
+void UAnimNotify_GameplayCueState::NotifyBegin(USkeletalMeshComponent* MeshComp, UAnimSequenceBase* Animation, float TotalDuration)
+{
+	Super::NotifyBegin(MeshComp, Animation, TotalDuration);
+
+	ProcessGameplayCue(&UGameplayCueManager::AddGameplayCue_NonReplicated, MeshComp, GameplayCue.GameplayCueTag, Animation);
+
+#if WITH_EDITORONLY_DATA
+	// Grab proxy tick delegate if someone registered it.
+	if (UGameplayCueManager::PreviewProxyTick.IsBound())
+	{
+		PreviewProxyTick = UGameplayCueManager::PreviewProxyTick;
+		UGameplayCueManager::PreviewProxyTick.Unbind();
+	}
+#endif // #if WITH_EDITORONLY_DATA
+}
+
+void UAnimNotify_GameplayCueState::NotifyTick(USkeletalMeshComponent * MeshComp, UAnimSequenceBase * Animation, float FrameDeltaTime)
+{
+	AActor* OwnerActor = MeshComp->GetOwner();
+
+#if WITH_EDITOR
+	if (GIsEditor && OwnerActor == nullptr)
+	{
+		// Make preview work in anim preview window.
+		UGameplayCueManager::PreviewComponent = MeshComp;
+		UGameplayCueManager::PreviewWorld = MeshComp->GetWorld();
+
+#if WITH_EDITORONLY_DATA
+		PreviewProxyTick.ExecuteIfBound(FrameDeltaTime);
+#endif // #if WITH_EDITORONLY_DATA
+	}
+#endif // #if WITH_EDITOR
+
+	Super::NotifyTick(MeshComp, Animation, FrameDeltaTime);
+
+#if WITH_EDITOR
+	if (GIsEditor)
+	{
+		UGameplayCueManager::PreviewComponent = nullptr;
+		UGameplayCueManager::PreviewWorld = nullptr;
+	}
+#endif // #if WITH_EDITOR
+}
+
+FString UAnimNotify_GameplayCueState::GetNotifyName_Implementation() const
+{
+	FString DisplayName = TEXT("GameplayCue");
+
+	if (GameplayCue.GameplayCueTag.IsValid())
+	{
+		DisplayName = GameplayCue.GameplayCueTag.ToString();
+		DisplayName += TEXT(" (Looping)");
+	}
+
+	return DisplayName;
+}
+
+void UAnimNotify_GameplayCueState::NotifyEnd(USkeletalMeshComponent* MeshComp, UAnimSequenceBase* Animation)
+{
+	Super::NotifyEnd(MeshComp, Animation);
+
+	ProcessGameplayCue(&UGameplayCueManager::RemoveGameplayCue_NonReplicated, MeshComp, GameplayCue.GameplayCueTag, Animation);
+}
+
+#if WITH_EDITOR
+bool UAnimNotify_GameplayCueState::CanBePlaced(UAnimSequenceBase* Animation) const
+{
+	return (Animation && Animation->IsA(UAnimMontage::StaticClass()));
+}
+#endif // #if WITH_EDITOR
diff --git a/Engine/Plugins/Runtime/GameplayAbilities/Source/GameplayAbilities/Private/GameplayCueManager.cpp b/Engine/Plugins/Runtime/GameplayAbilities/Source/GameplayAbilities/Private/GameplayCueManager.cpp
--- a/Engine/Plugins/Runtime/GameplayAbilities/Source/GameplayAbilities/Private/GameplayCueManager.cpp
+++ b/Engine/Plugins/Runtime/GameplayAbilities/Source/GameplayAbilities/Private/GameplayCueManager.cpp
@@ -16,6 +16,7 @@
 #include "GameplayCueSet.h"
 #include "GameplayCueNotify_Static.h"
 #include "AbilitySystemComponent.h"
+#include "AbilitySystemInterface.h"
 #include "Net/DataReplication.h"
 #include "Engine/ActorChannel.h"
 #include "Engine/NetConnection.h"
@@ -232,6 +233,47 @@ void UGameplayCueManager::TranslateGameplayCue(FGameplayTag& Tag, AActor* Target
 	TranslationManager.TranslateTag(Tag, TargetActor, Parameters);
 }
 
+void UGameplayCueManager::AddGameplayCue_NonReplicated(AActor* Target, const FGameplayTag GameplayCueTag, const FGameplayCueParameters& Parameters)
+{
+	if (const IAbilitySystemInterface* ASI = Cast<IAbilitySystemInterface>(Target))
+	{
+		if (UAbilitySystemComponent* ASC = ASI->GetAbilitySystemComponent())
+		{
+			ASC->AddLooseGameplayTag(GameplayCueTag);
+		}
+	}
+
+	if (UGameplayCueManager* GCM = UAbilitySystemGlobals::Get().GetGameplayCueManager())
+	{
+		GCM->HandleGameplayCue(Target, GameplayCueTag, EGameplayCueEvent::OnActive, Parameters);
+		GCM->HandleGameplayCue(Target, GameplayCueTag, EGameplayCueEvent::WhileActive, Parameters);
+	}
+}
+
+void UGameplayCueManager::RemoveGameplayCue_NonReplicated(AActor* Target, const FGameplayTag GameplayCueTag, const FGameplayCueParameters& Parameters)
+{
+	if (const IAbilitySystemInterface* ASI = Cast<IAbilitySystemInterface>(Target))
+	{
+		if (UAbilitySystemComponent* ASC = ASI->GetAbilitySystemComponent())
+		{
+			ASC->RemoveLooseGameplayTag(GameplayCueTag);
+		}
+	}
+
+	if (UGameplayCueManager* GCM = UAbilitySystemGlobals::Get().GetGameplayCueManager())
+	{
+		GCM->HandleGameplayCue(Target, GameplayCueTag, EGameplayCueEvent::Removed, Parameters);
+	}
+}
+
+void UGameplayCueManager::ExecuteGameplayCue_NonReplicated(AActor* Target, const FGameplayTag GameplayCueTag, const FGameplayCueParameters& Parameters)
+{
+	if (UGameplayCueManager* GCM = UAbilitySystemGlobals::Get().GetGameplayCueManager())
+	{
+		GCM->HandleGameplayCue(Target, GameplayCueTag, EGameplayCueEvent::Executed, Parameters);
+	}
+}
+
 void UGameplayCueManager::EndGameplayCuesFor(AActor* TargetActor)
 {
 	for (auto It = NotifyMapActor.CreateIterator(); It; ++It)
diff --git a/Engine/Plugins/Runtime/GameplayAbilities/Source/GameplayAbilities/Private/GameplayCueNotifyTypes.cpp b/Engine/Plugins/Runtime/GameplayAbilities/Source/GameplayAbilities/Private/GameplayCueNotifyTypes.cpp
new file mode 100644
--- /dev/null
+++ b/Engine/Plugins/Runtime/GameplayAbilities/Source/GameplayAbilities/Private/GameplayCueNotifyTypes.cpp
@@ -0,0 +1,1133 @@
+// Copyright Epic Games, Inc. All Rights Reserved.
+
+#include "GameplayCueNotifyTypes.h"
+#include "Components/SkeletalMeshComponent.h"
+#include "Particles/ParticleSystemComponent.h"
+#include "Particles/EmitterCameraLensEffectBase.h"
+#include "Components/AudioComponent.h"
+#include "Sound/SoundBase.h"
+#include "Camera/CameraShake.h"
+#include "Camera/CameraAnim.h"
+#include "Camera/CameraAnimInst.h"
+#include "Components/ForceFeedbackComponent.h"
+#include "Engine/DecalActor.h"
+#include "Components/DecalComponent.h"
+#include "Kismet/GameplayStatics.h"
+#include "GameFramework/PlayerController.h"
+#include "GameFramework/PlayerState.h"
+#include "GameFramework/Character.h"
+
+
+DEFINE_LOG_CATEGORY(LogGameplayCueNotify);
+
+
+#define LOCTEXT_NAMESPACE "GameplayCueNotify"
+
+
+static APlayerController* FindPlayerControllerFromActor(AActor* Actor)
+{
+	if (!Actor)
+	{
+		return nullptr;
+	}
+
+	// Check the actual actor first.
+	APlayerController* PC = Cast<APlayerController>(Actor);
+	if (PC)
+	{
+		return PC;
+	}
+
+	// Check using pawn next.
+	APawn* Pawn = Cast<APawn>(Actor);
+	if (!Pawn)
+	{
+		Pawn = Actor->GetInstigator<APawn>();
+	}
+
+	PC = (Pawn ? Cast<APlayerController>(Pawn->Controller) : nullptr);
+	if (PC)
+	{
+		return PC;
+	}
+
+	// Check using player state.
+	APlayerState* PS = Cast<APlayerState>(Actor);
+
+	PC = (PS ? Cast<APlayerController>(PS->GetOwner()) : nullptr);
+	if (PC)
+	{
+		return PC;
+	}
+
+	return nullptr;
+}
+
+static float CalculateFalloffIntensity(const FVector& SourceLocation, const FVector& TargetLocation, float InnerRadius, float OuterRadius, float FalloffExponent)
+{
+	if ((InnerRadius <= 0.0f) && (OuterRadius <= 0.0f))
+	{
+		return 1.0f;
+	}
+
+	const float DistanceSqr = FVector::DistSquared(SourceLocation, TargetLocation);
+
+	if (InnerRadius < OuterRadius)
+	{
+		const float Distance = FMath::Sqrt(DistanceSqr);
+		const float Percent = 1.0f - FMath::Clamp(((Distance - InnerRadius) / (OuterRadius - InnerRadius)), 0.0f, 1.0f);
+
+		return FMath::Pow(Percent, FalloffExponent);
+	}
+
+	// Ignore the outer radius and do a cliff falloff at the inner radius.
+	return (DistanceSqr < FMath::Square(InnerRadius)) ? 1.0f : 0.0f;
+}
+
+static EAttachLocation::Type GetAttachLocationTypeFromRule(EAttachmentRule AttachmentRule)
+{
+	switch (AttachmentRule)
+	{
+	case EAttachmentRule::KeepRelative:
+		return EAttachLocation::KeepRelativeOffset;
+
+	case EAttachmentRule::KeepWorld:
+		return EAttachLocation::KeepWorldPosition;
+
+	case EAttachmentRule::SnapToTarget:
+		return EAttachLocation::SnapToTarget;
+
+	default:
+		checkf(false, TEXT("GameplayCueNotify: Invalid attachment rule [%d]\n"), (uint8)AttachmentRule);
+	}
+
+	return EAttachLocation::KeepWorldPosition;
+}
+
+static ECameraAnimPlaySpace::Type GetCameraAnimPlaySpace(EGameplayCueNotify_EffectPlaySpace EffectPlaySpace)
+{
+	switch (EffectPlaySpace)
+	{
+	case EGameplayCueNotify_EffectPlaySpace::WorldSpace:
+		return ECameraAnimPlaySpace::World;
+
+	case EGameplayCueNotify_EffectPlaySpace::CameraSpace:
+		return ECameraAnimPlaySpace::CameraLocal;
+
+	default:
+		checkf(false, TEXT("GameplayCueNotify: Invalid effect play space [%d]\n"), (uint8)EffectPlaySpace);
+	}
+
+	return ECameraAnimPlaySpace::World;
+}
+
+
+//////////////////////////////////////////////////////////////////////////
+// FGameplayCueNotify_SpawnCondition
+//////////////////////////////////////////////////////////////////////////
+FGameplayCueNotify_SpawnCondition::FGameplayCueNotify_SpawnCondition()
+{
+	LocallyControlledSource = EGameplayCueNotify_LocallyControlledSource::InstigatorActor;
+	LocallyControlledPolicy = EGameplayCueNotify_LocallyControlledPolicy::Always;
+
+	ChanceToPlay = 1.0f;
+
+	AllowedSurfaceTypes.Reset();
+	AllowedSurfaceMask = 0x0;
+
+	RejectedSurfaceTypes.Reset();
+	RejectedSurfaceMask = 0x0;
+}
+
+bool FGameplayCueNotify_SpawnCondition::ShouldSpawn(const FGameplayCueNotify_SpawnContext& SpawnContext) const
+{
+	const APlayerController* PC = SpawnContext.FindLocalPlayerController(LocallyControlledSource);
+
+	const bool bIsLocallyControlledSource = (PC != nullptr);
+
+	if ((LocallyControlledPolicy == EGameplayCueNotify_LocallyControlledPolicy::LocalOnly) && !bIsLocallyControlledSource)
+	{
+		return false;
+	}
+	else if ((LocallyControlledPolicy == EGameplayCueNotify_LocallyControlledPolicy::NotLocal) && bIsLocallyControlledSource)
+	{
+		return false;
+	}
+
+	if ((ChanceToPlay < 1.0f) && (ChanceToPlay < FMath::FRand()))
+	{
+		return false;
+	}
+
+	const FGameplayCueNotify_SurfaceMask SpawnContextSurfaceMask = (1ULL << SpawnContext.SurfaceType);
+
+	if (AllowedSurfaceTypes.Num() > 0)
+	{
+		// Build the surface type mask if needed.
+		if (AllowedSurfaceMask == 0x0)
+		{
+			for (EPhysicalSurface SurfaceType : AllowedSurfaceTypes)
+			{
+				AllowedSurfaceMask |= (1ULL << SurfaceType);
+			}
+		}
+
+		if ((AllowedSurfaceMask & SpawnContextSurfaceMask) == 0x0)
+		{
+			return false;
+		}
+	}
+
+	if (RejectedSurfaceTypes.Num() > 0)
+	{
+		// Build the surface type mask if needed.
+		if (RejectedSurfaceMask == 0x0)
+		{
+			for (EPhysicalSurface SurfaceType : RejectedSurfaceTypes)
+			{
+				RejectedSurfaceMask |= (1ULL << SurfaceType);
+			}
+		}
+
+		if ((RejectedSurfaceMask & SpawnContextSurfaceMask) != 0x0)
+		{
+			return false;
+		}
+	}
+
+	return true;
+}
+
+
+//////////////////////////////////////////////////////////////////////////
+// FGameplayCueNotify_PlacementInfo
+//////////////////////////////////////////////////////////////////////////
+FGameplayCueNotify_PlacementInfo::FGameplayCueNotify_PlacementInfo()
+{
+	SocketName = NAME_None;
+	AttachPolicy = EGameplayCueNotify_AttachPolicy::DoNotAttach;
+	AttachmentRule = EAttachmentRule::KeepWorld;
+
+	bOverrideRotation = false;
+	bOverrideScale = true;
+	RotationOverride = FRotator::ZeroRotator;
+	ScaleOverride = FVector::OneVector;
+}
+
+bool FGameplayCueNotify_PlacementInfo::FindSpawnTransform(const FGameplayCueNotify_SpawnContext& SpawnContext, FTransform& OutTransform) const
+{
+	const FGameplayCueParameters& CueParameters = SpawnContext.CueParameters;
+
+	bool bSetTransform = false;
+
+	// First attempt to get the transform from the hit result.
+	// If that fails, use the gameplay cue parameters.
+	// If that fails, fall back to the target actor.
+	if (SpawnContext.HitResult && SpawnContext.HitResult->bBlockingHit)
+	{
+		OutTransform = FTransform(SpawnContext.HitResult->ImpactNormal.Rotation(), SpawnContext.HitResult->ImpactPoint);
+		bSetTransform = true;
+	}
+	else if (!CueParameters.Location.IsZero())
+	{
+		OutTransform = FTransform(CueParameters.Normal.Rotation(), CueParameters.Location);
+		bSetTransform = true;
+	}
+	else if (SpawnContext.TargetComponent)
+	{
+		OutTransform = SpawnContext.TargetComponent->GetSocketTransform(SocketName);
+		bSetTransform = true;
+	}
+
+	if (bSetTransform)
+	{
+		if (bOverrideRotation)
+		{
+			OutTransform.SetRotation(RotationOverride.Quaternion());
+		}
+
+		if (bOverrideScale)
+		{
+			OutTransform.SetScale3D(ScaleOverride);
+		}
+	}
+	else
+	{
+		// This can happen if the target actor if destroyed out from under us.  We could alternatively pass the actor down, from ::GameplayCueFinishedCallback.
+		// Right now it seems to make more sense to skip spawning more stuff in this case.  As the OnRemove effects will spawn if the GC OnRemove is actually broadcasted "naturally".
+		// The actor dieing and forcing the OnRemove event is more of a "please cleanup everything" event. This could be changed in the future.
+		UE_LOG(LogGameplayCueNotify, Log, TEXT("GameplayCueNotify: Failed to find spawn transform for gameplay cue notify [%s]."), *CueParameters.MatchedTagName.ToString());
+	}
+
+	return bSetTransform;
+}
+
+void FGameplayCueNotify_PlacementInfo::SetComponentTransform(USceneComponent* Component, const FTransform& Transform) const
+{
+	check(Component);
+
+	Component->SetAbsolute(Component->IsUsingAbsoluteLocation(), bOverrideRotation, bOverrideScale);
+	Component->SetWorldTransform(Transform);
+}
+
+void FGameplayCueNotify_PlacementInfo::TryComponentAttachment(USceneComponent* Component, USceneComponent* AttachComponent) const
+{
+	check(Component);
+
+	if (AttachComponent && (AttachPolicy == EGameplayCueNotify_AttachPolicy::AttachToTarget))
+	{
+		FAttachmentTransformRules AttachmentRules(AttachmentRule, true);
+
+		Component->AttachToComponent(AttachComponent, AttachmentRules, SocketName);
+	}
+}
+
+
+//////////////////////////////////////////////////////////////////////////
+// FGameplayCueNotify_SpawnContext
+//////////////////////////////////////////////////////////////////////////
+FGameplayCueNotify_SpawnContext::FGameplayCueNotify_SpawnContext(UWorld* InWorld, AActor* InTargetActor, const FGameplayCueParameters& InCueParameters)
+	: World(InWorld)
+	, TargetActor(InTargetActor)
+	, CueParameters(InCueParameters)
+{
+	HitResult = nullptr;
+	TargetComponent = nullptr;
+	SurfaceType = EPhysicalSurface::SurfaceType_Default;
+
+	DefaultSpawnCondition = nullptr;
+	DefaultPlacementInfo = nullptr;
+
+	InitializeContext();
+}
+
+void FGameplayCueNotify_SpawnContext::InitializeContext()
+{
+	/*
+		Note that FGameplayCueParameters::EffectContext may not be valid in some projects.
+		Projects can override UAbilitySystemGlobals::InitGameplayCueParameters() to extract only necessary data from FGameplayEffectContext directly into FGameplayCueParameters.
+		See UAbilitySystemGlobals::InitGameplayCueParameters() and AbilitySystem.AlwaysConvertGESpecToGCParams for more info.
+	*/
+
+	if (!ensure(TargetActor))
+	{
+		return;
+	}
+
+	if (CueParameters.EffectContext.IsValid())
+	{
+		HitResult = CueParameters.EffectContext.GetHitResult();
+	}
+
+	// Find the target component.
+	TargetComponent = CueParameters.TargetAttachComponent.Get();
+	if (!TargetComponent)
+	{
+		const ACharacter* TargetCharacter = Cast<ACharacter>(TargetActor);
+
+		// Always use the mesh for characters since the mesh position will be smoothed out.
+		TargetComponent = (TargetCharacter ? TargetCharacter->GetMesh() : TargetActor->GetRootComponent());
+	}
+
+	// Find the surface type.
+	const UPhysicalMaterial* PhysicalMaterial = nullptr;
+
+	if (HitResult && HitResult->PhysMaterial.IsValid())
+	{
+		PhysicalMaterial = HitResult->PhysMaterial.Get();
+	}
+	else
+	{
+		PhysicalMaterial = CueParameters.PhysicalMaterial.Get();
+	}
+
+	if (PhysicalMaterial)
+	{
+		SurfaceType = PhysicalMaterial->SurfaceType;
+	}
+}
+
+APlayerController* FGameplayCueNotify_SpawnContext::FindLocalPlayerController(EGameplayCueNotify_LocallyControlledSource Source) const
+{
+	AActor* ActorToSearch = nullptr;
+
+	switch (Source)
+	{
+	case EGameplayCueNotify_LocallyControlledSource::InstigatorActor:
+		ActorToSearch = CueParameters.Instigator.Get();
+		break;
+
+	case EGameplayCueNotify_LocallyControlledSource::TargetActor:
+		ActorToSearch = TargetActor;
+		break;
+
+	default:
+		break;
+	}
+
+	APlayerController* PC = FindPlayerControllerFromActor(ActorToSearch);
+
+	return (PC && PC->IsLocalController()) ? PC : nullptr;
+}
+
+
+//////////////////////////////////////////////////////////////////////////
+// FGameplayCueNotify_ParticleInfo
+//////////////////////////////////////////////////////////////////////////
+FGameplayCueNotify_ParticleInfo::FGameplayCueNotify_ParticleInfo()
+{
+	ParticleSystem = nullptr;
+	bOverrideSpawnCondition = false;
+	bOverridePlacementInfo = false;
+	bCastShadow = false;
+}
+
+bool FGameplayCueNotify_ParticleInfo::PlayParticleEffect(const FGameplayCueNotify_SpawnContext& SpawnContext, FGameplayCueNotify_SpawnResult& OutSpawnResult) const
+{
+	UParticleSystemComponent* SpawnedPSC = nullptr;
+
+	if (ParticleSystem != nullptr)
+	{
+		const FGameplayCueNotify_SpawnCondition& SpawnCondition = SpawnContext.GetSpawnCondition(bOverrideSpawnCondition, SpawnConditionOverride);
+		const FGameplayCueNotify_PlacementInfo& PlacementInfo = SpawnContext.GetPlacementInfo(bOverridePlacementInfo, PlacementInfoOverride);
+
+		if (SpawnCondition.ShouldSpawn(SpawnContext))
+		{
+			FTransform SpawnTransform;
+
+			if (PlacementInfo.FindSpawnTransform(SpawnContext, SpawnTransform))
+			{
+				check(SpawnContext.World);
+
+				SpawnedPSC = SpawnContext.World->GetPSCPool().CreateWorldParticleSystem(ParticleSystem, SpawnContext.World, EPSCPoolMethod::AutoRelease);
+				if (ensure(SpawnedPSC))
+				{
+					PlacementInfo.SetComponentTransform(SpawnedPSC, SpawnTransform);
+					PlacementInfo.TryComponentAttachment(SpawnedPSC, SpawnContext.TargetComponent);
+
+					SpawnedPSC->RegisterComponentWithWorld(SpawnContext.World);
+					SpawnedPSC->SetCastShadow(bCastShadow);
+					SpawnedPSC->Activate();
+				}
+			}
+		}
+	}
+
+	// Always add to the list, even if null, so that the list is stable and in order for blueprint users.
+	OutSpawnResult.ParticleComponents.Add(SpawnedPSC);
+
+	return (SpawnedPSC != nullptr);
+}
+
+void FGameplayCueNotify_ParticleInfo::ValidateBurstAssets(UObject* ContainingAsset, const FString& Context, TArray<FText>& ValidationErrors) const
+{
+#if WITH_EDITORONLY_DATA
+	if (ParticleSystem != nullptr)
+	{
+		if (ParticleSystem->IsLooping())
+		{
+			ValidationErrors.Add(FText::Format(
+				LOCTEXT("ParticleSystem_ShouldNotLoop", "Particle system [{0}] used in slot [{1}] for asset [{2}] is set to looping, but the slot is a one-shot (the instance will leak)."),
+				FText::AsCultureInvariant(ParticleSystem->GetPathName()),
+				FText::AsCultureInvariant(Context),
+				FText::AsCultureInvariant(ContainingAsset->GetPathName())));
+		}
+	}
+#endif // #if WITH_EDITORONLY_DATA
+}
+
+
+//////////////////////////////////////////////////////////////////////////
+// FGameplayCueNotify_SoundInfo
+//////////////////////////////////////////////////////////////////////////
+FGameplayCueNotify_SoundInfo::FGameplayCueNotify_SoundInfo()
+{
+	SoundCue = nullptr;
+	LoopingFadeOutDuration = 0.5f;
+	LoopingFadeVolumeLevel = 0.0f;
+	bOverrideSpawnCondition = false;
+	bOverridePlacementInfo = false;
+}
+
+bool FGameplayCueNotify_SoundInfo::PlaySound(const FGameplayCueNotify_SpawnContext& SpawnContext, FGameplayCueNotify_SpawnResult& OutSpawnResult) const
+{
+	UAudioComponent* AudioComponent = nullptr;
+	bool bSoundPlayed = false;
+
+	if (SoundCue != nullptr)
+	{
+		const FGameplayCueNotify_SpawnCondition& SpawnCondition = SpawnContext.GetSpawnCondition(bOverrideSpawnCondition, SpawnConditionOverride);
+		const FGameplayCueNotify_PlacementInfo& PlacementInfo = SpawnContext.GetPlacementInfo(bOverridePlacementInfo, PlacementInfoOverride);
+
+		if (SpawnCondition.ShouldSpawn(SpawnContext))
+		{
+			FTransform SpawnTransform;
+
+			if (PlacementInfo.FindSpawnTransform(SpawnContext, SpawnTransform))
+			{
+				const bool bStopWhenAttachedToDestroyed = false;
+				const bool bAutoDestroy = true;
+				const float VolumeMultiplier = 1.0f;
+				const float PitchMultiplier = 1.0f;
+				const float StartTime = 0.0f;
+				USoundAttenuation* AttenuationSettings = nullptr;
+				USoundConcurrency* ConcurrencySettings = nullptr;
+
+				const FVector SpawnLocation = SpawnTransform.GetLocation();
+				const FRotator SpawnRotation = SpawnTransform.Rotator();
+
+				if (SpawnContext.TargetComponent && (PlacementInfo.AttachPolicy == EGameplayCueNotify_AttachPolicy::AttachToTarget))
+				{
+					const EAttachLocation::Type AttachLocationType = GetAttachLocationTypeFromRule(PlacementInfo.AttachmentRule);
+
+					AudioComponent = UGameplayStatics::SpawnSoundAttached(SoundCue, SpawnContext.TargetComponent, PlacementInfo.SocketName,
+						SpawnLocation, SpawnRotation, AttachLocationType, bStopWhenAttachedToDestroyed,
+						VolumeMultiplier, PitchMultiplier, StartTime, AttenuationSettings, ConcurrencySettings, bAutoDestroy);
+				}
+				else if (SoundCue->IsLooping())
+				{
+					AudioComponent = UGameplayStatics::SpawnSoundAtLocation(SpawnContext.World, SoundCue, SpawnLocation, SpawnRotation,
+						VolumeMultiplier, PitchMultiplier, StartTime, AttenuationSettings, ConcurrencySettings, bAutoDestroy);
+				}
+				else
+				{
+					UGameplayStatics::PlaySoundAtLocation(SpawnContext.World, SoundCue, SpawnLocation, SpawnRotation,
+						VolumeMultiplier, PitchMultiplier, StartTime, AttenuationSettings, ConcurrencySettings, SpawnContext.TargetActor);
+				}
+
+				bSoundPlayed = true;
+			}
+		}
+	}
+
+	// Always add to the list, even if null, so that the list is stable and in order for blueprint users.
+	OutSpawnResult.AudioComponents.Add(AudioComponent);
+
+	return bSoundPlayed;
+}
+
+void FGameplayCueNotify_SoundInfo::ValidateBurstAssets(UObject* ContainingAsset, const FString& Context, TArray<FText>& ValidationErrors) const
+{
+#if WITH_EDITORONLY_DATA
+	if (SoundCue != nullptr)
+	{
+		if (SoundCue->IsLooping())
+		{
+			ValidationErrors.Add(FText::Format(
+				LOCTEXT("SoundCue_ShouldNotLoop", "Sound [{0}] used in slot [{1}] for asset [{2}] is set to looping, but the slot is a one-shot (the instance will leak)."),
+				FText::AsCultureInvariant(SoundCue->GetPathName()),
+				FText::AsCultureInvariant(Context),
+				FText::AsCultureInvariant(ContainingAsset->GetPathName())));
+		}
+	}
+#endif //#if WITH_EDITORONLY_DATA
+}
+
+
+//////////////////////////////////////////////////////////////////////////
+// FGameplayCueNotify_CameraShakeInfo
+//////////////////////////////////////////////////////////////////////////
+FGameplayCueNotify_CameraShakeInfo::FGameplayCueNotify_CameraShakeInfo()
+{
+	CameraShake = nullptr;
+	ShakeScale = 1.0f;
+
+	PlaySpace = EGameplayCueNotify_EffectPlaySpace::CameraSpace;
+
+	bOverrideSpawnCondition = false;
+	bOverridePlacementInfo = false;
+
+	bPlayInWorld = false;
+	WorldInnerRadius = 0.0f;
+	WorldOuterRadius = 0.0f;
+	WorldFalloffExponent = 1.0f;
+}
+
+bool FGameplayCueNotify_CameraShakeInfo::PlayCameraShake(const FGameplayCueNotify_SpawnContext& SpawnContext, FGameplayCueNotify_SpawnResult& OutSpawnResult) const
+{
+	if (!CameraShake || (ShakeScale <= 0.0f))
+	{
+		return false;
+	}
+
+	const FGameplayCueNotify_SpawnCondition& SpawnCondition = SpawnContext.GetSpawnCondition(bOverrideSpawnCondition, SpawnConditionOverride);
+	const FGameplayCueNotify_PlacementInfo& PlacementInfo = SpawnContext.GetPlacementInfo(bOverridePlacementInfo, PlacementInfoOverride);
+
+	if (!SpawnCondition.ShouldSpawn(SpawnContext))
+	{
+		return false;
+	}
+
+	const ECameraAnimPlaySpace::Type CameraAnimPlaySpace = GetCameraAnimPlaySpace(PlaySpace);
+
+	// Check if the camera shake should be played in world for all players.
+	if (bPlayInWorld)
+	{
+		FTransform SpawnTransform;
+		if (!PlacementInfo.FindSpawnTransform(SpawnContext, SpawnTransform))
+		{
+			return false;
+		}
+
+		const FVector SpawnLocation = SpawnTransform.GetLocation();
+
+		check(SpawnContext.World);
+		for (FConstPlayerControllerIterator Iterator = SpawnContext.World->GetPlayerControllerIterator(); Iterator; ++Iterator)
+		{
+			APlayerController* PC = Iterator->Get();
+			if (!PC || !PC->PlayerCameraManager || !PC->IsLocalController())
+			{
+				continue;
+			}
+
+			const APawn* Pawn = PC->GetPawn();
+			if (!Pawn)
+			{
+				continue;
+			}
+
+			const FVector PawnLocation = Pawn->GetActorLocation();
+			const float FalloffIntensity = CalculateFalloffIntensity(SpawnLocation, PawnLocation, WorldInnerRadius, WorldOuterRadius, WorldFalloffExponent);
+
+			if (FalloffIntensity <= 0.0f)
+			{
+				continue;
+			}
+
+			UCameraShake* CameraShakeInstance = PC->PlayerCameraManager->PlayCameraShake(CameraShake, (ShakeScale * FalloffIntensity), CameraAnimPlaySpace, FRotator::ZeroRotator);
+			OutSpawnResult.CameraShakes.Add(CameraShakeInstance);
+		}
+	}
+	else
+	{
+		// Play the camera shake directly on the target player controller.  No intensity falloff is applied when playing directly on the target.
+		APlayerController* TargetPC = SpawnContext.FindLocalPlayerController(EGameplayCueNotify_LocallyControlledSource::TargetActor);
+
+		if (TargetPC && TargetPC->PlayerCameraManager)
+		{
+			UCameraShake* CameraShakeInstance = TargetPC->PlayerCameraManager->PlayCameraShake(CameraShake, ShakeScale, CameraAnimPlaySpace, FRotator::ZeroRotator);
+			OutSpawnResult.CameraShakes.Add(CameraShakeInstance);
+		}
+	}
+
+	return (OutSpawnResult.CameraShakes.Num() > 0);
+}
+
+void FGameplayCueNotify_CameraShakeInfo::ValidateBurstAssets(UObject* ContainingAsset, const FString& Context, TArray<FText>& ValidationErrors) const
+{
+#if WITH_EDITORONLY_DATA
+	if (CameraShake != nullptr)
+	{
+		UCameraShake* CameraShakeCDO = CameraShake->GetDefaultObject<UCameraShake>();
+		if (CameraShakeCDO->IsLooping())
+		{
+			ValidationErrors.Add(FText::Format(
+				LOCTEXT("CameraShake_ShouldNotLoop", "Camera shake [{0}] used in slot [{1}] for asset [{2}] will oscillate forever, but the slot is a one-shot (the instance will leak)."),
+				FText::AsCultureInvariant(CameraShake->GetPathName()),
+				FText::AsCultureInvariant(Context),
+				FText::AsCultureInvariant(ContainingAsset->GetPathName())));
+		}
+	}
+#endif // #if WITH_EDITORONLY_DATA
+}
+
+
+//////////////////////////////////////////////////////////////////////////
+// FGameplayCueNotify_CameraAnimInfo
+//////////////////////////////////////////////////////////////////////////
+FGameplayCueNotify_CameraAnimInfo::FGameplayCueNotify_CameraAnimInfo()
+{
+	CameraAnim = nullptr;
+	AnimScale = 1.0f;
+	PlayRate = 1.0f;
+
+	BlendInTime = 0.0f;
+	BlendOutTime = 0.0f;
+
+	PlaySpace = EGameplayCueNotify_EffectPlaySpace::CameraSpace;
+
+	bIsLooping = false;
+	bRandomStartTime = false;
+
+	bOverrideSpawnCondition = false;
+}
+
+bool FGameplayCueNotify_CameraAnimInfo::PlayCameraAnim(const FGameplayCueNotify_SpawnContext& SpawnContext, FGameplayCueNotify_SpawnResult& OutSpawnResult) const
+{
+	if (!CameraAnim || (AnimScale <= 0.0f) || (PlayRate <= 0.0f))
+	{
+		return false;
+	}
+
+	const FGameplayCueNotify_SpawnCondition& SpawnCondition = SpawnContext.GetSpawnCondition(bOverrideSpawnCondition, SpawnConditionOverride);
+
+	if (!SpawnCondition.ShouldSpawn(SpawnContext))
+	{
+		return false;
+	}
+
+	// Camera animations only play on the target actor when it is locally controlled.
+	APlayerController* TargetPC = SpawnContext.FindLocalPlayerController(EGameplayCueNotify_LocallyControlledSource::TargetActor);
+	if (TargetPC && TargetPC->PlayerCameraManager)
+	{
+		const ECameraAnimPlaySpace::Type CameraAnimPlaySpace = GetCameraAnimPlaySpace(PlaySpace);
+
+		OutSpawnResult.CameraAnim = TargetPC->PlayerCameraManager->PlayCameraAnim(CameraAnim, PlayRate, AnimScale, BlendInTime, BlendOutTime, bIsLooping, bRandomStartTime, 0.0f, CameraAnimPlaySpace, FRotator::ZeroRotator);
+	}
+
+	return (OutSpawnResult.CameraAnim != nullptr);
+}
+
+void FGameplayCueNotify_CameraAnimInfo::ValidateBurstAssets(UObject* ContainingAsset, const FString& Context, TArray<FText>& ValidationErrors) const
+{
+#if WITH_EDITORONLY_DATA
+	if (CameraAnim)
+	{
+		if (bIsLooping)
+		{
+			ValidationErrors.Add(FText::Format(
+				LOCTEXT("CameraAnim_ShouldNotLoop", "Camera animation [{0}] used in slot [{1}] for asset [{2}] is set to looping, but the slot is a one-shot (the instance will leak)."),
+				FText::AsCultureInvariant(CameraAnim->GetPathName()),
+				FText::AsCultureInvariant(Context),
+				FText::AsCultureInvariant(ContainingAsset->GetPathName())));
+		}
+	}
+#endif // #if WITH_EDITORONLY_DATA
+}
+
+
+//////////////////////////////////////////////////////////////////////////
+// FGameplayCueNotify_CameraLensEffectInfo
+//////////////////////////////////////////////////////////////////////////
+FGameplayCueNotify_CameraLensEffectInfo::FGameplayCueNotify_CameraLensEffectInfo()
+{
+	CameraLensEffect = nullptr;
+
+	bOverrideSpawnCondition = false;
+	bOverridePlacementInfo = false;
+
+	bPlayInWorld = false;
+	WorldInnerRadius = 0.0f;
+	WorldOuterRadius = 0.0f;
+}
+
+bool FGameplayCueNotify_CameraLensEffectInfo::PlayCameraLensEffect(const FGameplayCueNotify_SpawnContext& SpawnContext, FGameplayCueNotify_SpawnResult& OutSpawnResult) const
+{
+	if (!CameraLensEffect)
+	{
+		return false;
+	}
+
+	const FGameplayCueNotify_SpawnCondition& SpawnCondition = SpawnContext.GetSpawnCondition(bOverrideSpawnCondition, SpawnConditionOverride);
+	const FGameplayCueNotify_PlacementInfo& PlacementInfo = SpawnContext.GetPlacementInfo(bOverridePlacementInfo, PlacementInfoOverride);
+
+	if (!SpawnCondition.ShouldSpawn(SpawnContext))
+	{
+		return false;
+	}
+
+	// Check if the camera lens effect should be played in world for all players.
+	if (bPlayInWorld)
+	{
+		FTransform SpawnTransform;
+		if (!PlacementInfo.FindSpawnTransform(SpawnContext, SpawnTransform))
+		{
+			return false;
+		}
+
+		const FVector SpawnLocation = SpawnTransform.GetLocation();
+
+		check(SpawnContext.World);
+		for (FConstPlayerControllerIterator Iterator = SpawnContext.World->GetPlayerControllerIterator(); Iterator; ++Iterator)
+		{
+			APlayerController* PC = Iterator->Get();
+			if (!PC || !PC->PlayerCameraManager || !PC->IsLocalController())
+			{
+				continue;
+			}
+
+			const APawn* Pawn = PC->GetPawn();
+			if (!Pawn)
+			{
+				continue;
+			}
+
+			const FVector PawnLocation = Pawn->GetActorLocation();
+			const float FalloffIntensity = CalculateFalloffIntensity(SpawnLocation, PawnLocation, WorldInnerRadius, WorldOuterRadius, 1.0f);
+
+			if (FalloffIntensity <= 0.0f)
+			{
+				continue;
+			}
+
+			AEmitterCameraLensEffectBase* CameraLensEffectInstance = PC->PlayerCameraManager->AddCameraLensEffect(CameraLensEffect);
+			OutSpawnResult.CameraLensEffects.Add(CameraLensEffectInstance);
+		}
+	}
+	else
+	{
+		// Play the camera lens effect directly on the target player controller.  No radius check when playing directly on the target.
+		APlayerController* TargetPC = SpawnContext.FindLocalPlayerController(EGameplayCueNotify_LocallyControlledSource::TargetActor);
+
+		if (TargetPC && TargetPC->PlayerCameraManager)
+		{
+			AEmitterCameraLensEffectBase* CameraLensEffectInstance = TargetPC->PlayerCameraManager->AddCameraLensEffect(CameraLensEffect);
+			OutSpawnResult.CameraLensEffects.Add(CameraLensEffectInstance);
+		}
+	}
+
+	return (OutSpawnResult.CameraLensEffects.Num() > 0);
+}
+
+void FGameplayCueNotify_CameraLensEffectInfo::ValidateBurstAssets(UObject* ContainingAsset, const FString& Context, TArray<FText>& ValidationErrors) const
+{
+#if WITH_EDITORONLY_DATA
+	if (CameraLensEffect != nullptr)
+	{
+		AEmitterCameraLensEffectBase* CameraLensEffectCDO = CameraLensEffect->GetDefaultObject<AEmitterCameraLensEffectBase>();
+		if (CameraLensEffectCDO->IsLooping())
+		{
+			ValidationErrors.Add(FText::Format(
+				LOCTEXT("CameraLensEffect_ShouldNotLoop", "Camera lens effect [{0}] used in slot [{1}] for asset [{2}] is set to looping, but the slot is a one-shot (the instance will leak)."),
+				FText::AsCultureInvariant(CameraLensEffect->GetPathName()),
+				FText::AsCultureInvariant(Context),
+				FText::AsCultureInvariant(ContainingAsset->GetPathName())));
+		}
+	}
+#endif // #if WITH_EDITORONLY_DATA
+}
+
+
+//////////////////////////////////////////////////////////////////////////
+// FGameplayCueNotify_ForceFeedbackInfo
+//////////////////////////////////////////////////////////////////////////
+FGameplayCueNotify_ForceFeedbackInfo::FGameplayCueNotify_ForceFeedbackInfo()
+{
+	ForceFeedbackEffect = nullptr;
+	ForceFeedbackTag = NAME_None;
+
+	bIsLooping = false;
+	bOverrideSpawnCondition = false;
+	bOverridePlacementInfo = false;
+
+	bPlayInWorld = false;
+	WorldIntensity = 1.0f;
+	WorldAttenuation = nullptr;
+}
+
+bool FGameplayCueNotify_ForceFeedbackInfo::PlayForceFeedback(const FGameplayCueNotify_SpawnContext& SpawnContext, FGameplayCueNotify_SpawnResult& OutSpawnResult) const
+{
+	if (!ForceFeedbackEffect)
+	{
+		return false;
+	}
+
+	const FGameplayCueNotify_SpawnCondition& SpawnCondition = SpawnContext.GetSpawnCondition(bOverrideSpawnCondition, SpawnConditionOverride);
+	const FGameplayCueNotify_PlacementInfo& PlacementInfo = SpawnContext.GetPlacementInfo(bOverridePlacementInfo, PlacementInfoOverride);
+
+	if (!SpawnCondition.ShouldSpawn(SpawnContext))
+	{
+		return false;
+	}
+
+	// Spawn a component only if the force feedback is playing in world for all players.
+	if (bPlayInWorld)
+	{
+		FTransform SpawnTransform;
+		if (!PlacementInfo.FindSpawnTransform(SpawnContext, SpawnTransform))
+		{
+			return false;
+		}
+
+		const FVector SpawnLocation = SpawnTransform.GetLocation();
+		const FRotator SpawnRotation = SpawnTransform.Rotator();
+		const bool bStopWhenAttachedToDestroyed = false;
+		const bool bAutoDestroy = true;
+		const float StartTime = 0.0f;
+
+		const EAttachLocation::Type AttachLocationType = GetAttachLocationTypeFromRule(PlacementInfo.AttachmentRule);
+
+		if (SpawnContext.TargetComponent && (PlacementInfo.AttachPolicy == EGameplayCueNotify_AttachPolicy::AttachToTarget))
+		{
+			OutSpawnResult.ForceFeedbackComponent = UGameplayStatics::SpawnForceFeedbackAttached(ForceFeedbackEffect, SpawnContext.TargetComponent, PlacementInfo.SocketName,
+				SpawnLocation, SpawnRotation, AttachLocationType, bStopWhenAttachedToDestroyed,
+				bIsLooping, WorldIntensity, StartTime, WorldAttenuation, bAutoDestroy);
+		}
+		else
+		{
+			OutSpawnResult.ForceFeedbackComponent = UGameplayStatics::SpawnForceFeedbackAtLocation(SpawnContext.World, ForceFeedbackEffect, SpawnLocation, SpawnRotation,
+				bIsLooping, WorldIntensity, StartTime, WorldAttenuation, bAutoDestroy);
+		}
+
+		return (OutSpawnResult.ForceFeedbackComponent != nullptr);
+	}
+
+	// Play the force feedback directly on the target player controller.  No attenuation is used when playing directly on the player controller.
+	APlayerController* TargetPC = SpawnContext.FindLocalPlayerController(EGameplayCueNotify_LocallyControlledSource::TargetActor);
+	if (TargetPC)
+	{
+		FForceFeedbackParameters ForceFeedbackParameters;
+		ForceFeedbackParameters.Tag = ForceFeedbackTag;
+		ForceFeedbackParameters.bLooping = bIsLooping;
+		ForceFeedbackParameters.bIgnoreTimeDilation = false;
+		ForceFeedbackParameters.bPlayWhilePaused = false;
+
+		TargetPC->ClientPlayForceFeedback(ForceFeedbackEffect, ForceFeedbackParameters);
+
+		OutSpawnResult.ForceFeedbackTargetPC = TargetPC;
+		return true;
+	}
+
+	return false;
+}
+
+void FGameplayCueNotify_ForceFeedbackInfo::ValidateBurstAssets(UObject* ContainingAsset, const FString& Context, TArray<FText>& ValidationErrors) const
+{
+#if WITH_EDITORONLY_DATA
+	if (ForceFeedbackEffect)
+	{
+		if (bIsLooping)
+		{
+			ValidationErrors.Add(FText::Format(
+				LOCTEXT("ForceFeedback_ShouldNotLoop", "Force feedback effect [{0}] used in slot [{1}] for asset [{2}] is set to looping, but the slot is a one-shot (the instance will leak)."),
+				FText::AsCultureInvariant(ForceFeedbackEffect->GetPathName()),
+				FText::AsCultureInvariant(Context),
+				FText::AsCultureInvariant(ContainingAsset->GetPathName())));
+		}
+	}
+#endif // #if WITH_EDITORONLY_DATA
+}
+
+
+//////////////////////////////////////////////////////////////////////////
+// FGameplayCueNotify_DecalInfo
+//////////////////////////////////////////////////////////////////////////
+FGameplayCueNotify_DecalInfo::FGameplayCueNotify_DecalInfo()
+{
+	Decal = nullptr;
+
+	bOverrideSpawnCondition = false;
+	bOverridePlacementInfo = false;
+
+	bOverrideFadeOut = false;
+	FadeOutStartDelay = 0.0f;
+	FadeOutDuration = 0.0f;
+}
+
+bool FGameplayCueNotify_DecalInfo::SpawnDecal(const FGameplayCueNotify_SpawnContext& SpawnContext, FGameplayCueNotify_SpawnResult& OutSpawnResult) const
+{
+	if (!Decal)
+	{
+		return false;
+	}
+
+	const FGameplayCueNotify_SpawnCondition& SpawnCondition = SpawnContext.GetSpawnCondition(bOverrideSpawnCondition, SpawnConditionOverride);
+	const FGameplayCueNotify_PlacementInfo& PlacementInfo = SpawnContext.GetPlacementInfo(bOverridePlacementInfo, PlacementInfoOverride);
+
+	if (!SpawnCondition.ShouldSpawn(SpawnContext))
+	{
+		return false;
+	}
+
+	FTransform SpawnTransform;
+	if (!PlacementInfo.FindSpawnTransform(SpawnContext, SpawnTransform))
+	{
+		return false;
+	}
+
+	// Invert inherited rotation as decals are reversed generally.
+	if (!PlacementInfo.bOverrideRotation)
+	{
+		const FQuat InverseRotation = SpawnTransform.GetRotation().Inverse();
+		SpawnTransform.SetRotation(InverseRotation);
+	}
+
+	ADecalActor* SpawnedDecal = SpawnContext.World->SpawnActor<ADecalActor>(Decal, SpawnTransform);
+	if (!SpawnedDecal)
+	{
+		return false;
+	}
+
+	UDecalComponent* DecalComponent = SpawnedDecal->GetDecal();
+	USceneComponent* DecalRootComponent = SpawnedDecal->GetRootComponent();
+
+	PlacementInfo.SetComponentTransform(DecalRootComponent, SpawnTransform);
+	PlacementInfo.TryComponentAttachment(DecalRootComponent, SpawnContext.TargetComponent);
+
+	if (DecalComponent && bOverrideFadeOut)
+	{
+		DecalComponent->SetFadeOut(FadeOutStartDelay, FadeOutDuration);
+	}
+
+	OutSpawnResult.DecalActor = SpawnedDecal;
+	return true;
+}
+
+
+//////////////////////////////////////////////////////////////////////////
+// FGameplayCueNotify_BurstEffects
+//////////////////////////////////////////////////////////////////////////
+FGameplayCueNotify_BurstEffects::FGameplayCueNotify_BurstEffects()
+{
+}
+
+void FGameplayCueNotify_BurstEffects::ExecuteEffects(const FGameplayCueNotify_SpawnContext& SpawnContext, FGameplayCueNotify_SpawnResult& OutSpawnResult) const
+{
+	if (!SpawnContext.World)
+	{
+		UE_LOG(LogGameplayCueNotify, Error, TEXT("GameplayCueNotify: Trying to execute Burst effects with a NULL world."));
+		return;
+	}
+
+	for (const FGameplayCueNotify_ParticleInfo& ParticleInfo : BurstParticles)
+	{
+		ParticleInfo.PlayParticleEffect(SpawnContext, OutSpawnResult);
+	}
+
+	for (const FGameplayCueNotify_SoundInfo& SoundInfo : BurstSounds)
+	{
+		SoundInfo.PlaySound(SpawnContext, OutSpawnResult);
+	}
+
+	BurstCameraShake.PlayCameraShake(SpawnContext, OutSpawnResult);
+	BurstCameraAnim.PlayCameraAnim(SpawnContext, OutSpawnResult);
+	BurstCameraLensEffect.PlayCameraLensEffect(SpawnContext, OutSpawnResult);
+	BurstForceFeedback.PlayForceFeedback(SpawnContext, OutSpawnResult);
+	BurstDecal.SpawnDecal(SpawnContext, OutSpawnResult);
+}
+
+void FGameplayCueNotify_BurstEffects::ValidateAssociatedAssets(UObject* ContainingAsset, const FString& Context, TArray<FText>& ValidationErrors) const
+{
+	for (const FGameplayCueNotify_ParticleInfo& ParticleInfo : BurstParticles)
+	{
+		ParticleInfo.ValidateBurstAssets(ContainingAsset, Context + TEXT(".BurstParticles"), ValidationErrors);
+	}
+
+	for (const FGameplayCueNotify_SoundInfo& SoundInfo : BurstSounds)
+	{
+		SoundInfo.ValidateBurstAssets(ContainingAsset, Context + TEXT(".BurstSounds"), ValidationErrors);
+	}
+
+	BurstCameraShake.ValidateBurstAssets(ContainingAsset, Context + TEXT(".BurstCameraShake"), ValidationErrors);
+	BurstCameraAnim.ValidateBurstAssets(ContainingAsset, Context + TEXT(".BurstCameraAnim"), ValidationErrors);
+	BurstCameraLensEffect.ValidateBurstAssets(ContainingAsset, Context + TEXT(".BurstCameraLensEffect"), ValidationErrors);
+	BurstForceFeedback.ValidateBurstAssets(ContainingAsset, Context + TEXT(".BurstForceFeedback"), ValidationErrors);
+}
+
+
+//////////////////////////////////////////////////////////////////////////
+// FGameplayCueNotify_LoopingEffects
+//////////////////////////////////////////////////////////////////////////
+FGameplayCueNotify_LoopingEffects::FGameplayCueNotify_LoopingEffects()
+{
+}
+
+void FGameplayCueNotify_LoopingEffects::StartEffects(const FGameplayCueNotify_SpawnContext& SpawnContext, FGameplayCueNotify_SpawnResult& OutSpawnResult) const
+{
+	if (!SpawnContext.World)
+	{
+		UE_LOG(LogGameplayCueNotify, Error, TEXT("GameplayCueNotify: Trying to start Looping effects with a NULL world."));
+		return;
+	}
+
+	for (const FGameplayCueNotify_ParticleInfo& ParticleInfo : LoopingParticles)
+	{
+		ParticleInfo.PlayParticleEffect(SpawnContext, OutSpawnResult);
+	}
+
+	for (const FGameplayCueNotify_SoundInfo& SoundInfo : LoopingSounds)
+	{
+		SoundInfo.PlaySound(SpawnContext, OutSpawnResult);
+	}
+
+	LoopingCameraShake.PlayCameraShake(SpawnContext, OutSpawnResult);
+	LoopingCameraAnim.PlayCameraAnim(SpawnContext, OutSpawnResult);
+	LoopingCameraLensEffect.PlayCameraLensEffect(SpawnContext, OutSpawnResult);
+	LoopingForceFeedback.PlayForceFeedback(SpawnContext, OutSpawnResult);
+}
+
+void FGameplayCueNotify_LoopingEffects::StopEffects(FGameplayCueNotify_SpawnResult& SpawnResult) const
+{
+	// Stop all particle effects.
+	for (UParticleSystemComponent* PSC : SpawnResult.ParticleComponents)
+	{
+		if (PSC)
+		{
+			PSC->Deactivate();
+		}
+	}
+
+	// Stop all sound effects.  This assumes there is one AudioComponent entry for each FGameplayCueNotify_SoundInfo.
+	ensure(LoopingSounds.Num() == SpawnResult.AudioComponents.Num());
+
+	for (int32 SoundIndex = 0; SoundIndex < SpawnResult.AudioComponents.Num(); ++SoundIndex)
+	{
+		UAudioComponent* AudioComponent = SpawnResult.AudioComponents[SoundIndex];
+		if (AudioComponent)
+		{
+			if (LoopingSounds.IsValidIndex(SoundIndex))
+			{
+				const FGameplayCueNotify_SoundInfo* SoundInfo = &LoopingSounds[SoundIndex];
+
+				if (SoundInfo->LoopingFadeOutDuration > 0.0f)
+				{
+					AudioComponent->FadeOut(SoundInfo->LoopingFadeOutDuration, SoundInfo->LoopingFadeVolumeLevel);
+					continue;
+				}
+			}
+
+			AudioComponent->Stop();
+		}
+	}
+
+	// Stop all camera shakes.
+	for (UCameraShake* CameraShake : SpawnResult.CameraShakes)
+	{
+		if (CameraShake)
+		{
+			const bool bStopImmediately = false;
+			CameraShake->StopShake(bStopImmediately);
+		}
+	}
+
+	// Stop the camera animation.
+	if (SpawnResult.CameraAnim)
+	{
+		const bool bStopImmediately = false;
+		SpawnResult.CameraAnim->Stop(bStopImmediately);
+	}
+
+	// Stop the camera lens effect.
+	for (AEmitterCameraLensEffectBase* CameraLensEffect : SpawnResult.CameraLensEffects)
+	{
+		if (CameraLensEffect)
+		{
+			CameraLensEffect->DeactivateLensEffect();
+		}
+	}
+
+	// Stop the force feedback.  The component is only created when the effect is played in world.
+	// If it's not in world, it needs to be stopped on the player controller.
+	if (SpawnResult.ForceFeedbackComponent)
+	{
+		SpawnResult.ForceFeedbackComponent->Stop();
+	}
+
+	if (SpawnResult.ForceFeedbackTargetPC)
+	{
+		SpawnResult.ForceFeedbackTargetPC->ClientStopForceFeedback(LoopingForceFeedback.ForceFeedbackEffect, LoopingForceFeedback.ForceFeedbackTag);
+	}
+
+	// There should be no decal on looping gameplay cues.
+	ensure(SpawnResult.DecalActor == nullptr);
+
+	// Clear the spawn results.
+	SpawnResult.Reset();
+}
+
+void FGameplayCueNotify_LoopingEffects::ValidateAssociatedAssets(UObject* ContainingAsset, const FString& Context, TArray<FText>& ValidationErrors) const
+{
+}
+
+
+#undef LOCTEXT_NAMESPACE
diff --git a/Engine/Plugins/Runtime/GameplayAbilities/Source/GameplayAbilities/Private/GameplayCueNotify_Burst.cpp b/Engine/Plugins/Runtime/GameplayAbilities/Source/GameplayAbilities/Private/GameplayCueNotify_Burst.cpp
new file mode 100644
--- /dev/null
+++ b/Engine/Plugins/Runtime/GameplayAbilities/Source/GameplayAbilities/Private/GameplayCueNotify_Burst.cpp
@@ -0,0 +1,40 @@
+// Copyright Epic Games, Inc. All Rights Reserved.
+
+#include "GameplayCueNotify_Burst.h"
+
+
+//////////////////////////////////////////////////////////////////////////
+// UGameplayCueNotify_Burst
+//////////////////////////////////////////////////////////////////////////
+UGameplayCueNotify_Burst::UGameplayCueNotify_Burst()
+{
+}
+
+bool UGameplayCueNotify_Burst::OnExecute_Implementation(AActor* Target, const FGameplayCueParameters& Parameters) const
+{
+	UWorld* World = (Target ? Target->GetWorld() : GetWorld());
+
+	FGameplayCueNotify_SpawnContext SpawnContext(World, Target, Parameters);
+	SpawnContext.SetDefaultSpawnCondition(&DefaultSpawnCondition);
+	SpawnContext.SetDefaultPlacementInfo(&DefaultPlacementInfo);
+
+	if (DefaultSpawnCondition.ShouldSpawn(SpawnContext))
+	{
+		FGameplayCueNotify_SpawnResult SpawnResult;
+
+		BurstEffects.ExecuteEffects(SpawnContext, SpawnResult);
+
+		OnBurst(Target, Parameters, SpawnResult);
+	}
+
+	return false;
+}
+
+#if WITH_EDITOR
+EDataValidationResult UGameplayCueNotify_Burst::IsDataValid(TArray<FText>& ValidationErrors)
+{
+	BurstEffects.ValidateAssociatedAssets(this, TEXT("BurstEffects"), ValidationErrors);
+
+	return ((ValidationErrors.Num() > 0) ? EDataValidationResult::Invalid : EDataValidationResult::Valid);
+}
+#endif // #if WITH_EDITOR
diff --git a/Engine/Plugins/Runtime/GameplayAbilities/Source/GameplayAbilities/Private/GameplayCueNotify_BurstLatent.cpp b/Engine/Plugins/Runtime/GameplayAbilities/Source/GameplayAbilities/Private/GameplayCueNotify_BurstLatent.cpp
new file mode 100644
--- /dev/null
+++ b/Engine/Plugins/Runtime/GameplayAbilities/Source/GameplayAbilities/Private/GameplayCueNotify_BurstLatent.cpp
@@ -0,0 +1,62 @@
+// Copyright Epic Games, Inc. All Rights Reserved.
+
+#include "GameplayCueNotify_BurstLatent.h"
+#include "Engine/World.h"
+#include "TimerManager.h"
+
+
+const float DefaultBurstLatentLifetime = 5.0f;
+
+
+//////////////////////////////////////////////////////////////////////////
+// AGameplayCueNotify_BurstLatent
+//////////////////////////////////////////////////////////////////////////
+AGameplayCueNotify_BurstLatent::AGameplayCueNotify_BurstLatent()
+{
+	Recycle();
+}
+
+bool AGameplayCueNotify_BurstLatent::Recycle()
+{
+	Super::Recycle();
+
+	BurstSpawnResults.Reset();
+
+	return true;
+}
+
+bool AGameplayCueNotify_BurstLatent::OnExecute_Implementation(AActor* Target, const FGameplayCueParameters& Parameters)
+{
+	UWorld* World = GetWorld();
+
+	FGameplayCueNotify_SpawnContext SpawnContext(World, Target, Parameters);
+	SpawnContext.SetDefaultSpawnCondition(&DefaultSpawnCondition);
+	SpawnContext.SetDefaultPlacementInfo(&DefaultPlacementInfo);
+
+	if (DefaultSpawnCondition.ShouldSpawn(SpawnContext))
+	{
+		BurstEffects.ExecuteEffects(SpawnContext, BurstSpawnResults);
+
+		OnBurst(Target, Parameters, BurstSpawnResults);
+	}
+
+	// Handle GC removal by default. This is a simple default to handle all cases we can currently think of.
+	// If we didn't do this, we'd be relying on every BurstLatent GC manually setting up its removal within BP graphs,
+	// or some inference based on parameters.
+	if (World)
+	{
+		const float Lifetime = FMath::Max<float>(AutoDestroyDelay, DefaultBurstLatentLifetime);
+		World->GetTimerManager().SetTimer(FinishTimerHandle, this, &AGameplayCueNotify_Actor::GameplayCueFinishedCallback, Lifetime);
+	}
+
+	return false;
+}
+
+#if WITH_EDITOR
+EDataValidationResult AGameplayCueNotify_BurstLatent::IsDataValid(TArray<FText>& ValidationErrors)
+{
+	BurstEffects.ValidateAssociatedAssets(this, TEXT("BurstEffects"), ValidationErrors);
+
+	return ((ValidationErrors.Num() > 0) ? EDataValidationResult::Invalid : EDataValidationResult::Valid);
+}
+#endif // #if WITH_EDITOR
diff --git a/Engine/Plugins/Runtime/GameplayAbilities/Source/GameplayAbilities/Private/GameplayCueNotify_Instanced.cpp b/Engine/Plugins/Runtime/GameplayAbilities/Source/GameplayAbilities/Private/GameplayCueNotify_Instanced.cpp
new file mode 100644
--- /dev/null
+++ b/Engine/Plugins/Runtime/GameplayAbilities/Source/GameplayAbilities/Private/GameplayCueNotify_Instanced.cpp
@@ -0,0 +1,72 @@
+// Copyright Epic Games, Inc. All Rights Reserved.
+
+#include "GameplayCueNotify_Instanced.h"
+#include "Engine/World.h"
+#include "GameplayCueManager.h"
+#include "GameplayCueNotifyTypes.h"
+
+
+//////////////////////////////////////////////////////////////////////////
+// AGameplayCueNotify_Instanced
+//////////////////////////////////////////////////////////////////////////
+AGameplayCueNotify_Instanced::AGameplayCueNotify_Instanced()
+{
+	RootComponent = CreateDefaultSubobject<USceneComponent>(TEXT("RootComponent"));
+
+	PrimaryActorTick.bStartWithTickEnabled = false;
+	bAutoDestroyOnRemove = true;
+	NumPreallocatedInstances = 3;
+}
+
+void AGameplayCueNotify_Instanced::EndPlay(const EEndPlayReason::Type EndPlayReason)
+{
+	if (UGameplayCueManager::IsGameplayCueRecylingEnabled() && (EndPlayReason == EEndPlayReason::Destroyed))
+	{
+		UWorld* World = GetWorld();
+		if (World && !World->IsPlayingReplay() && (World->WorldType == EWorldType::Game))
+		{
+			UE_LOG(LogGameplayCueNotify, Error, TEXT("GameplayCueNotify [%s] is calling EndPlay(). This should not happen since they are recycled through the GameplayCueManager."), *GetName());
+		}
+	}
+
+	Super::EndPlay(EndPlayReason);
+}
+
+void AGameplayCueNotify_Instanced::K2_DestroyActor()
+{
+	UE_LOG(LogGameplayCueNotify, Warning, TEXT("GameplayCueNotify [%s] is calling DesotryActor(). This is not necessary as GCs will be cleaned up and recycled automatically."), *GetName());
+}
+
+void AGameplayCueNotify_Instanced::Destroyed()
+{
+	if (UGameplayCueManager::IsGameplayCueRecylingEnabled())
+	{
+		UWorld* World = GetWorld();
+		if (World && !World->IsPlayingReplay() && (World->WorldType == EWorldType::Game))
+		{
+			UE_LOG(LogGameplayCueNotify, Error, TEXT("GameplayCueNotify [%s] is calling Destroyed(). This should not happen since they are recycled through the GameplayCueManager."), *GetName());
+		}
+	}
+
+	Super::Destroyed();
+}
+
+bool AGameplayCueNotify_Instanced::WhileActive_Implementation(AActor* Target, const FGameplayCueParameters& Parameters)
+{
+	if (IsHidden())
+	{
+		SetActorHiddenInGame(false);
+	}
+
+	return false;
+}
+
+bool AGameplayCueNotify_Instanced::OnRemove_Implementation(AActor* Target, const FGameplayCueParameters& Parameters)
+{
+	if (!IsHidden())
+	{
+		SetActorHiddenInGame(true);
+	}
+
+	return false;
+}
diff --git a/Engine/Plugins/Runtime/GameplayAbilities/Source/GameplayAbilities/Private/GameplayCueNotify_Looping.cpp b/Engine/Plugins/Runtime/GameplayAbilities/Source/GameplayAbilities/Private/GameplayCueNotify_Looping.cpp
new file mode 100644
--- /dev/null
+++ b/Engine/Plugins/Runtime/GameplayAbilities/Source/GameplayAbilities/Private/GameplayCueNotify_Looping.cpp
@@ -0,0 +1,148 @@
+// Copyright Epic Games, Inc. All Rights Reserved.
+
+#include "GameplayCueNotify_Looping.h"
+
+
+//////////////////////////////////////////////////////////////////////////
+// AGameplayCueNotify_Looping
+//////////////////////////////////////////////////////////////////////////
+AGameplayCueNotify_Looping::AGameplayCueNotify_Looping()
+{
+	bAllowMultipleWhileActiveEvents = false;
+
+	DefaultPlacementInfo.AttachPolicy = EGameplayCueNotify_AttachPolicy::AttachToTarget;
+
+	Recycle();
+}
+
+void AGameplayCueNotify_Looping::EndPlay(const EEndPlayReason::Type EndPlayReason)
+{
+	if (EndPlayReason == EEndPlayReason::Destroyed)
+	{
+		RemoveLoopingEffects();
+	}
+
+	Super::EndPlay(EndPlayReason);
+}
+
+bool AGameplayCueNotify_Looping::Recycle()
+{
+	Super::Recycle();
+
+	// Extra check to make sure looping effects have been removed.  Normally they will have been removed in the OnRemove event.
+	RemoveLoopingEffects();
+
+	ApplicationSpawnResults.Reset();
+	LoopingSpawnResults.Reset();
+	RecurringSpawnResults.Reset();
+	RemovalSpawnResults.Reset();
+
+	bLoopingEffectsRemoved = false;
+
+	return true;
+}
+
+bool AGameplayCueNotify_Looping::OnActive_Implementation(AActor* Target, const FGameplayCueParameters& Parameters)
+{
+	UWorld* World = GetWorld();
+
+	FGameplayCueNotify_SpawnContext SpawnContext(World, Target, Parameters);
+	SpawnContext.SetDefaultSpawnCondition(&DefaultSpawnCondition);
+	SpawnContext.SetDefaultPlacementInfo(&DefaultPlacementInfo);
+
+	if (DefaultSpawnCondition.ShouldSpawn(SpawnContext))
+	{
+		ApplicationEffects.ExecuteEffects(SpawnContext, ApplicationSpawnResults);
+
+		OnApplication(Target, Parameters, ApplicationSpawnResults);
+	}
+
+	return false;
+}
+
+bool AGameplayCueNotify_Looping::WhileActive_Implementation(AActor* Target, const FGameplayCueParameters& Parameters)
+{
+	UWorld* World = GetWorld();
+
+	FGameplayCueNotify_SpawnContext SpawnContext(World, Target, Parameters);
+	SpawnContext.SetDefaultSpawnCondition(&DefaultSpawnCondition);
+	SpawnContext.SetDefaultPlacementInfo(&DefaultPlacementInfo);
+
+	if (DefaultSpawnCondition.ShouldSpawn(SpawnContext))
+	{
+		bLoopingEffectsRemoved = false;
+
+		LoopingEffects.StartEffects(SpawnContext, LoopingSpawnResults);	
+
+		OnLoopingStart(Target, Parameters, LoopingSpawnResults);
+	}
+
+	return false;
+}
+
+bool AGameplayCueNotify_Looping::OnExecute_Implementation(AActor* Target, const FGameplayCueParameters& Parameters)
+{
+	UWorld* World = GetWorld();
+
+	FGameplayCueNotify_SpawnContext SpawnContext(World, Target, Parameters);
+	SpawnContext.SetDefaultSpawnCondition(&DefaultSpawnCondition);
+	SpawnContext.SetDefaultPlacementInfo(&DefaultPlacementInfo);
+
+	if (DefaultSpawnCondition.ShouldSpawn(SpawnContext))
+	{
+		RecurringEffects.ExecuteEffects(SpawnContext, RecurringSpawnResults);
+
+		OnRecurring(Target, Parameters, RecurringSpawnResults);
+	}
+
+	return false;
+}
+
+bool AGameplayCueNotify_Looping::OnRemove_Implementation(AActor* Target, const FGameplayCueParameters& Parameters)
+{
+	RemoveLoopingEffects();
+
+	// Don't spawn removal effects if our target is gone.
+	if (Target)
+	{
+		UWorld* World = GetWorld();
+
+		FGameplayCueNotify_SpawnContext SpawnContext(World, Target, Parameters);
+		SpawnContext.SetDefaultSpawnCondition(&DefaultSpawnCondition);
+		SpawnContext.SetDefaultPlacementInfo(&DefaultPlacementInfo);
+
+		if (DefaultSpawnCondition.ShouldSpawn(SpawnContext))
+		{
+			RemovalEffects.ExecuteEffects(SpawnContext, RemovalSpawnResults);
+		}
+	}
+
+	// Always call OnRemoval(), even if target is bad, so it can clean up BP-spawned things.
+	OnRemoval(Target, Parameters, RemovalSpawnResults);
+
+	return false;
+}
+
+void AGameplayCueNotify_Looping::RemoveLoopingEffects()
+{
+	if (bLoopingEffectsRemoved)
+	{
+		return;
+	}
+
+	bLoopingEffectsRemoved = true;
+
+	LoopingEffects.StopEffects(LoopingSpawnResults);
+}
+
+#if WITH_EDITOR
+EDataValidationResult AGameplayCueNotify_Looping::IsDataValid(TArray<FText>& ValidationErrors)
+{
+	ApplicationEffects.ValidateAssociatedAssets(this, TEXT("ApplicationEffects"), ValidationErrors);
+	LoopingEffects.ValidateAssociatedAssets(this, TEXT("LoopingEffects"), ValidationErrors);
+	RecurringEffects.ValidateAssociatedAssets(this, TEXT("RecurringEffects"), ValidationErrors);
+	RemovalEffects.ValidateAssociatedAssets(this, TEXT("RemovalEffects"), ValidationErrors);
+
+	return ((ValidationErrors.Num() > 0) ? EDataValidationResult::Invalid : EDataValidationResult::Valid);
+}
+#endif // #if WITH_EDITOR
diff --git a/Engine/Plugins/Runtime/GameplayAbilities/Source/GameplayAbilities/Public/AnimNotify_GameplayCue.h b/Engine/Plugins/Runtime/GameplayAbilities/Source/GameplayAbilities/Public/AnimNotify_GameplayCue.h
new file mode 100644
--- /dev/null
+++ b/Engine/Plugins/Runtime/GameplayAbilities/Source/GameplayAbilities/Public/AnimNotify_GameplayCue.h
@@ -0,0 +1,76 @@
+// Copyright Epic Games, Inc. All Rights Reserved.
+
+#pragma once
+
+#include "CoreMinimal.h"
+#include "Animation/AnimNotifies/AnimNotify.h"
+#include "Animation/AnimNotifies/AnimNotifyState.h"
+#include "GameplayCueInterface.h"
+#include "GameplayCueManager.h"
+#include "AnimNotify_GameplayCue.generated.h"
+
+
+/**
+ * UAnimNotify_GameplayCue
+ *
+ *	An animation notify used for instantaneous gameplay cues (Burst/Latent).
+ */
+UCLASS(editinlinenew, Const, hideCategories = Object, collapseCategories, MinimalAPI, Meta = (DisplayName = "GameplayCue (Burst)"))
+class UAnimNotify_GameplayCue : public UAnimNotify
+{
+	GENERATED_BODY()
+
+public:
+
+	UAnimNotify_GameplayCue();
+
+	virtual void Notify(USkeletalMeshComponent* MeshComp, UAnimSequenceBase* Animation) override;
+
+	FString GetNotifyName_Implementation() const override;
+
+#if WITH_EDITOR
+	virtual bool CanBePlaced(UAnimSequenceBase* Animation) const override;
+#endif // #if WITH_EDITOR
+
+protected:
+
+	// GameplayCue tag to invoke.
+	UPROPERTY(EditAnywhere, BlueprintReadOnly, Category = GameplayCue, meta = (Categories = "GameplayCue"))
+	FGameplayCueTag GameplayCue;
+};
+
+
+/**
+ * UAnimNotify_GameplayCueState
+ *
+ *	An animation notify state used for duration based gameplay cues (Looping).
+ */
+UCLASS(editinlinenew, Const, hideCategories = Object, collapseCategories, MinimalAPI, Meta = (DisplayName = "GameplayCue (Looping)"))
+class UAnimNotify_GameplayCueState : public UAnimNotifyState
+{
+	GENERATED_BODY()
+
+public:
+
+	UAnimNotify_GameplayCueState();
+
+	virtual void NotifyBegin(USkeletalMeshComponent* MeshComp, UAnimSequenceBase* Animation, float TotalDuration) override;
+	virtual void NotifyTick(USkeletalMeshComponent* MeshComp, UAnimSequenceBase* Animation, float FrameDeltaTime) override;
+	virtual void NotifyEnd(USkeletalMeshComponent* MeshComp, UAnimSequenceBase* Animation) override;
+
+	FString GetNotifyName_Implementation() const override;
+
+#if WITH_EDITOR
+	virtual bool CanBePlaced(UAnimSequenceBase* Animation) const override;
+#endif // #if WITH_EDITOR
+
+protected:
+
+	// GameplayCue tag to invoke
+	UPROPERTY(EditAnywhere, BlueprintReadOnly, Category = GameplayCue, meta = (Categories = "GameplayCue"))
+	FGameplayCueTag GameplayCue;
+
+#if WITH_EDITORONLY_DATA
+	FGameplayCueProxyTick PreviewProxyTick;
+#endif // #if WITH_EDITORONLY_DATA
+};
diff --git a/Engine/Plugins/Runtime/GameplayAbilities/Source/GameplayAbilities/Public/GameplayCueManager.h b/Engine/Plugins/Runtime/GameplayAbilities/Source/GameplayAbilities/Public/GameplayCueManager.h
--- a/Engine/Plugins/Runtime/GameplayAbilities/Source/GameplayAbilities/Public/GameplayCueManager.h
+++ b/Engine/Plugins/Runtime/GameplayAbilities/Source/GameplayAbilities/Public/GameplayCueManager.h
@@ -158,6 +158,22 @@ class GAMEPLAYABILITIES_API UGameplayCueManager : public UDataAsset
 	/** 3. Actually routes the gameplaycue event to the right place.  */
 	virtual void RouteGameplayCue(AActor* TargetActor, FGameplayTag GameplayCueTag, EGameplayCueEvent::Type EventType, const FGameplayCueParameters& Parameters, EGameplayCueExecutionOptions Options = EGameplayCueExecutionOptions::Default);
 
+	/** 
+	 *  Convenience methods for invoking non-replicated gameplay cue events. 
+	 * 
+	 *	We want to avoid exposing designers the choice of "is this gameplay cue replicated or non-replicated?".
+	 *	We want to make the decision for them in most cases:
+	 *	- Abilities will always use replicated GameplayCue events because they are not executed on simulated proxies.
+	 *	- Animations always use non-replicated GameplayCue events because they are always executed on simulated proxies.
+	 *	
+	 *  Sometimes it will be useful to give designers both options: in actor classes where there are many possible use cases.
+	 *  Still, we should keep the choice confined to the actor class, and not globally.  E.g., Don't add both choices to the function library
+	 *  since they would appear everywhere. Add the choices to the actor class so they only appear there.
+	 */
+	static void AddGameplayCue_NonReplicated(AActor* Target, const FGameplayTag GameplayCueTag, const FGameplayCueParameters& Parameters);
+	static void RemoveGameplayCue_NonReplicated(AActor* Target, const FGameplayTag GameplayCueTag, const FGameplayCueParameters& Parameters);
+	static void ExecuteGameplayCue_NonReplicated(AActor* Target, const FGameplayTag GameplayCueTag, const FGameplayCueParameters& Parameters);
+
 	// -------------------------------------------------------------
 
 	/** Force any instanced GameplayCueNotifies to stop */
diff --git a/Engine/Plugins/Runtime/GameplayAbilities/Source/GameplayAbilities/Public/GameplayCueNotifyTypes.h b/Engine/Plugins/Runtime/GameplayAbilities/Source/GameplayAbilities/Public/GameplayCueNotifyTypes.h
new file mode 100644
--- /dev/null
+++ b/Engine/Plugins/Runtime/GameplayAbilities/Source/GameplayAbilities/Public/GameplayCueNotifyTypes.h
@@ -0,0 +1,798 @@
+// Copyright Epic Games, Inc. All Rights Reserved.
+
+#pragma once
+
+#include "CoreMinimal.h"
+#include "Logging/LogMacros.h"
+#include "Camera/CameraTypes.h"
+#include "GameplayEffectTypes.h"
+#include "PhysicalMaterials/PhysicalMaterial.h"
+#include "GameplayCueNotifyTypes.generated.h"
+
+
+class UParticleSystem;
+class UParticleSystemComponent;
+class UAudioComponent;
+class USoundBase;
+class UCameraShake;
+class UCameraAnim;
+class UCameraAnimInst;
+class AEmitterCameraLensEffectBase;
+class UForceFeedbackEffect;
+class UForceFeedbackAttenuation;
+class UForceFeedbackComponent;
+class ADecalActor;
+class APlayerController;
+struct FGameplayCueParameters;
+struct FGameplayCueNotify_SpawnContext;
+
+
+typedef uint64 FGameplayCueNotify_SurfaceMask;
+static_assert((EPhysicalSurface::SurfaceType_Max <= (sizeof(FGameplayCueNotify_SurfaceMask) * 8)), "FGameplayCueNotify_SurfaceMask does not have enough bits for EPhysicalSurface");
+
+
+DECLARE_LOG_CATEGORY_EXTERN(LogGameplayCueNotify, Log, All);
+
+
+/**
+ * EGameplayCueNotify_EffectPlaySpace
+ *
+ *	Used by some effects (like camera shakes) to specify what coordinate space the they should be applied in.
+ */
+UENUM(BlueprintType)
+enum class EGameplayCueNotify_EffectPlaySpace : uint8
+{
+	// Play the effect in world space.
+	WorldSpace,
+
+	// Play the effect in camera space.
+	CameraSpace,
+};
+
+
+/**
+ * EGameplayCueNotify_LocallyControlledSource
+ *
+ *	Specifies what actor to use when determining if the gameplay cue notify is locally controlled.
+ */
+UENUM(BlueprintType)
+enum class EGameplayCueNotify_LocallyControlledSource : uint8
+{
+	// Use the instigator actor as the source when deciding the locally controlled policy.
+	InstigatorActor,
+
+	// Use the target actor as the source when deciding the locally controlled policy.
+	TargetActor,
+};
+
+
+/**
+ * EGameplayCueNotify_LocallyControlledPolicy
+ *
+ *	Specifies if the gameplay cue notify should spawn based on it being locally controlled.
+ */
+UENUM(BlueprintType)
+enum class EGameplayCueNotify_LocallyControlledPolicy : uint8
+{
+	// Always spawns regardless of locally controlled.
+	Always,
+
+	// Only spawn if the source actor is locally controlled.
+	LocalOnly,
+
+	// Only spawn if the source actor is NOT locally controlled.
+	NotLocal,
+};
+
+
+/**
+ * FGameplayCueNotify_SpawnCondition
+ *
+ *	Conditions used to determine if the gameplay cue notify should spawn.
+ */
+USTRUCT(BlueprintType)
+struct FGameplayCueNotify_SpawnCondition
+{
+	GENERATED_BODY()
+
+public:
+
+	FGameplayCueNotify_SpawnCondition();
+
+	bool ShouldSpawn(const FGameplayCueNotify_SpawnContext& SpawnContext) const;
+
+public:
+
+	// Source actor to use when determining if it is locally controlled.
+	UPROPERTY(EditAnywhere, BlueprintReadOnly, Category = GameplayCueNotify)
+	EGameplayCueNotify_LocallyControlledSource LocallyControlledSource;
+
+	// Locally controlled policy used to determine if the gameplay cue effects should spawn.
+	UPROPERTY(EditAnywhere, BlueprintReadOnly, Category = GameplayCueNotify)
+	EGameplayCueNotify_LocallyControlledPolicy LocallyControlledPolicy;
+
+	// Random chance to play the effects.  (1.0 = always play,  0.0 = never play)
+	UPROPERTY(EditAnywhere, BlueprintReadOnly, Category = GameplayCueNotify, Meta = (UIMin = "0.0", ClampMin = "0.0", UIMax = "1.0", ClampMax = "1.0"))
+	float ChanceToPlay;
+
+	// The gameplay cue effects will only spawn if the surface type is in this list.
+	// An empty list means everything is allowed.
+	UPROPERTY(EditAnywhere, BlueprintReadOnly, Category = GameplayCueNotify)
+	TArray<TEnumAsByte<EPhysicalSurface>> AllowedSurfaceTypes;
+	mutable FGameplayCueNotify_SurfaceMask AllowedSurfaceMask;
+
+	// The gameplay cue effects will only spawn if the surface type is NOT in this list.
+	// An empty list means nothing will be rejected.
+	UPROPERTY(EditAnywhere, BlueprintReadOnly, Category = GameplayCueNotify)
+	TArray<TEnumAsByte<EPhysicalSurface>> RejectedSurfaceTypes;
+	mutable FGameplayCueNotify_SurfaceMask RejectedSurfaceMask;
+};
+
+
+/**
+ * EGameplayCueNotify_AttachPolicy
+ *
+ *	Specifies if the gameplay cue notify should attach to the target actor.
+ */
+UENUM(BlueprintType)
+enum class EGameplayCueNotify_AttachPolicy : uint8
+{
+	// Do not attach to the target actor.  The target may still be used to get location and other information.
+	DoNotAttach,
+
+	// Attach to the target actor if possible.
+	AttachToTarget,
+};
+
+
+/**
+ * FGameplayCueNotify_PlacementInfo
+ *
+ *	Specifies how the gameplay cue notify will be positioned in the world.
+ */
+USTRUCT(BlueprintType)
+struct FGameplayCueNotify_PlacementInfo
+{
+	GENERATED_BODY()
+
+public:
+
+	FGameplayCueNotify_PlacementInfo();
+
+	bool FindSpawnTransform(const FGameplayCueNotify_SpawnContext& SpawnContext, FTransform& OutTransform) const;
+
+	void SetComponentTransform(USceneComponent* Component, const FTransform& Transform) const;
+	void TryComponentAttachment(USceneComponent* Component, USceneComponent* AttachComponent) const;
+
+public:
+
+	// Target's socket (or bone) used for location and rotation.  If "None", it uses the target's root.
+	UPROPERTY(EditAnywhere, BlueprintReadOnly, Category = GameplayCueNotify)
+	FName SocketName;
+
+	// Whether to attach to the target actor or not attach.
+	UPROPERTY(EditAnywhere, BlueprintReadOnly, Category = GameplayCueNotify)
+	EGameplayCueNotify_AttachPolicy AttachPolicy;
+
+	// How the transform is handled when attached.
+	UPROPERTY(EditAnywhere, BlueprintReadOnly, Category = GameplayCueNotify)
+	EAttachmentRule AttachmentRule;
+
+	// If enabled, will always spawn using rotation override.
+	UPROPERTY(EditAnywhere, BlueprintReadOnly, Category = GameplayCueNotify, Meta = (InlineEditConditionToggle))
+	uint32 bOverrideRotation : 1;
+
+	// If enabled, will always spawn using the scale override.
+	UPROPERTY(EditAnywhere, BlueprintReadOnly, Category = GameplayCueNotify, Meta = (InlineEditConditionToggle))
+	uint32 bOverrideScale : 1;
+
+	// If enabled, will always spawn using rotation override.
+	// This will also set the rotation to be absolute, so any changes to the parent's rotation will be ignored after attachment.
+	UPROPERTY(EditAnywhere, BlueprintReadOnly, Category = GameplayCueNotify, Meta = (EditCondition = "bOverrideRotation"))
+	FRotator RotationOverride;
+
+	// If enabled, will always spawn using scale override.
+	// This will also set the scale to be absolute, so any changes to the parent's scale will be ignored after attachment.
+	UPROPERTY(EditAnywhere, BlueprintReadOnly, Category = GameplayCueNotify, Meta = (EditCondition = "bOverrideScale"))
+	FVector ScaleOverride;
+};
+
+
+/**
+ * FGameplayCueNotify_SpawnContext
+ *
+ *	Temporary spawn information collected from input parameters.
+ */
+struct FGameplayCueNotify_SpawnContext
+{
+public:
+
+	FGameplayCueNotify_SpawnContext(UWorld* InWorld, AActor* InTargetActor, const FGameplayCueParameters& InCueParameters);
+
+	void SetDefaultSpawnCondition(const FGameplayCueNotify_SpawnCondition* InSpawnCondition) {DefaultSpawnCondition = InSpawnCondition;}
+	void SetDefaultPlacementInfo(const FGameplayCueNotify_PlacementInfo* InPlacementInfo) {DefaultPlacementInfo = InPlacementInfo;}
+
+	const FGameplayCueNotify_PlacementInfo& GetPlacementInfo(bool bUseOverride, const FGameplayCueNotify_PlacementInfo& PlacementInfoOverride) const
+	{
+		return (!bUseOverride && DefaultPlacementInfo) ? *DefaultPlacementInfo : PlacementInfoOverride;
+	}
+
+	const FGameplayCueNotify_SpawnCondition& GetSpawnCondition(bool bUseOverride, const FGameplayCueNotify_SpawnCondition& SpawnConditionOverride) const
+	{
+		return (!bUseOverride && DefaultSpawnCondition) ? *DefaultSpawnCondition : SpawnConditionOverride;
+	}
+
+	APlayerController* FindLocalPlayerController(EGameplayCueNotify_LocallyControlledSource Source) const;
+
+protected:
+
+	void InitializeContext();
+
+public:
+
+	UWorld* World;
+	AActor* TargetActor;
+	const FGameplayCueParameters& CueParameters;
+	const FHitResult* HitResult;
+	USceneComponent* TargetComponent;
+	EPhysicalSurface SurfaceType;
+
+private:
+
+	const FGameplayCueNotify_SpawnCondition* DefaultSpawnCondition;
+	const FGameplayCueNotify_PlacementInfo* DefaultPlacementInfo;
+};
+
+
+/**
+ * FGameplayCueNotify_SpawnResult
+ *
+ *	Temporary structure used to track results of spawning components.
+ */
+USTRUCT(BlueprintType)
+struct FGameplayCueNotify_SpawnResult
+{
+	GENERATED_BODY()
+
+public:
+
+	FGameplayCueNotify_SpawnResult() { Reset(); }
+
+	void Reset()
+	{
+		ParticleComponents.Reset();
+		AudioComponents.Reset();
+		CameraShakes.Reset();
+		CameraAnim = nullptr;
+		CameraLensEffects.Reset();
+		ForceFeedbackComponent = nullptr;
+		ForceFeedbackTargetPC = nullptr;
+		DecalActor = nullptr;
+	}
+
+public:
+
+	// List of particle components spawned.  There may be null pointers here as it matches the defined order.
+	UPROPERTY(BlueprintReadOnly, Transient, Category = GameplayCueNotify)
+	TArray<UParticleSystemComponent*> ParticleComponents;
+
+	// List of audio components spawned.  There may be null pointers here as it matches the defined order.
+	UPROPERTY(BlueprintReadOnly, Transient, Category = GameplayCueNotify)
+	TArray<UAudioComponent*> AudioComponents;
+
+	// List of camera shakes played.  There will be one camera shake per local player controller if shake is played in world.
+	UPROPERTY(BlueprintReadOnly, Transient, Category = GameplayCueNotify)
+	TArray<UCameraShake*> CameraShakes;
+
+	// Camera animation played.
+	UPROPERTY(BlueprintReadOnly, Transient, Category = GameplayCueNotify)
+	UCameraAnimInst* CameraAnim;
+
+	// List of camera len effects spawned.  There will be one camera lens effect per local player controller if the effect is played in world.
+	UPROPERTY(BlueprintReadOnly, Transient, Category = GameplayCueNotify)
+	TArray<AEmitterCameraLensEffectBase*> CameraLensEffects;
+
+	// Force feedback component that was spawned.  This is only valid when force feedback is set to play in world.
+	UPROPERTY(BlueprintReadOnly, Transient, Category = GameplayCueNotify)
+	UForceFeedbackComponent* ForceFeedbackComponent;
+
+	// Player controller used to play the force feedback effect.  Used to stop the effect later.
+	UPROPERTY(Transient)
+	APlayerController* ForceFeedbackTargetPC;
+
+	// Spawned decal actor.  This may be null.
+	UPROPERTY(BlueprintReadOnly, Transient, Category = GameplayCueNotify)
+	ADecalActor* DecalActor;
+};
+
+
+/**
+ * FGameplayCueNotify_ParticleInfo
+ *
+ *	Properties that specify how to play a particle effect.
+ */
+USTRUCT(BlueprintType)
+struct FGameplayCueNotify_ParticleInfo
+{
+	GENERATED_BODY()
+
+public:
+
+	FGameplayCueNotify_ParticleInfo();
+
+	bool PlayParticleEffect(const FGameplayCueNotify_SpawnContext& SpawnContext, FGameplayCueNotify_SpawnResult& OutSpawnResult) const;
+
+	void ValidateBurstAssets(UObject* ContainingAsset, const FString& Context, TArray<FText>& ValidationErrors) const;
+
+public:
+
+	// Condition to check before spawning the particle system.
+	UPROPERTY(EditAnywhere, BlueprintReadOnly, Category = GameplayCueNotify, Meta = (EditCondition = "bOverrideSpawnCondition"))
+	FGameplayCueNotify_SpawnCondition SpawnConditionOverride;
+
+	// Defines how the particle system will be placed.
+	UPROPERTY(EditAnywhere, BlueprintReadOnly, Category = GameplayCueNotify, Meta = (EditCondition = "bOverridePlacementInfo"))
+	FGameplayCueNotify_PlacementInfo PlacementInfoOverride;
+
+	// Particle system to spawn.
+	UPROPERTY(EditAnywhere, BlueprintReadOnly, Category = GameplayCueNotify)
+	UParticleSystem* ParticleSystem;
+
+	// If enabled, use the spawn condition override and not the default one.
+	UPROPERTY(EditAnywhere, BlueprintReadOnly, Category = GameplayCueNotify, Meta = (InlineEditConditionToggle))
+	uint32 bOverrideSpawnCondition : 1;
+
+	// If enabled, use the placement info override and not the default one.
+	UPROPERTY(EditAnywhere, BlueprintReadOnly, Category = GameplayCueNotify, Meta = (InlineEditConditionToggle))
+	uint32 bOverridePlacementInfo : 1;
+
+	// If enabled, this particle system will cast a shadow.
+	UPROPERTY(EditAnywhere, BlueprintReadOnly, Category = GameplayCueNotify)
+	uint32 bCastShadow : 1;
+};
+
+
+/**
+ * FGameplayCueNotify_SoundInfo
+ *
+ *	Properties that specify how to play a sound effect.
+ */
+USTRUCT(BlueprintType)
+struct FGameplayCueNotify_SoundInfo
+{
+	GENERATED_BODY()
+
+public:
+
+	FGameplayCueNotify_SoundInfo();
+
+	bool PlaySound(const FGameplayCueNotify_SpawnContext& SpawnContext, FGameplayCueNotify_SpawnResult& OutSpawnResult) const;
+
+	void ValidateBurstAssets(UObject* ContainingAsset, const FString& Context, TArray<FText>& ValidationErrors) const;
+
+public:
+
+	// Condition to check before playing the sound.
+	UPROPERTY(EditAnywhere, BlueprintReadOnly, Category = GameplayCueNotify, Meta = (EditCondition = "bOverrideSpawnCondition"))
+	FGameplayCueNotify_SpawnCondition SpawnConditionOverride;
+
+	// Defines how the sound will be placed.
+	UPROPERTY(EditAnywhere, BlueprintReadOnly, Category = GameplayCueNotify, Meta = (EditCondition = "bOverridePlacementInfo"))
+	FGameplayCueNotify_PlacementInfo PlacementInfoOverride;
+
+	// Sound cue to play.
+	UPROPERTY(EditAnywhere, BlueprintReadOnly, Category = GameplayCueNotify)
+	USoundBase* SoundCue;
+
+	// How long it should take to fade out.  Only used on looping gameplay cues.
+	UPROPERTY(EditAnywhere, BlueprintReadOnly, Category = GameplayCueNotify)
+	float LoopingFadeOutDuration;
+
+	// The volume level you want the sound to fade out to over the 'Looping Fade Out Duration' before stopping.
+	// This value is irrelevant if the 'Looping Fade Out Duration' is zero.
+	// NOTE: If the fade out duration is positive and this value is not lower than the volume the sound is playing at, the sound will never stop!
+	UPROPERTY(EditAnywhere, BlueprintReadOnly, Category = GameplayCueNotify)
+	float LoopingFadeVolumeLevel;
+
+	// If enabled, use the spawn condition override and not the default one.
+	UPROPERTY(EditAnywhere, BlueprintReadOnly, Category = GameplayCueNotify, Meta = (InlineEditConditionToggle))
+	uint32 bOverrideSpawnCondition : 1;
+
+	// If enabled, use the placement info override and not the default one.
+	UPROPERTY(EditAnywhere, BlueprintReadOnly, Category = GameplayCueNotify, Meta = (InlineEditConditionToggle))
+	uint32 bOverridePlacementInfo : 1;
+};
+
+
+/**
+ * FGameplayCueNotify_CameraShakeInfo
+ *
+ *	Properties that specify how to play a camera shake.
+ */
+USTRUCT(BlueprintType)
+struct FGameplayCueNotify_CameraShakeInfo
+{
+	GENERATED_BODY()
+
+public:
+
+	FGameplayCueNotify_CameraShakeInfo();
+
+	bool PlayCameraShake(const FGameplayCueNotify_SpawnContext& SpawnContext, FGameplayCueNotify_SpawnResult& OutSpawnResult) const;
+
+	void ValidateBurstAssets(UObject* ContainingAsset, const FString& Context, TArray<FText>& ValidationErrors) const;
+
+public:
+
+	// Condition to check before playing the camera shake.
+	UPROPERTY(EditAnywhere, BlueprintReadOnly, Category = GameplayCueNotify, Meta = (EditCondition = "bOverrideSpawnCondition"))
+	FGameplayCueNotify_SpawnCondition SpawnConditionOverride;
+
+	// Defines how the camera shake will be placed.
+	UPROPERTY(EditAnywhere, BlueprintReadOnly, Category = GameplayCueNotify, Meta = (EditCondition = "bOverridePlacementInfo"))
+	FGameplayCueNotify_PlacementInfo PlacementInfoOverride;
+
+	// Camera shake to play.
+	UPROPERTY(EditAnywhere, BlueprintReadOnly, Category = GameplayCueNotify)
+	TSubclassOf<UCameraShake> CameraShake;
+
+	// Scale applied to the camera shake.
+	UPROPERTY(EditAnywhere, BlueprintReadOnly, Category = GameplayCueNotify)
+	float ShakeScale;
+
+	// What coordinate space to play the camera shake relative to.
+	UPROPERTY(EditAnywhere, BlueprintReadOnly, Category = GameplayCueNotify)
+	EGameplayCueNotify_EffectPlaySpace PlaySpace;
+
+	// If enabled, use the spawn condition override and not the default one.
+	UPROPERTY(EditAnywhere, BlueprintReadOnly, Category = GameplayCueNotify, Meta = (InlineEditConditionToggle))
+	uint32 bOverrideSpawnCondition : 1;
+
+	// If enabled, use the placement info override and not the default one.
+	UPROPERTY(EditAnywhere, BlueprintReadOnly, Category = GameplayCueNotify, Meta = (InlineEditConditionToggle))
+	uint32 bOverridePlacementInfo : 1;
+
+	// If enabled, the camera shake will be played in the world and affect all players.
+	UPROPERTY(EditAnywhere, BlueprintReadOnly, Category = GameplayCueNotify)
+	uint32 bPlayInWorld : 1;
+
+	// Players inside this radius get the full intensity camera shake.
+	UPROPERTY(EditAnywhere, BlueprintReadOnly, Category = GameplayCueNotify, Meta = (EditCondition = "bPlayInWorld", UIMin = "0.0", ClampMin = "0.0"))
+	float WorldInnerRadius;
+
+	// Players outside this radius do not get the camera shake applied.
+	UPROPERTY(EditAnywhere, BlueprintReadOnly, Category = GameplayCueNotify, Meta = (EditCondition = "bPlayInWorld", UIMin = "0.0", ClampMin = "0.0"))
+	float WorldOuterRadius;
+
+	// Exponent that describes the shake intensity falloff curve between the inner and outer radii.  (1.0 is linear)
+	UPROPERTY(EditAnywhere, BlueprintReadOnly, Category = GameplayCueNotify, Meta = (EditCondition = "bPlayInWorld", UIMin = "0.0", ClampMin = "0.0"))
+	float WorldFalloffExponent;
+};
+
+
+/**
+ * FGameplayCueNotify_CameraAnimInfo
+ *
+ *	Properties that specify how to play a camera animation.
+ */
+USTRUCT(BlueprintType)
+struct FGameplayCueNotify_CameraAnimInfo
+{
+	GENERATED_BODY()
+
+public:
+
+	FGameplayCueNotify_CameraAnimInfo();
+
+	bool PlayCameraAnim(const FGameplayCueNotify_SpawnContext& SpawnContext, FGameplayCueNotify_SpawnResult& OutSpawnResult) const;
+
+	void ValidateBurstAssets(UObject* ContainingAsset, const FString& Context, TArray<FText>& ValidationErrors) const;
+
+public:
+
+	// Condition to check before playing the camera animation.
+	UPROPERTY(EditAnywhere, BlueprintReadOnly, Category = GameplayCueNotify, Meta = (EditCondition = "bOverrideSpawnCondition"))
+	FGameplayCueNotify_SpawnCondition SpawnConditionOverride;
+
+	// Camera animation to play.
+	UPROPERTY(EditAnywhere, BlueprintReadOnly, Category = GameplayCueNotify)
+	UCameraAnim* CameraAnim;
+
+	// Scale applied to the camera animation.
+	UPROPERTY(EditAnywhere, BlueprintReadOnly, Category = GameplayCueNotify)
+	float AnimScale;
+
+	// Rate to play the camera animation at.
+	UPROPERTY(EditAnywhere, BlueprintReadOnly, Category = GameplayCueNotify)
+	float PlayRate;
+
+	// Time (in seconds) used to blend in the camera animation.
+	UPROPERTY(EditAnywhere, BlueprintReadOnly, Category = GameplayCueNotify)
+	float BlendInTime;
+
+	// Time (in seconds) used to blend out the camera animation.
+	UPROPERTY(EditAnywhere, BlueprintReadOnly, Category = GameplayCueNotify)
+	float BlendOutTime;
+
+	// What coordinate space to play the camera animation relative to.
+	UPROPERTY(EditAnywhere, BlueprintReadOnly, Category = GameplayCueNotify)
+	EGameplayCueNotify_EffectPlaySpace PlaySpace;
+
+	// If enabled, the camera animation will be set to loop.
+	UPROPERTY(EditAnywhere, BlueprintReadOnly, Category = GameplayCueNotify)
+	uint32 bIsLooping : 1;
+
+	// If enabled, the camera animation will start at a random time.  Only used when looping.
+	UPROPERTY(EditAnywhere, BlueprintReadOnly, Category = GameplayCueNotify, Meta = (EditCondition = "bIsLooping"))
+	uint32 bRandomStartTime : 1;
+
+	// If enabled, use the spawn condition override and not the default one.
+	UPROPERTY(EditAnywhere, BlueprintReadOnly, Category = GameplayCueNotify, Meta = (InlineEditConditionToggle))
+	uint32 bOverrideSpawnCondition : 1;
+};
+
+
+/**
+ * FGameplayCueNotify_CameraLensEffectInfo
+ *
+ *	Properties that specify how to play a camera lens effect.
+ */
+USTRUCT(BlueprintType)
+struct FGameplayCueNotify_CameraLensEffectInfo
+{
+	GENERATED_BODY()
+
+public:
+
+	FGameplayCueNotify_CameraLensEffectInfo();
+
+	bool PlayCameraLensEffect(const FGameplayCueNotify_SpawnContext& SpawnContext, FGameplayCueNotify_SpawnResult& OutSpawnResult) const;
+
+	void ValidateBurstAssets(UObject* ContainingAsset, const FString& Context, TArray<FText>& ValidationErrors) const;
+
+public:
+
+	// Condition to check before playing the camera lens effect.
+	UPROPERTY(EditAnywhere, BlueprintReadOnly, Category = GameplayCueNotify, Meta = (EditCondition = "bOverrideSpawnCondition"))
+	FGameplayCueNotify_SpawnCondition SpawnConditionOverride;
+
+	// Defines how the camera lens effect will be placed.
+	UPROPERTY(EditAnywhere, BlueprintReadOnly, Category = GameplayCueNotify, Meta = (EditCondition = "bOverridePlacementInfo"))
+	FGameplayCueNotify_PlacementInfo PlacementInfoOverride;
+
+	// Camera lens effect to play.
+	UPROPERTY(EditAnywhere, BlueprintReadOnly, Category = GameplayCueNotify)
+	TSubclassOf<AEmitterCameraLensEffectBase> CameraLensEffect;
+
+	// If enabled, use the spawn condition override and not the default one.
+	UPROPERTY(EditAnywhere, BlueprintReadOnly, Category = GameplayCueNotify, Meta = (InlineEditConditionToggle))
+	uint32 bOverrideSpawnCondition : 1;
+
+	// If enabled, use the placement info override and not the default one.
+	UPROPERTY(EditAnywhere, BlueprintReadOnly, Category = GameplayCueNotify, Meta = (InlineEditConditionToggle))
+	uint32 bOverridePlacementInfo : 1;
+
+	// If enabled, the camera lens effect will be played in the world and affect all players.
+	UPROPERTY(EditAnywhere, BlueprintReadOnly, Category = GameplayCueNotify)
+	uint32 bPlayInWorld : 1;
+
+	// Players inside this radius get the full intensity camera lens effect.
+	UPROPERTY(EditAnywhere, BlueprintReadOnly, Category = GameplayCueNotify, Meta = (EditCondition = "bPlayInWorld", UIMin = "0.0", ClampMin = "0.0"))
+	float WorldInnerRadius;
+
+	// Players outside this radius do not get the camera lens effect applied.
+	UPROPERTY(EditAnywhere, BlueprintReadOnly, Category = GameplayCueNotify, Meta = (EditCondition = "bPlayInWorld", UIMin = "0.0", ClampMin = "0.0"))
+	float WorldOuterRadius;
+};
+
+
+/**
+ * FGameplayCueNotify_ForceFeedbackInfo
+ *
+ *	Properties that specify how to play a force feedback effect.
+ */
+USTRUCT(BlueprintType)
+struct FGameplayCueNotify_ForceFeedbackInfo
+{
+	GENERATED_BODY()
+
+public:
+
+	FGameplayCueNotify_ForceFeedbackInfo();
+
+	bool PlayForceFeedback(const FGameplayCueNotify_SpawnContext& SpawnContext, FGameplayCueNotify_SpawnResult& OutSpawnResult) const;
+	
+	void ValidateBurstAssets(UObject* ContainingAsset, const FString& Context, TArray<FText>& ValidationErrors) const;
+
+public:
+
+	// Condition to check before playing the force feedback.
+	UPROPERTY(EditAnywhere, BlueprintReadOnly, Category = GameplayCueNotify, Meta = (EditCondition = "bOverrideSpawnCondition"))
+	FGameplayCueNotify_SpawnCondition SpawnConditionOverride;
+
+	// Defines how the force feedback will be placed.
+	UPROPERTY(EditAnywhere, BlueprintReadOnly, Category = GameplayCueNotify, Meta = (EditCondition = "bOverridePlacementInfo"))
+	FGameplayCueNotify_PlacementInfo PlacementInfoOverride;
+
+	// Force feedback effect to play.
+	UPROPERTY(EditAnywhere, BlueprintReadWrite, Category = GameplayCueNotify)
+	UForceFeedbackEffect* ForceFeedbackEffect;
+
+	// Tag used to identify the force feedback effect so it can later be canceled.
+	// Warning: If this is "None" it will stop ALL force feedback effects if it is canceled.
+	UPROPERTY(EditAnywhere, BlueprintReadOnly, Category = GameplayCueNotify)
+	FName ForceFeedbackTag;
+
+	// If enabled, the force feedback will be set to loop.
+	UPROPERTY(EditAnywhere, BlueprintReadOnly, Category = GameplayCueNotify)
+	uint32 bIsLooping : 1;
+
+	// If enabled, use the spawn condition override and not the default one.
+	UPROPERTY(EditAnywhere, BlueprintReadOnly, Category = GameplayCueNotify, Meta = (InlineEditConditionToggle))
+	uint32 bOverrideSpawnCondition : 1;
+
+	// If enabled, use the placement info override and not the default one.
+	UPROPERTY(EditAnywhere, BlueprintReadOnly, Category = GameplayCueNotify, Meta = (InlineEditConditionToggle))
+	uint32 bOverridePlacementInfo : 1;
+
+	// If enabled, the force feedback will be played in the world and affect all players.
+	UPROPERTY(EditAnywhere, BlueprintReadOnly, Category = GameplayCueNotify)
+	uint32 bPlayInWorld : 1;
+
+	// Multiplier applied to the force feedback when played in world.
+	UPROPERTY(EditAnywhere, BlueprintReadWrite, Category = GameplayCueNotify, Meta = (EditCondition = "bPlayInWorld"))
+	float WorldIntensity;
+
+	// How the force feedback is attenuated over distance when played in world.
+	UPROPERTY(EditAnywhere, BlueprintReadWrite, Category = GameplayCueNotify, Meta = (EditCondition = "bPlayInWorld"))
+	UForceFeedbackAttenuation* WorldAttenuation;
+};
+
+
+/**
+ * FGameplayCueNotify_DecalInfo
+ *
+ *	Properties that specify how to spawn a decal.
+ */
+USTRUCT(BlueprintType)
+struct FGameplayCueNotify_DecalInfo
+{
+	GENERATED_BODY()
+
+public:
+	
+	FGameplayCueNotify_DecalInfo();
+
+	bool SpawnDecal(const FGameplayCueNotify_SpawnContext& SpawnContext, FGameplayCueNotify_SpawnResult& OutSpawnResult) const;
+
+public:
+
+	// Condition to check before spawning the decal.
+	UPROPERTY(EditAnywhere, BlueprintReadOnly, Category = GameplayCueNotify, Meta = (EditCondition = "bOverrideSpawnCondition"))
+	FGameplayCueNotify_SpawnCondition SpawnConditionOverride;
+
+	// Defines how the decal will be placed.
+	UPROPERTY(EditAnywhere, BlueprintReadOnly, Category = GameplayCueNotify, Meta = (EditCondition = "bOverridePlacementInfo"))
+	FGameplayCueNotify_PlacementInfo PlacementInfoOverride;
+
+	// Decal actor class to spawn.
+	UPROPERTY(EditAnywhere, BlueprintReadOnly, Category = GameplayCueNotify)
+	TSubclassOf<ADecalActor> Decal;
+
+	// If enabled, use the spawn condition override and not the default one.
+	UPROPERTY(EditAnywhere, BlueprintReadOnly, Category = GameplayCueNotify, Meta = (InlineEditConditionToggle))
+	uint32 bOverrideSpawnCondition : 1;
+
+	// If enabled, use the placement info override and not the default one.
+	UPROPERTY(EditAnywhere, BlueprintReadOnly, Category = GameplayCueNotify, Meta = (InlineEditConditionToggle))
+	uint32 bOverridePlacementInfo : 1;
+
+	// If enabled, override default decal actor fade out values.
+	UPROPERTY(EditAnywhere, BlueprintReadOnly, Category = GameplayCueNotify)
+	uint32 bOverrideFadeOut : 1;
+
+	// Sets when the decal actor will start fading out.  Will override setting in class.
+	UPROPERTY(EditAnywhere, BlueprintReadOnly, Category = GameplayCueNotify, meta = (EditCondition = "bOverrideFadeOut"))
+	float FadeOutStartDelay;
+
+	// Sets how long it takes for decal actor to fade out.  Will override setting in class.
+	UPROPERTY(EditAnywhere, BlueprintReadOnly, Category = GameplayCueNotify, meta = (EditCondition = "bOverrideFadeOut"))
+	float FadeOutDuration;
+};
+
+
+/**
+ * FGameplayCueNotify_BurstEffects
+ *
+ *	Set of effects to spawn for a single event, used by all gameplay cue notify types.
+ */
+USTRUCT(BlueprintType)
+struct FGameplayCueNotify_BurstEffects
+{
+	GENERATED_BODY()
+
+public:
+
+	FGameplayCueNotify_BurstEffects();
+
+	void ExecuteEffects(const FGameplayCueNotify_SpawnContext& SpawnContext, FGameplayCueNotify_SpawnResult& OutSpawnResult) const;
+
+	void ValidateAssociatedAssets(UObject* ContainingAsset, const FString& Context, TArray<FText>& ValidationErrors) const;
+
+protected:
+
+	// Particle systems to be spawned on gameplay cue execution.  These should never use looping effects!
+	UPROPERTY(EditAnywhere, BlueprintReadWrite, Category = GameplayCueNotify)
+	TArray<FGameplayCueNotify_ParticleInfo> BurstParticles;
+
+	// Sound to be played on gameplay cue execution.  These should never use looping effects!
+	UPROPERTY(EditAnywhere, BlueprintReadWrite, Category = GameplayCueNotify)
+	TArray<FGameplayCueNotify_SoundInfo> BurstSounds;
+
+	// Camera shake to be played on gameplay cue execution.  This should never use a looping effect!
+	UPROPERTY(EditAnywhere, BlueprintReadWrite, Category = GameplayCueNotify)
+	FGameplayCueNotify_CameraShakeInfo BurstCameraShake;
+
+	// Camera animation to be played on gameplay cue execution.  This should never use a looping effect!
+	UPROPERTY(EditAnywhere, BlueprintReadWrite, Category = GameplayCueNotify)
+	FGameplayCueNotify_CameraAnimInfo BurstCameraAnim;
+
+	// Camera lens effect to be played on gameplay cue execution.  This should never use a looping effect!
+	UPROPERTY(EditAnywhere, BlueprintReadWrite, Category = GameplayCueNotify)
+	FGameplayCueNotify_CameraLensEffectInfo BurstCameraLensEffect;
+
+	// Force feedback to be played on gameplay cue execution.  This should never use a looping effect!
+	UPROPERTY(EditAnywhere, BlueprintReadWrite, Category = GameplayCueNotify)
+	FGameplayCueNotify_ForceFeedbackInfo BurstForceFeedback;
+
+	// Decal to be spawned on gameplay cue execution.  Actor should have fade out time or override should be set so it will clean up properly.
+	UPROPERTY(EditAnywhere, BlueprintReadWrite, Category = GameplayCueNotify)
+	FGameplayCueNotify_DecalInfo BurstDecal;
+};
+
+
+/**
+ * FGameplayCueNotify_LoopingEffects
+ *
+ *	Set of looping effects to spawn for looping gameplay cues.
+ */
+USTRUCT(BlueprintType)
+struct FGameplayCueNotify_LoopingEffects
+{
+	GENERATED_BODY()
+
+public:
+
+	FGameplayCueNotify_LoopingEffects();
+
+	void StartEffects(const FGameplayCueNotify_SpawnContext& SpawnContext, FGameplayCueNotify_SpawnResult& OutSpawnResult) const;
+	void StopEffects(FGameplayCueNotify_SpawnResult& SpawnResult) const;
+	
+	void ValidateAssociatedAssets(UObject* ContainingAsset, const FString& Context, TArray<FText>& ValidationErrors) const;
+
+protected:
+
+	// Particle systems to be spawned on gameplay cue loop start.
+	UPROPERTY(EditAnywhere, BlueprintReadWrite, Category = GameplayCueNotify)
+	TArray<FGameplayCueNotify_ParticleInfo> LoopingParticles;
+
+	// Sound to be played on gameplay cue loop start.
+	UPROPERTY(EditAnywhere, BlueprintReadWrite, Category = GameplayCueNotify)
+	TArray<FGameplayCueNotify_SoundInfo> LoopingSounds;
+
+	// Camera shake to be played on gameplay cue loop start.
+	UPROPERTY(EditAnywhere, BlueprintReadWrite, Category = GameplayCueNotify)
+	FGameplayCueNotify_CameraShakeInfo LoopingCameraShake;
+
+	// Camera animation to be played on gameplay cue loop start.
+	UPROPERTY(EditAnywhere, BlueprintReadWrite, Category = GameplayCueNotify)
+	FGameplayCueNotify_CameraAnimInfo LoopingCameraAnim;
+
+	// Camera lens effect to be played on gameplay cue loop start.
+	UPROPERTY(EditAnywhere, BlueprintReadWrite, Category = GameplayCueNotify)
+	FGameplayCueNotify_CameraLensEffectInfo LoopingCameraLensEffect;
+
+	// Force feedback to be played on gameplay cue loop start.
+	UPROPERTY(EditAnywhere, BlueprintReadWrite, Category = GameplayCueNotify)
+	FGameplayCueNotify_ForceFeedbackInfo LoopingForceFeedback;
+};
diff --git a/Engine/Plugins/Runtime/GameplayAbilities/Source/GameplayAbilities/Public/GameplayCueNotify_Burst.h b/Engine/Plugins/Runtime/GameplayAbilities/Source/GameplayAbilities/Public/GameplayCueNotify_Burst.h
new file mode 100644
--- /dev/null
+++ b/Engine/Plugins/Runtime/GameplayAbilities/Source/GameplayAbilities/Public/GameplayCueNotify_Burst.h
@@ -0,0 +1,50 @@
+// Copyright Epic Games, Inc. All Rights Reserved.
+
+#pragma once
+
+#include "CoreMinimal.h"
+#include "GameplayCueNotify_Static.h"
+#include "GameplayCueNotifyTypes.h"
+#include "GameplayCueNotify_Burst.generated.h"
+
+
+/**
+ * UGameplayCueNotify_Burst
+ *
+ *	This is a non-instanced gameplay cue notify for effects that are one-offs.
+ *	Since it is not instanced, it cannot do latent actions such as delays and time lines.
+ */
+UCLASS(Blueprintable, notplaceable, Category = "GameplayCueNotify", Meta = (ShowWorldContextPin, DisplayName = "GCN Burst", ShortTooltip = "A one-off GameplayCueNotify that is never spawned into the world."))
+class UGameplayCueNotify_Burst : public UGameplayCueNotify_Static
+{
+	GENERATED_BODY()
+
+public:
+
+	UGameplayCueNotify_Burst();
+
+protected:
+
+	virtual bool OnExecute_Implementation(AActor* Target, const FGameplayCueParameters& Parameters) const override;
+
+	UFUNCTION(BlueprintImplementableEvent)
+	void OnBurst(AActor* Target, const FGameplayCueParameters& Parameters, const FGameplayCueNotify_SpawnResult& SpawnResults) const;
+
+#if WITH_EDITOR
+	virtual EDataValidationResult IsDataValid(TArray<FText>& ValidationErrors) override;
+#endif // #if WITH_EDITOR
+
+protected:
+
+	// Default condition to check before spawning anything.  Applies for all spawns unless overridden.
+	UPROPERTY(EditAnywhere, BlueprintReadOnly, Category = "Defaults")
+	FGameplayCueNotify_SpawnCondition DefaultSpawnCondition;
+
+	// Default placement rules.  Applies for all spawns unless overridden.
+	UPROPERTY(EditAnywhere, BlueprintReadOnly, Category = "Defaults")
+	FGameplayCueNotify_PlacementInfo DefaultPlacementInfo;
+
+	// List of effects to spawn on burst.
+	UPROPERTY(EditAnywhere, BlueprintReadOnly, Category = "Effects")
+	FGameplayCueNotify_BurstEffects BurstEffects;
+};
diff --git a/Engine/Plugins/Runtime/GameplayAbilities/Source/GameplayAbilities/Public/GameplayCueNotify_BurstLatent.h b/Engine/Plugins/Runtime/GameplayAbilities/Source/GameplayAbilities/Public/GameplayCueNotify_BurstLatent.h
new file mode 100644
--- /dev/null
+++ b/Engine/Plugins/Runtime/GameplayAbilities/Source/GameplayAbilities/Public/GameplayCueNotify_BurstLatent.h
@@ -0,0 +1,56 @@
+// Copyright Epic Games, Inc. All Rights Reserved.
+
+#pragma once
+
+#include "CoreMinimal.h"
+#include "GameplayCueNotify_Instanced.h"
+#include "GameplayCueNotifyTypes.h"
+#include "GameplayCueNotify_BurstLatent.generated.h"
+
+
+/**
+ * AGameplayCueNotify_BurstLatent
+ *
+ *	This is an instanced gameplay cue notify for effects that are one-offs.
+ *	Since it is instanced, it can do latent things like time lines or delays.
+ */
+UCLASS(Blueprintable, notplaceable, Category = "GameplayCueNotify", Meta = (ShowWorldContextPin, DisplayName = "GCN Burst Latent", ShortTooltip = "A one-off GameplayCueNotify that can use latent actions such as timelines."))
+class AGameplayCueNotify_BurstLatent : public AGameplayCueNotify_Instanced
+{
+	GENERATED_BODY()
+
+public:
+
+	AGameplayCueNotify_BurstLatent();
+
+protected:
+
+	virtual bool Recycle() override;
+
+	virtual bool OnExecute_Implementation(AActor* Target, const FGameplayCueParameters& Parameters) override;
+
+	UFUNCTION(BlueprintImplementableEvent)
+	void OnBurst(AActor* Target, const FGameplayCueParameters& Parameters, const FGameplayCueNotify_SpawnResult& SpawnResults) const;
+
+#if WITH_EDITOR
+	virtual EDataValidationResult IsDataValid(TArray<FText>& ValidationErrors) override;
+#endif // #if WITH_EDITOR
+
+protected:
+
+	// Default condition to check before spawning anything.  Applies for all spawns unless overridden.
+	UPROPERTY(EditAnywhere, BlueprintReadOnly, Category = "Defaults")
+	FGameplayCueNotify_SpawnCondition DefaultSpawnCondition;
+
+	// Default placement rules.  Applies for all spawns unless overridden.
+	UPROPERTY(EditAnywhere, BlueprintReadOnly, Category = "Defaults")
+	FGameplayCueNotify_PlacementInfo DefaultPlacementInfo;
+
+	// List of effects to spawn on burst.
+	UPROPERTY(EditAnywhere, BlueprintReadOnly, Category = "Effects")
+	FGameplayCueNotify_BurstEffects BurstEffects;
+
+	// Results of spawned burst effects.
+	UPROPERTY(BlueprintReadOnly, Category = "Effects")
+	FGameplayCueNotify_SpawnResult BurstSpawnResults;
+};
diff --git a/Engine/Plugins/Runtime/GameplayAbilities/Source/GameplayAbilities/Public/GameplayCueNotify_Instanced.h b/Engine/Plugins/Runtime/GameplayAbilities/Source/GameplayAbilities/Public/GameplayCueNotify_Instanced.h
new file mode 100644
--- /dev/null
+++ b/Engine/Plugins/Runtime/GameplayAbilities/Source/GameplayAbilities/Public/GameplayCueNotify_Instanced.h
@@ -0,0 +1,32 @@
+// Copyright Epic Games, Inc. All Rights Reserved.
+
+#pragma once
+
+#include "CoreMinimal.h"
+#include "GameplayCueNotify_Actor.h"
+#include "GameplayCueNotify_Instanced.generated.h"
+
+
+/**
+ * AGameplayCueNotify_Instanced
+ *
+ *	Base class for instanced gameplay cue notifies (such as Looping and BurstLatent).
+ */
+UCLASS(NotBlueprintable, notplaceable, Abstract, Meta = (ShowWorldContextPin), HideCategories = (Tick, Rendering, Input, Actor))
+class AGameplayCueNotify_Instanced : public AGameplayCueNotify_Actor
+{
+	GENERATED_BODY()
+
+public:
+
+	AGameplayCueNotify_Instanced();
+
+protected:
+
+	virtual void EndPlay(const EEndPlayReason::Type EndPlayReason) override;
+	virtual void K2_DestroyActor() override;
+	virtual void Destroyed() override;
+
+	virtual bool WhileActive_Implementation(AActor* Target, const FGameplayCueParameters& Parameters) override;
+	virtual bool OnRemove_Implementation(AActor* Target, const FGameplayCueParameters& Parameters) override;
+};
diff --git a/Engine/Plugins/Runtime/GameplayAbilities/Source/GameplayAbilities/Public/GameplayCueNotify_Looping.h b/Engine/Plugins/Runtime/GameplayAbilities/Source/GameplayAbilities/Public/GameplayCueNotify_Looping.h
new file mode 100644
--- /dev/null
+++ b/Engine/Plugins/Runtime/GameplayAbilities/Source/GameplayAbilities/Public/GameplayCueNotify_Looping.h
@@ -0,0 +1,97 @@
+// Copyright Epic Games, Inc. All Rights Reserved.
+
+#pragma once
+
+#include "CoreMinimal.h"
+#include "GameplayCueNotify_Instanced.h"
+#include "GameplayCueNotifyTypes.h"
+#include "GameplayCueNotify_Looping.generated.h"
+
+
+/**
+ * AGameplayCueNotify_Looping
+ *
+ *	This is an instanced gameplay cue notify for continuous looping effects.
+ *	The game is responsible for defining the start/stop by adding/removing the gameplay cue.
+ */
+UCLASS(Blueprintable, notplaceable, Category = "GameplayCueNotify", Meta = (ShowWorldContextPin, DisplayName = "GCN Looping", ShortTooltip = "A GameplayCueNotify that has a duration that is driven by the game."))
+class AGameplayCueNotify_Looping : public AGameplayCueNotify_Instanced
+{
+	GENERATED_BODY()
+
+public:
+
+	AGameplayCueNotify_Looping();
+
+protected:
+
+	virtual void EndPlay(const EEndPlayReason::Type EndPlayReason) override;
+	virtual bool Recycle() override;
+
+	virtual bool OnActive_Implementation(AActor* Target, const FGameplayCueParameters& Parameters) override;
+	virtual bool WhileActive_Implementation(AActor* Target, const FGameplayCueParameters& Parameters) override;
+	virtual bool OnExecute_Implementation(AActor* Target, const FGameplayCueParameters& Parameters) override;
+	virtual bool OnRemove_Implementation(AActor* Target, const FGameplayCueParameters& Parameters) override;
+
+	UFUNCTION(BlueprintImplementableEvent)
+	void OnApplication(AActor* Target, const FGameplayCueParameters& Parameters, const FGameplayCueNotify_SpawnResult& SpawnResults) const;
+
+	UFUNCTION(BlueprintImplementableEvent)
+	void OnLoopingStart(AActor* Target, const FGameplayCueParameters& Parameters, const FGameplayCueNotify_SpawnResult& SpawnResults) const;
+
+	UFUNCTION(BlueprintImplementableEvent)
+	void OnRecurring(AActor* Target, const FGameplayCueParameters& Parameters, const FGameplayCueNotify_SpawnResult& SpawnResults) const;
+
+	UFUNCTION(BlueprintImplementableEvent)
+	void OnRemoval(AActor* Target, const FGameplayCueParameters& Parameters, const FGameplayCueNotify_SpawnResult& SpawnResults) const;
+
+	void RemoveLoopingEffects();
+
+#if WITH_EDITOR
+	virtual EDataValidationResult IsDataValid(TArray<FText>& ValidationErrors) override;
+#endif // #if WITH_EDITOR
+
+protected:
+
+	// Default condition to check before spawning anything.  Applies for all spawns unless overridden.
+	UPROPERTY(EditAnywhere, BlueprintReadOnly, Category = "Defaults")
+	FGameplayCueNotify_SpawnCondition DefaultSpawnCondition;
+
+	// Default placement rules.  Applies for all spawns unless overridden.
+	UPROPERTY(EditAnywhere, BlueprintReadOnly, Category = "Defaults")
+	FGameplayCueNotify_PlacementInfo DefaultPlacementInfo;
+
+	// List of effects to spawn on application.  These should not be looping effects!
+	UPROPERTY(EditAnywhere, BlueprintReadOnly, Category = "Application Effects (On Active)")
+	FGameplayCueNotify_BurstEffects ApplicationEffects;
+
+	// Results of spawned application effects.
+	UPROPERTY(BlueprintReadOnly, Category = "Application Effects (On Active)")
+	FGameplayCueNotify_SpawnResult ApplicationSpawnResults;
+
+	// List of effects to spawn on loop start.
+	UPROPERTY(EditAnywhere, BlueprintReadOnly, Category = "Looping Effects (While Active)")
+	FGameplayCueNotify_LoopingEffects LoopingEffects;
+
+	// Results of spawned looping effects.
+	UPROPERTY(BlueprintReadOnly, Category = "Looping Effects (While Active)")
+	FGameplayCueNotify_SpawnResult LoopingSpawnResults;
+
+	// List of effects to spawn for a recurring gameplay effect (e.g. each time a DOT ticks).  These should not be looping effects!
+	UPROPERTY(EditAnywhere, BlueprintReadOnly, Category = "Recurring Effects (On Execute)")
+	FGameplayCueNotify_BurstEffects RecurringEffects;
+
+	// Results of spawned recurring effects.
+	UPROPERTY(BlueprintReadOnly, Category = "Recurring Effects (On Execute)")
+	FGameplayCueNotify_SpawnResult RecurringSpawnResults;
+
+	// List of effects to spawn on removal.  These should not be looping effects!
+	UPROPERTY(EditAnywhere, BlueprintReadOnly, Category = "Removal Effects (On Remove)")
+	FGameplayCueNotify_BurstEffects RemovalEffects;
+
+	// Results of spawned removal effects.
+	UPROPERTY(BlueprintReadOnly, Category = "Removal Effects (On Remove)")
+	FGameplayCueNotify_SpawnResult RemovalSpawnResults;
+
+	bool bLoopingEffectsRemoved;
+};
