Added AddLooseGameplayTags and RemoveLooseGameplayTags functions to UAbilitySystemBlueprintLibrary

diff --git a/Engine/Plugins/Runtime/GameplayAbilities/Source/GameplayAbilities/Private/AbilitySystemBlueprintLibrary.cpp b/Engine/Plugins/Runtime/GameplayAbilities/Source/GameplayAbilities/Private/AbilitySystemBlueprintLibrary.cpp
--- a/Engine/Plugins/Runtime/GameplayAbilities/Source/GameplayAbilities/Private/AbilitySystemBlueprintLibrary.cpp
+++ b/Engine/Plugins/Runtime/GameplayAbilities/Source/GameplayAbilities/Private/AbilitySystemBlueprintLibrary.cpp
@@ -1044,3 +1044,25 @@ FString UAbilitySystemBlueprintLibrary::GetActiveGameplayEffectDebugString(FActi
 	}
 	return Str;
 }
+
+bool UAbilitySystemBlueprintLibrary::AddLooseGameplayTags(AActor* Actor, const FGameplayTagContainer& GameplayTags)
+{
+	if (UAbilitySystemComponent* AbilitySysComp = GetAbilitySystemComponent(Actor))
+	{
+		AbilitySysComp->AddLooseGameplayTags(GameplayTags);
+		return true;
+	}
+
+	return false;
+}
+
+bool UAbilitySystemBlueprintLibrary::RemoveLooseGameplayTags(AActor* Actor, const FGameplayTagContainer& GameplayTags)
+{
+	if (UAbilitySystemComponent* AbilitySysComp = GetAbilitySystemComponent(Actor))
+	{
+		AbilitySysComp->RemoveLooseGameplayTags(GameplayTags);
+		return true;
+	}
+
+	return false;
+}
diff --git a/Engine/Plugins/Runtime/GameplayAbilities/Source/GameplayAbilities/Public/AbilitySystemBlueprintLibrary.h b/Engine/Plugins/Runtime/GameplayAbilities/Source/GameplayAbilities/Public/AbilitySystemBlueprintLibrary.h
--- a/Engine/Plugins/Runtime/GameplayAbilities/Source/GameplayAbilities/Public/AbilitySystemBlueprintLibrary.h
+++ b/Engine/Plugins/Runtime/GameplayAbilities/Source/GameplayAbilities/Public/AbilitySystemBlueprintLibrary.h
@@ -334,6 +334,14 @@ class GAMEPLAYABILITIES_API UAbilitySystemBlueprintLibrary : public UBlueprintFu
 	UFUNCTION(BlueprintPure, Category = "Ability|GameplayEffect")
 	static TArray<FGameplayEffectSpecHandle> GetAllLinkedGameplayEffectSpecHandles(FGameplayEffectSpecHandle SpecHandle);
 
+	/** manually adds a set of tags to given actor */
+	UFUNCTION(BlueprintCallable, Category="Ability|GameplayEffect")
+	static bool AddLooseGameplayTags(AActor* Actor, const FGameplayTagContainer& GameplayTags);
+
+	/** manually removes a set of tags from a given actor */
+	UFUNCTION(BlueprintCallable, Category="Ability|GameplayEffect")
+	static bool RemoveLooseGameplayTags(AActor* Actor, const FGameplayTagContainer& GameplayTags);
+
 	// -------------------------------------------------------------------------------
 	//		GameplayEffectSpec
 	// -------------------------------------------------------------------------------
