Update the blueprint events on the instanced gameplay cue notifies (BurstLatent, Looping) to not be const.

diff --git a/Engine/Plugins/Runtime/GameplayAbilities/Source/GameplayAbilities/Public/GameplayCueNotify_BurstLatent.h b/Engine/Plugins/Runtime/GameplayAbilities/Source/GameplayAbilities/Public/GameplayCueNotify_BurstLatent.h
--- a/Engine/Plugins/Runtime/GameplayAbilities/Source/GameplayAbilities/Public/GameplayCueNotify_BurstLatent.h
+++ b/Engine/Plugins/Runtime/GameplayAbilities/Source/GameplayAbilities/Public/GameplayCueNotify_BurstLatent.h
@@ -30,7 +30,7 @@ protected:
 	virtual bool OnExecute_Implementation(AActor* Target, const FGameplayCueParameters& Parameters) override;
 
 	UFUNCTION(BlueprintImplementableEvent)
-	void OnBurst(AActor* Target, const FGameplayCueParameters& Parameters, const FGameplayCueNotify_SpawnResult& SpawnResults) const;
+	void OnBurst(AActor* Target, const FGameplayCueParameters& Parameters, const FGameplayCueNotify_SpawnResult& SpawnResults);
 
 #if WITH_EDITOR
 	virtual EDataValidationResult IsDataValid(TArray<FText>& ValidationErrors) override;
diff --git a/Engine/Plugins/Runtime/GameplayAbilities/Source/GameplayAbilities/Public/GameplayCueNotify_Looping.h b/Engine/Plugins/Runtime/GameplayAbilities/Source/GameplayAbilities/Public/GameplayCueNotify_Looping.h
--- a/Engine/Plugins/Runtime/GameplayAbilities/Source/GameplayAbilities/Public/GameplayCueNotify_Looping.h
+++ b/Engine/Plugins/Runtime/GameplayAbilities/Source/GameplayAbilities/Public/GameplayCueNotify_Looping.h
@@ -34,16 +34,16 @@ protected:
 	virtual bool OnRemove_Implementation(AActor* Target, const FGameplayCueParameters& Parameters) override;
 
 	UFUNCTION(BlueprintImplementableEvent)
-	void OnApplication(AActor* Target, const FGameplayCueParameters& Parameters, const FGameplayCueNotify_SpawnResult& SpawnResults) const;
+	void OnApplication(AActor* Target, const FGameplayCueParameters& Parameters, const FGameplayCueNotify_SpawnResult& SpawnResults);
 
 	UFUNCTION(BlueprintImplementableEvent)
-	void OnLoopingStart(AActor* Target, const FGameplayCueParameters& Parameters, const FGameplayCueNotify_SpawnResult& SpawnResults) const;
+	void OnLoopingStart(AActor* Target, const FGameplayCueParameters& Parameters, const FGameplayCueNotify_SpawnResult& SpawnResults);
 
 	UFUNCTION(BlueprintImplementableEvent)
-	void OnRecurring(AActor* Target, const FGameplayCueParameters& Parameters, const FGameplayCueNotify_SpawnResult& SpawnResults) const;
+	void OnRecurring(AActor* Target, const FGameplayCueParameters& Parameters, const FGameplayCueNotify_SpawnResult& SpawnResults);
 
 	UFUNCTION(BlueprintImplementableEvent)
-	void OnRemoval(AActor* Target, const FGameplayCueParameters& Parameters, const FGameplayCueNotify_SpawnResult& SpawnResults) const;
+	void OnRemoval(AActor* Target, const FGameplayCueParameters& Parameters, const FGameplayCueNotify_SpawnResult& SpawnResults);
 
 	void RemoveLoopingEffects();
 
