Converted BurstLatent and Looping gameplay cue notifies to use Sparse Class Data.

diff --git a/Engine/Plugins/Runtime/GameplayAbilities/Source/GameplayAbilities/Private/GameplayCueNotify_BurstLatent.cpp b/Engine/Plugins/Runtime/GameplayAbilities/Source/GameplayAbilities/Private/GameplayCueNotify_BurstLatent.cpp
--- a/Engine/Plugins/Runtime/GameplayAbilities/Source/GameplayAbilities/Private/GameplayCueNotify_BurstLatent.cpp
+++ b/Engine/Plugins/Runtime/GameplayAbilities/Source/GameplayAbilities/Private/GameplayCueNotify_BurstLatent.cpp
@@ -29,13 +29,16 @@ bool AGameplayCueNotify_BurstLatent::OnExecute_Implementation(AActor* Target, co
 {
 	UWorld* World = GetWorld();
 
+	const FGameplayCueNotify_BurstLatentSparseData* SparseData = GetGameplayCueNotify_BurstLatentSparseData();
+	check(SparseData);
+
 	FGameplayCueNotify_SpawnContext SpawnContext(World, Target, Parameters);
-	SpawnContext.SetDefaultSpawnCondition(&DefaultSpawnCondition);
-	SpawnContext.SetDefaultPlacementInfo(&DefaultPlacementInfo);
+	SpawnContext.SetDefaultSpawnCondition(&SparseData->DefaultSpawnCondition);
+	SpawnContext.SetDefaultPlacementInfo(&SparseData->DefaultPlacementInfo);
 
-	if (DefaultSpawnCondition.ShouldSpawn(SpawnContext))
+	if (SparseData->DefaultSpawnCondition.ShouldSpawn(SpawnContext))
 	{
-		BurstEffects.ExecuteEffects(SpawnContext, BurstSpawnResults);
+		SparseData->BurstEffects.ExecuteEffects(SpawnContext, BurstSpawnResults);
 
 		OnBurst(Target, Parameters, BurstSpawnResults);
 	}
@@ -55,7 +58,10 @@ bool AGameplayCueNotify_BurstLatent::OnExecute_Implementation(AActor* Target, co
 #if WITH_EDITOR
 EDataValidationResult AGameplayCueNotify_BurstLatent::IsDataValid(TArray<FText>& ValidationErrors)
 {
-	BurstEffects.ValidateAssociatedAssets(this, TEXT("BurstEffects"), ValidationErrors);
+	const FGameplayCueNotify_BurstLatentSparseData* SparseData = GetGameplayCueNotify_BurstLatentSparseData();
+	check(SparseData);
+
+	SparseData->BurstEffects.ValidateAssociatedAssets(this, TEXT("BurstEffects"), ValidationErrors);
 
 	return ((ValidationErrors.Num() > 0) ? EDataValidationResult::Invalid : EDataValidationResult::Valid);
 }
diff --git a/Engine/Plugins/Runtime/GameplayAbilities/Source/GameplayAbilities/Private/GameplayCueNotify_Looping.cpp b/Engine/Plugins/Runtime/GameplayAbilities/Source/GameplayAbilities/Private/GameplayCueNotify_Looping.cpp
--- a/Engine/Plugins/Runtime/GameplayAbilities/Source/GameplayAbilities/Private/GameplayCueNotify_Looping.cpp
+++ b/Engine/Plugins/Runtime/GameplayAbilities/Source/GameplayAbilities/Private/GameplayCueNotify_Looping.cpp
@@ -10,8 +10,6 @@ AGameplayCueNotify_Looping::AGameplayCueNotify_Looping()
 {
 	bAllowMultipleWhileActiveEvents = false;
 
-	DefaultPlacementInfo.AttachPolicy = EGameplayCueNotify_AttachPolicy::AttachToTarget;
-
 	Recycle();
 }
 
@@ -46,13 +44,16 @@ bool AGameplayCueNotify_Looping::OnActive_Implementation(AActor* Target, const F
 {
 	UWorld* World = GetWorld();
 
+	const FGameplayCueNotify_LoopingSparseData* SparseData = GetGameplayCueNotify_LoopingSparseData();
+	check(SparseData);
+
 	FGameplayCueNotify_SpawnContext SpawnContext(World, Target, Parameters);
-	SpawnContext.SetDefaultSpawnCondition(&DefaultSpawnCondition);
-	SpawnContext.SetDefaultPlacementInfo(&DefaultPlacementInfo);
+	SpawnContext.SetDefaultSpawnCondition(&SparseData->DefaultSpawnCondition);
+	SpawnContext.SetDefaultPlacementInfo(&SparseData->DefaultPlacementInfo);
 
-	if (DefaultSpawnCondition.ShouldSpawn(SpawnContext))
+	if (SparseData->DefaultSpawnCondition.ShouldSpawn(SpawnContext))
 	{
-		ApplicationEffects.ExecuteEffects(SpawnContext, ApplicationSpawnResults);
+		SparseData->ApplicationEffects.ExecuteEffects(SpawnContext, ApplicationSpawnResults);
 
 		OnApplication(Target, Parameters, ApplicationSpawnResults);
 	}
@@ -64,15 +65,18 @@ bool AGameplayCueNotify_Looping::WhileActive_Implementation(AActor* Target, cons
 {
 	UWorld* World = GetWorld();
 
+	const FGameplayCueNotify_LoopingSparseData* SparseData = GetGameplayCueNotify_LoopingSparseData();
+	check(SparseData);
+
 	FGameplayCueNotify_SpawnContext SpawnContext(World, Target, Parameters);
-	SpawnContext.SetDefaultSpawnCondition(&DefaultSpawnCondition);
-	SpawnContext.SetDefaultPlacementInfo(&DefaultPlacementInfo);
+	SpawnContext.SetDefaultSpawnCondition(&SparseData->DefaultSpawnCondition);
+	SpawnContext.SetDefaultPlacementInfo(&SparseData->DefaultPlacementInfo);
 
-	if (DefaultSpawnCondition.ShouldSpawn(SpawnContext))
+	if (SparseData->DefaultSpawnCondition.ShouldSpawn(SpawnContext))
 	{
 		bLoopingEffectsRemoved = false;
 
-		LoopingEffects.StartEffects(SpawnContext, LoopingSpawnResults);	
+		SparseData->LoopingEffects.StartEffects(SpawnContext, LoopingSpawnResults);	
 
 		OnLoopingStart(Target, Parameters, LoopingSpawnResults);
 	}
@@ -84,13 +88,16 @@ bool AGameplayCueNotify_Looping::OnExecute_Implementation(AActor* Target, const
 {
 	UWorld* World = GetWorld();
 
+	const FGameplayCueNotify_LoopingSparseData* SparseData = GetGameplayCueNotify_LoopingSparseData();
+	check(SparseData);
+
 	FGameplayCueNotify_SpawnContext SpawnContext(World, Target, Parameters);
-	SpawnContext.SetDefaultSpawnCondition(&DefaultSpawnCondition);
-	SpawnContext.SetDefaultPlacementInfo(&DefaultPlacementInfo);
+	SpawnContext.SetDefaultSpawnCondition(&SparseData->DefaultSpawnCondition);
+	SpawnContext.SetDefaultPlacementInfo(&SparseData->DefaultPlacementInfo);
 
-	if (DefaultSpawnCondition.ShouldSpawn(SpawnContext))
+	if (SparseData->DefaultSpawnCondition.ShouldSpawn(SpawnContext))
 	{
-		RecurringEffects.ExecuteEffects(SpawnContext, RecurringSpawnResults);
+		SparseData->RecurringEffects.ExecuteEffects(SpawnContext, RecurringSpawnResults);
 
 		OnRecurring(Target, Parameters, RecurringSpawnResults);
 	}
@@ -107,13 +114,16 @@ bool AGameplayCueNotify_Looping::OnRemove_Implementation(AActor* Target, const F
 	{
 		UWorld* World = GetWorld();
 
+		const FGameplayCueNotify_LoopingSparseData* SparseData = GetGameplayCueNotify_LoopingSparseData();
+		check(SparseData);
+
 		FGameplayCueNotify_SpawnContext SpawnContext(World, Target, Parameters);
-		SpawnContext.SetDefaultSpawnCondition(&DefaultSpawnCondition);
-		SpawnContext.SetDefaultPlacementInfo(&DefaultPlacementInfo);
+		SpawnContext.SetDefaultSpawnCondition(&SparseData->DefaultSpawnCondition);
+		SpawnContext.SetDefaultPlacementInfo(&SparseData->DefaultPlacementInfo);
 
-		if (DefaultSpawnCondition.ShouldSpawn(SpawnContext))
+		if (SparseData->DefaultSpawnCondition.ShouldSpawn(SpawnContext))
 		{
-			RemovalEffects.ExecuteEffects(SpawnContext, RemovalSpawnResults);
+			SparseData->RemovalEffects.ExecuteEffects(SpawnContext, RemovalSpawnResults);
 		}
 	}
 
@@ -130,18 +140,24 @@ void AGameplayCueNotify_Looping::RemoveLoopingEffects()
 		return;
 	}
 
+	const FGameplayCueNotify_LoopingSparseData* SparseData = GetGameplayCueNotify_LoopingSparseData();
+	check(SparseData);
+
 	bLoopingEffectsRemoved = true;
 
-	LoopingEffects.StopEffects(LoopingSpawnResults);
+	SparseData->LoopingEffects.StopEffects(LoopingSpawnResults);
 }
 
 #if WITH_EDITOR
 EDataValidationResult AGameplayCueNotify_Looping::IsDataValid(TArray<FText>& ValidationErrors)
 {
-	ApplicationEffects.ValidateAssociatedAssets(this, TEXT("ApplicationEffects"), ValidationErrors);
-	LoopingEffects.ValidateAssociatedAssets(this, TEXT("LoopingEffects"), ValidationErrors);
-	RecurringEffects.ValidateAssociatedAssets(this, TEXT("RecurringEffects"), ValidationErrors);
-	RemovalEffects.ValidateAssociatedAssets(this, TEXT("RemovalEffects"), ValidationErrors);
+	const FGameplayCueNotify_LoopingSparseData* SparseData = GetGameplayCueNotify_LoopingSparseData();
+	check(SparseData);
+
+	SparseData->ApplicationEffects.ValidateAssociatedAssets(this, TEXT("ApplicationEffects"), ValidationErrors);
+	SparseData->LoopingEffects.ValidateAssociatedAssets(this, TEXT("LoopingEffects"), ValidationErrors);
+	SparseData->RecurringEffects.ValidateAssociatedAssets(this, TEXT("RecurringEffects"), ValidationErrors);
+	SparseData->RemovalEffects.ValidateAssociatedAssets(this, TEXT("RemovalEffects"), ValidationErrors);
 
 	return ((ValidationErrors.Num() > 0) ? EDataValidationResult::Invalid : EDataValidationResult::Valid);
 }
diff --git a/Engine/Plugins/Runtime/GameplayAbilities/Source/GameplayAbilities/Public/GameplayCueNotifyTypes.h b/Engine/Plugins/Runtime/GameplayAbilities/Source/GameplayAbilities/Public/GameplayCueNotifyTypes.h
--- a/Engine/Plugins/Runtime/GameplayAbilities/Source/GameplayAbilities/Public/GameplayCueNotifyTypes.h
+++ b/Engine/Plugins/Runtime/GameplayAbilities/Source/GameplayAbilities/Public/GameplayCueNotifyTypes.h
@@ -104,26 +104,26 @@ public:
 public:
 
 	// Source actor to use when determining if it is locally controlled.
-	UPROPERTY(EditAnywhere, BlueprintReadOnly, Category = GameplayCueNotify)
+	UPROPERTY(EditDefaultsOnly, BlueprintReadOnly, Category = GameplayCueNotify)
 	EGameplayCueNotify_LocallyControlledSource LocallyControlledSource;
 
 	// Locally controlled policy used to determine if the gameplay cue effects should spawn.
-	UPROPERTY(EditAnywhere, BlueprintReadOnly, Category = GameplayCueNotify)
+	UPROPERTY(EditDefaultsOnly, BlueprintReadOnly, Category = GameplayCueNotify)
 	EGameplayCueNotify_LocallyControlledPolicy LocallyControlledPolicy;
 
 	// Random chance to play the effects.  (1.0 = always play,  0.0 = never play)
-	UPROPERTY(EditAnywhere, BlueprintReadOnly, Category = GameplayCueNotify, Meta = (UIMin = "0.0", ClampMin = "0.0", UIMax = "1.0", ClampMax = "1.0"))
+	UPROPERTY(EditDefaultsOnly, BlueprintReadOnly, Category = GameplayCueNotify, Meta = (UIMin = "0.0", ClampMin = "0.0", UIMax = "1.0", ClampMax = "1.0"))
 	float ChanceToPlay;
 
 	// The gameplay cue effects will only spawn if the surface type is in this list.
 	// An empty list means everything is allowed.
-	UPROPERTY(EditAnywhere, BlueprintReadOnly, Category = GameplayCueNotify)
+	UPROPERTY(EditDefaultsOnly, BlueprintReadOnly, Category = GameplayCueNotify)
 	TArray<TEnumAsByte<EPhysicalSurface>> AllowedSurfaceTypes;
 	mutable FGameplayCueNotify_SurfaceMask AllowedSurfaceMask;
 
 	// The gameplay cue effects will only spawn if the surface type is NOT in this list.
 	// An empty list means nothing will be rejected.
-	UPROPERTY(EditAnywhere, BlueprintReadOnly, Category = GameplayCueNotify)
+	UPROPERTY(EditDefaultsOnly, BlueprintReadOnly, Category = GameplayCueNotify)
 	TArray<TEnumAsByte<EPhysicalSurface>> RejectedSurfaceTypes;
 	mutable FGameplayCueNotify_SurfaceMask RejectedSurfaceMask;
 };
@@ -167,33 +167,33 @@ public:
 public:
 
 	// Target's socket (or bone) used for location and rotation.  If "None", it uses the target's root.
-	UPROPERTY(EditAnywhere, BlueprintReadOnly, Category = GameplayCueNotify)
+	UPROPERTY(EditDefaultsOnly, BlueprintReadOnly, Category = GameplayCueNotify)
 	FName SocketName;
 
 	// Whether to attach to the target actor or not attach.
-	UPROPERTY(EditAnywhere, BlueprintReadOnly, Category = GameplayCueNotify)
+	UPROPERTY(EditDefaultsOnly, BlueprintReadOnly, Category = GameplayCueNotify)
 	EGameplayCueNotify_AttachPolicy AttachPolicy;
 
 	// How the transform is handled when attached.
-	UPROPERTY(EditAnywhere, BlueprintReadOnly, Category = GameplayCueNotify)
+	UPROPERTY(EditDefaultsOnly, BlueprintReadOnly, Category = GameplayCueNotify)
 	EAttachmentRule AttachmentRule;
 
 	// If enabled, will always spawn using rotation override.
-	UPROPERTY(EditAnywhere, BlueprintReadOnly, Category = GameplayCueNotify, Meta = (InlineEditConditionToggle))
+	UPROPERTY(EditDefaultsOnly, BlueprintReadOnly, Category = GameplayCueNotify, Meta = (InlineEditConditionToggle))
 	uint32 bOverrideRotation : 1;
 
 	// If enabled, will always spawn using the scale override.
-	UPROPERTY(EditAnywhere, BlueprintReadOnly, Category = GameplayCueNotify, Meta = (InlineEditConditionToggle))
+	UPROPERTY(EditDefaultsOnly, BlueprintReadOnly, Category = GameplayCueNotify, Meta = (InlineEditConditionToggle))
 	uint32 bOverrideScale : 1;
 
 	// If enabled, will always spawn using rotation override.
 	// This will also set the rotation to be absolute, so any changes to the parent's rotation will be ignored after attachment.
-	UPROPERTY(EditAnywhere, BlueprintReadOnly, Category = GameplayCueNotify, Meta = (EditCondition = "bOverrideRotation"))
+	UPROPERTY(EditDefaultsOnly, BlueprintReadOnly, Category = GameplayCueNotify, Meta = (EditCondition = "bOverrideRotation"))
 	FRotator RotationOverride;
 
 	// If enabled, will always spawn using scale override.
 	// This will also set the scale to be absolute, so any changes to the parent's scale will be ignored after attachment.
-	UPROPERTY(EditAnywhere, BlueprintReadOnly, Category = GameplayCueNotify, Meta = (EditCondition = "bOverrideScale"))
+	UPROPERTY(EditDefaultsOnly, BlueprintReadOnly, Category = GameplayCueNotify, Meta = (EditCondition = "bOverrideScale"))
 	FVector ScaleOverride;
 };
 
@@ -327,27 +327,27 @@ public:
 public:
 
 	// Condition to check before spawning the particle system.
-	UPROPERTY(EditAnywhere, BlueprintReadOnly, Category = GameplayCueNotify, Meta = (EditCondition = "bOverrideSpawnCondition"))
+	UPROPERTY(EditDefaultsOnly, BlueprintReadOnly, Category = GameplayCueNotify, Meta = (EditCondition = "bOverrideSpawnCondition"))
 	FGameplayCueNotify_SpawnCondition SpawnConditionOverride;
 
 	// Defines how the particle system will be placed.
-	UPROPERTY(EditAnywhere, BlueprintReadOnly, Category = GameplayCueNotify, Meta = (EditCondition = "bOverridePlacementInfo"))
+	UPROPERTY(EditDefaultsOnly, BlueprintReadOnly, Category = GameplayCueNotify, Meta = (EditCondition = "bOverridePlacementInfo"))
 	FGameplayCueNotify_PlacementInfo PlacementInfoOverride;
 
 	// Particle system to spawn.
-	UPROPERTY(EditAnywhere, BlueprintReadOnly, Category = GameplayCueNotify)
+	UPROPERTY(EditDefaultsOnly, BlueprintReadOnly, Category = GameplayCueNotify)
 	UParticleSystem* ParticleSystem;
 
 	// If enabled, use the spawn condition override and not the default one.
-	UPROPERTY(EditAnywhere, BlueprintReadOnly, Category = GameplayCueNotify, Meta = (InlineEditConditionToggle))
+	UPROPERTY(EditDefaultsOnly, BlueprintReadOnly, Category = GameplayCueNotify, Meta = (InlineEditConditionToggle))
 	uint32 bOverrideSpawnCondition : 1;
 
 	// If enabled, use the placement info override and not the default one.
-	UPROPERTY(EditAnywhere, BlueprintReadOnly, Category = GameplayCueNotify, Meta = (InlineEditConditionToggle))
+	UPROPERTY(EditDefaultsOnly, BlueprintReadOnly, Category = GameplayCueNotify, Meta = (InlineEditConditionToggle))
 	uint32 bOverridePlacementInfo : 1;
 
 	// If enabled, this particle system will cast a shadow.
-	UPROPERTY(EditAnywhere, BlueprintReadOnly, Category = GameplayCueNotify)
+	UPROPERTY(EditDefaultsOnly, BlueprintReadOnly, Category = GameplayCueNotify)
 	uint32 bCastShadow : 1;
 };
 
@@ -373,33 +373,33 @@ public:
 public:
 
 	// Condition to check before playing the sound.
-	UPROPERTY(EditAnywhere, BlueprintReadOnly, Category = GameplayCueNotify, Meta = (EditCondition = "bOverrideSpawnCondition"))
+	UPROPERTY(EditDefaultsOnly, BlueprintReadOnly, Category = GameplayCueNotify, Meta = (EditCondition = "bOverrideSpawnCondition"))
 	FGameplayCueNotify_SpawnCondition SpawnConditionOverride;
 
 	// Defines how the sound will be placed.
-	UPROPERTY(EditAnywhere, BlueprintReadOnly, Category = GameplayCueNotify, Meta = (EditCondition = "bOverridePlacementInfo"))
+	UPROPERTY(EditDefaultsOnly, BlueprintReadOnly, Category = GameplayCueNotify, Meta = (EditCondition = "bOverridePlacementInfo"))
 	FGameplayCueNotify_PlacementInfo PlacementInfoOverride;
 
 	// Sound cue to play.
-	UPROPERTY(EditAnywhere, BlueprintReadOnly, Category = GameplayCueNotify)
+	UPROPERTY(EditDefaultsOnly, BlueprintReadOnly, Category = GameplayCueNotify)
 	USoundBase* SoundCue;
 
 	// How long it should take to fade out.  Only used on looping gameplay cues.
-	UPROPERTY(EditAnywhere, BlueprintReadOnly, Category = GameplayCueNotify)
+	UPROPERTY(EditDefaultsOnly, BlueprintReadOnly, Category = GameplayCueNotify)
 	float LoopingFadeOutDuration;
 
 	// The volume level you want the sound to fade out to over the 'Looping Fade Out Duration' before stopping.
 	// This value is irrelevant if the 'Looping Fade Out Duration' is zero.
 	// NOTE: If the fade out duration is positive and this value is not lower than the volume the sound is playing at, the sound will never stop!
-	UPROPERTY(EditAnywhere, BlueprintReadOnly, Category = GameplayCueNotify)
+	UPROPERTY(EditDefaultsOnly, BlueprintReadOnly, Category = GameplayCueNotify)
 	float LoopingFadeVolumeLevel;
 
 	// If enabled, use the spawn condition override and not the default one.
-	UPROPERTY(EditAnywhere, BlueprintReadOnly, Category = GameplayCueNotify, Meta = (InlineEditConditionToggle))
+	UPROPERTY(EditDefaultsOnly, BlueprintReadOnly, Category = GameplayCueNotify, Meta = (InlineEditConditionToggle))
 	uint32 bOverrideSpawnCondition : 1;
 
 	// If enabled, use the placement info override and not the default one.
-	UPROPERTY(EditAnywhere, BlueprintReadOnly, Category = GameplayCueNotify, Meta = (InlineEditConditionToggle))
+	UPROPERTY(EditDefaultsOnly, BlueprintReadOnly, Category = GameplayCueNotify, Meta = (InlineEditConditionToggle))
 	uint32 bOverridePlacementInfo : 1;
 };
 
@@ -425,47 +425,47 @@ public:
 public:
 
 	// Condition to check before playing the camera shake.
-	UPROPERTY(EditAnywhere, BlueprintReadOnly, Category = GameplayCueNotify, Meta = (EditCondition = "bOverrideSpawnCondition"))
+	UPROPERTY(EditDefaultsOnly, BlueprintReadOnly, Category = GameplayCueNotify, Meta = (EditCondition = "bOverrideSpawnCondition"))
 	FGameplayCueNotify_SpawnCondition SpawnConditionOverride;
 
 	// Defines how the camera shake will be placed.
-	UPROPERTY(EditAnywhere, BlueprintReadOnly, Category = GameplayCueNotify, Meta = (EditCondition = "bOverridePlacementInfo"))
+	UPROPERTY(EditDefaultsOnly, BlueprintReadOnly, Category = GameplayCueNotify, Meta = (EditCondition = "bOverridePlacementInfo"))
 	FGameplayCueNotify_PlacementInfo PlacementInfoOverride;
 
 	// Camera shake to play.
-	UPROPERTY(EditAnywhere, BlueprintReadOnly, Category = GameplayCueNotify)
+	UPROPERTY(EditDefaultsOnly, BlueprintReadOnly, Category = GameplayCueNotify)
 	TSubclassOf<UCameraShake> CameraShake;
 
 	// Scale applied to the camera shake.
-	UPROPERTY(EditAnywhere, BlueprintReadOnly, Category = GameplayCueNotify)
+	UPROPERTY(EditDefaultsOnly, BlueprintReadOnly, Category = GameplayCueNotify)
 	float ShakeScale;
 
 	// What coordinate space to play the camera shake relative to.
-	UPROPERTY(EditAnywhere, BlueprintReadOnly, Category = GameplayCueNotify)
+	UPROPERTY(EditDefaultsOnly, BlueprintReadOnly, Category = GameplayCueNotify)
 	EGameplayCueNotify_EffectPlaySpace PlaySpace;
 
 	// If enabled, use the spawn condition override and not the default one.
-	UPROPERTY(EditAnywhere, BlueprintReadOnly, Category = GameplayCueNotify, Meta = (InlineEditConditionToggle))
+	UPROPERTY(EditDefaultsOnly, BlueprintReadOnly, Category = GameplayCueNotify, Meta = (InlineEditConditionToggle))
 	uint32 bOverrideSpawnCondition : 1;
 
 	// If enabled, use the placement info override and not the default one.
-	UPROPERTY(EditAnywhere, BlueprintReadOnly, Category = GameplayCueNotify, Meta = (InlineEditConditionToggle))
+	UPROPERTY(EditDefaultsOnly, BlueprintReadOnly, Category = GameplayCueNotify, Meta = (InlineEditConditionToggle))
 	uint32 bOverridePlacementInfo : 1;
 
 	// If enabled, the camera shake will be played in the world and affect all players.
-	UPROPERTY(EditAnywhere, BlueprintReadOnly, Category = GameplayCueNotify)
+	UPROPERTY(EditDefaultsOnly, BlueprintReadOnly, Category = GameplayCueNotify)
 	uint32 bPlayInWorld : 1;
 
 	// Players inside this radius get the full intensity camera shake.
-	UPROPERTY(EditAnywhere, BlueprintReadOnly, Category = GameplayCueNotify, Meta = (EditCondition = "bPlayInWorld", UIMin = "0.0", ClampMin = "0.0"))
+	UPROPERTY(EditDefaultsOnly, BlueprintReadOnly, Category = GameplayCueNotify, Meta = (EditCondition = "bPlayInWorld", UIMin = "0.0", ClampMin = "0.0"))
 	float WorldInnerRadius;
 
 	// Players outside this radius do not get the camera shake applied.
-	UPROPERTY(EditAnywhere, BlueprintReadOnly, Category = GameplayCueNotify, Meta = (EditCondition = "bPlayInWorld", UIMin = "0.0", ClampMin = "0.0"))
+	UPROPERTY(EditDefaultsOnly, BlueprintReadOnly, Category = GameplayCueNotify, Meta = (EditCondition = "bPlayInWorld", UIMin = "0.0", ClampMin = "0.0"))
 	float WorldOuterRadius;
 
 	// Exponent that describes the shake intensity falloff curve between the inner and outer radii.  (1.0 is linear)
-	UPROPERTY(EditAnywhere, BlueprintReadOnly, Category = GameplayCueNotify, Meta = (EditCondition = "bPlayInWorld", UIMin = "0.0", ClampMin = "0.0"))
+	UPROPERTY(EditDefaultsOnly, BlueprintReadOnly, Category = GameplayCueNotify, Meta = (EditCondition = "bPlayInWorld", UIMin = "0.0", ClampMin = "0.0"))
 	float WorldFalloffExponent;
 };
 
@@ -491,43 +491,43 @@ public:
 public:
 
 	// Condition to check before playing the camera animation.
-	UPROPERTY(EditAnywhere, BlueprintReadOnly, Category = GameplayCueNotify, Meta = (EditCondition = "bOverrideSpawnCondition"))
+	UPROPERTY(EditDefaultsOnly, BlueprintReadOnly, Category = GameplayCueNotify, Meta = (EditCondition = "bOverrideSpawnCondition"))
 	FGameplayCueNotify_SpawnCondition SpawnConditionOverride;
 
 	// Camera animation to play.
-	UPROPERTY(EditAnywhere, BlueprintReadOnly, Category = GameplayCueNotify)
+	UPROPERTY(EditDefaultsOnly, BlueprintReadOnly, Category = GameplayCueNotify)
 	UCameraAnim* CameraAnim;
 
 	// Scale applied to the camera animation.
-	UPROPERTY(EditAnywhere, BlueprintReadOnly, Category = GameplayCueNotify)
+	UPROPERTY(EditDefaultsOnly, BlueprintReadOnly, Category = GameplayCueNotify)
 	float AnimScale;
 
 	// Rate to play the camera animation at.
-	UPROPERTY(EditAnywhere, BlueprintReadOnly, Category = GameplayCueNotify)
+	UPROPERTY(EditDefaultsOnly, BlueprintReadOnly, Category = GameplayCueNotify)
 	float PlayRate;
 
 	// Time (in seconds) used to blend in the camera animation.
-	UPROPERTY(EditAnywhere, BlueprintReadOnly, Category = GameplayCueNotify)
+	UPROPERTY(EditDefaultsOnly, BlueprintReadOnly, Category = GameplayCueNotify)
 	float BlendInTime;
 
 	// Time (in seconds) used to blend out the camera animation.
-	UPROPERTY(EditAnywhere, BlueprintReadOnly, Category = GameplayCueNotify)
+	UPROPERTY(EditDefaultsOnly, BlueprintReadOnly, Category = GameplayCueNotify)
 	float BlendOutTime;
 
 	// What coordinate space to play the camera animation relative to.
-	UPROPERTY(EditAnywhere, BlueprintReadOnly, Category = GameplayCueNotify)
+	UPROPERTY(EditDefaultsOnly, BlueprintReadOnly, Category = GameplayCueNotify)
 	EGameplayCueNotify_EffectPlaySpace PlaySpace;
 
 	// If enabled, the camera animation will be set to loop.
-	UPROPERTY(EditAnywhere, BlueprintReadOnly, Category = GameplayCueNotify)
+	UPROPERTY(EditDefaultsOnly, BlueprintReadOnly, Category = GameplayCueNotify)
 	uint32 bIsLooping : 1;
 
 	// If enabled, the camera animation will start at a random time.  Only used when looping.
-	UPROPERTY(EditAnywhere, BlueprintReadOnly, Category = GameplayCueNotify, Meta = (EditCondition = "bIsLooping"))
+	UPROPERTY(EditDefaultsOnly, BlueprintReadOnly, Category = GameplayCueNotify, Meta = (EditCondition = "bIsLooping"))
 	uint32 bRandomStartTime : 1;
 
 	// If enabled, use the spawn condition override and not the default one.
-	UPROPERTY(EditAnywhere, BlueprintReadOnly, Category = GameplayCueNotify, Meta = (InlineEditConditionToggle))
+	UPROPERTY(EditDefaultsOnly, BlueprintReadOnly, Category = GameplayCueNotify, Meta = (InlineEditConditionToggle))
 	uint32 bOverrideSpawnCondition : 1;
 };
 
@@ -553,35 +553,35 @@ public:
 public:
 
 	// Condition to check before playing the camera lens effect.
-	UPROPERTY(EditAnywhere, BlueprintReadOnly, Category = GameplayCueNotify, Meta = (EditCondition = "bOverrideSpawnCondition"))
+	UPROPERTY(EditDefaultsOnly, BlueprintReadOnly, Category = GameplayCueNotify, Meta = (EditCondition = "bOverrideSpawnCondition"))
 	FGameplayCueNotify_SpawnCondition SpawnConditionOverride;
 
 	// Defines how the camera lens effect will be placed.
-	UPROPERTY(EditAnywhere, BlueprintReadOnly, Category = GameplayCueNotify, Meta = (EditCondition = "bOverridePlacementInfo"))
+	UPROPERTY(EditDefaultsOnly, BlueprintReadOnly, Category = GameplayCueNotify, Meta = (EditCondition = "bOverridePlacementInfo"))
 	FGameplayCueNotify_PlacementInfo PlacementInfoOverride;
 
 	// Camera lens effect to play.
-	UPROPERTY(EditAnywhere, BlueprintReadOnly, Category = GameplayCueNotify)
+	UPROPERTY(EditDefaultsOnly, BlueprintReadOnly, Category = GameplayCueNotify)
 	TSubclassOf<AEmitterCameraLensEffectBase> CameraLensEffect;
 
 	// If enabled, use the spawn condition override and not the default one.
-	UPROPERTY(EditAnywhere, BlueprintReadOnly, Category = GameplayCueNotify, Meta = (InlineEditConditionToggle))
+	UPROPERTY(EditDefaultsOnly, BlueprintReadOnly, Category = GameplayCueNotify, Meta = (InlineEditConditionToggle))
 	uint32 bOverrideSpawnCondition : 1;
 
 	// If enabled, use the placement info override and not the default one.
-	UPROPERTY(EditAnywhere, BlueprintReadOnly, Category = GameplayCueNotify, Meta = (InlineEditConditionToggle))
+	UPROPERTY(EditDefaultsOnly, BlueprintReadOnly, Category = GameplayCueNotify, Meta = (InlineEditConditionToggle))
 	uint32 bOverridePlacementInfo : 1;
 
 	// If enabled, the camera lens effect will be played in the world and affect all players.
-	UPROPERTY(EditAnywhere, BlueprintReadOnly, Category = GameplayCueNotify)
+	UPROPERTY(EditDefaultsOnly, BlueprintReadOnly, Category = GameplayCueNotify)
 	uint32 bPlayInWorld : 1;
 
 	// Players inside this radius get the full intensity camera lens effect.
-	UPROPERTY(EditAnywhere, BlueprintReadOnly, Category = GameplayCueNotify, Meta = (EditCondition = "bPlayInWorld", UIMin = "0.0", ClampMin = "0.0"))
+	UPROPERTY(EditDefaultsOnly, BlueprintReadOnly, Category = GameplayCueNotify, Meta = (EditCondition = "bPlayInWorld", UIMin = "0.0", ClampMin = "0.0"))
 	float WorldInnerRadius;
 
 	// Players outside this radius do not get the camera lens effect applied.
-	UPROPERTY(EditAnywhere, BlueprintReadOnly, Category = GameplayCueNotify, Meta = (EditCondition = "bPlayInWorld", UIMin = "0.0", ClampMin = "0.0"))
+	UPROPERTY(EditDefaultsOnly, BlueprintReadOnly, Category = GameplayCueNotify, Meta = (EditCondition = "bPlayInWorld", UIMin = "0.0", ClampMin = "0.0"))
 	float WorldOuterRadius;
 };
 
@@ -607,44 +607,44 @@ public:
 public:
 
 	// Condition to check before playing the force feedback.
-	UPROPERTY(EditAnywhere, BlueprintReadOnly, Category = GameplayCueNotify, Meta = (EditCondition = "bOverrideSpawnCondition"))
+	UPROPERTY(EditDefaultsOnly, BlueprintReadOnly, Category = GameplayCueNotify, Meta = (EditCondition = "bOverrideSpawnCondition"))
 	FGameplayCueNotify_SpawnCondition SpawnConditionOverride;
 
 	// Defines how the force feedback will be placed.
-	UPROPERTY(EditAnywhere, BlueprintReadOnly, Category = GameplayCueNotify, Meta = (EditCondition = "bOverridePlacementInfo"))
+	UPROPERTY(EditDefaultsOnly, BlueprintReadOnly, Category = GameplayCueNotify, Meta = (EditCondition = "bOverridePlacementInfo"))
 	FGameplayCueNotify_PlacementInfo PlacementInfoOverride;
 
 	// Force feedback effect to play.
-	UPROPERTY(EditAnywhere, BlueprintReadWrite, Category = GameplayCueNotify)
+	UPROPERTY(EditDefaultsOnly, BlueprintReadOnly, Category = GameplayCueNotify)
 	UForceFeedbackEffect* ForceFeedbackEffect;
 
 	// Tag used to identify the force feedback effect so it can later be canceled.
 	// Warning: If this is "None" it will stop ALL force feedback effects if it is canceled.
-	UPROPERTY(EditAnywhere, BlueprintReadOnly, Category = GameplayCueNotify)
+	UPROPERTY(EditDefaultsOnly, BlueprintReadOnly, Category = GameplayCueNotify)
 	FName ForceFeedbackTag;
 
 	// If enabled, the force feedback will be set to loop.
-	UPROPERTY(EditAnywhere, BlueprintReadOnly, Category = GameplayCueNotify)
+	UPROPERTY(EditDefaultsOnly, BlueprintReadOnly, Category = GameplayCueNotify)
 	uint32 bIsLooping : 1;
 
 	// If enabled, use the spawn condition override and not the default one.
-	UPROPERTY(EditAnywhere, BlueprintReadOnly, Category = GameplayCueNotify, Meta = (InlineEditConditionToggle))
+	UPROPERTY(EditDefaultsOnly, BlueprintReadOnly, Category = GameplayCueNotify, Meta = (InlineEditConditionToggle))
 	uint32 bOverrideSpawnCondition : 1;
 
 	// If enabled, use the placement info override and not the default one.
-	UPROPERTY(EditAnywhere, BlueprintReadOnly, Category = GameplayCueNotify, Meta = (InlineEditConditionToggle))
+	UPROPERTY(EditDefaultsOnly, BlueprintReadOnly, Category = GameplayCueNotify, Meta = (InlineEditConditionToggle))
 	uint32 bOverridePlacementInfo : 1;
 
 	// If enabled, the force feedback will be played in the world and affect all players.
-	UPROPERTY(EditAnywhere, BlueprintReadOnly, Category = GameplayCueNotify)
+	UPROPERTY(EditDefaultsOnly, BlueprintReadOnly, Category = GameplayCueNotify)
 	uint32 bPlayInWorld : 1;
 
 	// Multiplier applied to the force feedback when played in world.
-	UPROPERTY(EditAnywhere, BlueprintReadWrite, Category = GameplayCueNotify, Meta = (EditCondition = "bPlayInWorld"))
+	UPROPERTY(EditDefaultsOnly, BlueprintReadOnly, Category = GameplayCueNotify, Meta = (EditCondition = "bPlayInWorld"))
 	float WorldIntensity;
 
 	// How the force feedback is attenuated over distance when played in world.
-	UPROPERTY(EditAnywhere, BlueprintReadWrite, Category = GameplayCueNotify, Meta = (EditCondition = "bPlayInWorld"))
+	UPROPERTY(EditDefaultsOnly, BlueprintReadOnly, Category = GameplayCueNotify, Meta = (EditCondition = "bPlayInWorld"))
 	UForceFeedbackAttenuation* WorldAttenuation;
 };
 
@@ -668,35 +668,35 @@ public:
 public:
 
 	// Condition to check before spawning the decal.
-	UPROPERTY(EditAnywhere, BlueprintReadOnly, Category = GameplayCueNotify, Meta = (EditCondition = "bOverrideSpawnCondition"))
+	UPROPERTY(EditDefaultsOnly, BlueprintReadOnly, Category = GameplayCueNotify, Meta = (EditCondition = "bOverrideSpawnCondition"))
 	FGameplayCueNotify_SpawnCondition SpawnConditionOverride;
 
 	// Defines how the decal will be placed.
-	UPROPERTY(EditAnywhere, BlueprintReadOnly, Category = GameplayCueNotify, Meta = (EditCondition = "bOverridePlacementInfo"))
+	UPROPERTY(EditDefaultsOnly, BlueprintReadOnly, Category = GameplayCueNotify, Meta = (EditCondition = "bOverridePlacementInfo"))
 	FGameplayCueNotify_PlacementInfo PlacementInfoOverride;
 
 	// Decal actor class to spawn.
-	UPROPERTY(EditAnywhere, BlueprintReadOnly, Category = GameplayCueNotify)
+	UPROPERTY(EditDefaultsOnly, BlueprintReadOnly, Category = GameplayCueNotify)
 	TSubclassOf<ADecalActor> Decal;
 
 	// If enabled, use the spawn condition override and not the default one.
-	UPROPERTY(EditAnywhere, BlueprintReadOnly, Category = GameplayCueNotify, Meta = (InlineEditConditionToggle))
+	UPROPERTY(EditDefaultsOnly, BlueprintReadOnly, Category = GameplayCueNotify, Meta = (InlineEditConditionToggle))
 	uint32 bOverrideSpawnCondition : 1;
 
 	// If enabled, use the placement info override and not the default one.
-	UPROPERTY(EditAnywhere, BlueprintReadOnly, Category = GameplayCueNotify, Meta = (InlineEditConditionToggle))
+	UPROPERTY(EditDefaultsOnly, BlueprintReadOnly, Category = GameplayCueNotify, Meta = (InlineEditConditionToggle))
 	uint32 bOverridePlacementInfo : 1;
 
 	// If enabled, override default decal actor fade out values.
-	UPROPERTY(EditAnywhere, BlueprintReadOnly, Category = GameplayCueNotify)
+	UPROPERTY(EditDefaultsOnly, BlueprintReadOnly, Category = GameplayCueNotify)
 	uint32 bOverrideFadeOut : 1;
 
 	// Sets when the decal actor will start fading out.  Will override setting in class.
-	UPROPERTY(EditAnywhere, BlueprintReadOnly, Category = GameplayCueNotify, meta = (EditCondition = "bOverrideFadeOut"))
+	UPROPERTY(EditDefaultsOnly, BlueprintReadOnly, Category = GameplayCueNotify, meta = (EditCondition = "bOverrideFadeOut"))
 	float FadeOutStartDelay;
 
 	// Sets how long it takes for decal actor to fade out.  Will override setting in class.
-	UPROPERTY(EditAnywhere, BlueprintReadOnly, Category = GameplayCueNotify, meta = (EditCondition = "bOverrideFadeOut"))
+	UPROPERTY(EditDefaultsOnly, BlueprintReadOnly, Category = GameplayCueNotify, meta = (EditCondition = "bOverrideFadeOut"))
 	float FadeOutDuration;
 };
 
@@ -722,31 +722,31 @@ public:
 protected:
 
 	// Particle systems to be spawned on gameplay cue execution.  These should never use looping effects!
-	UPROPERTY(EditAnywhere, BlueprintReadWrite, Category = GameplayCueNotify)
+	UPROPERTY(EditDefaultsOnly, BlueprintReadOnly, Category = GameplayCueNotify)
 	TArray<FGameplayCueNotify_ParticleInfo> BurstParticles;
 
 	// Sound to be played on gameplay cue execution.  These should never use looping effects!
-	UPROPERTY(EditAnywhere, BlueprintReadWrite, Category = GameplayCueNotify)
+	UPROPERTY(EditDefaultsOnly, BlueprintReadOnly, Category = GameplayCueNotify)
 	TArray<FGameplayCueNotify_SoundInfo> BurstSounds;
 
 	// Camera shake to be played on gameplay cue execution.  This should never use a looping effect!
-	UPROPERTY(EditAnywhere, BlueprintReadWrite, Category = GameplayCueNotify)
+	UPROPERTY(EditDefaultsOnly, BlueprintReadOnly, Category = GameplayCueNotify)
 	FGameplayCueNotify_CameraShakeInfo BurstCameraShake;
 
 	// Camera animation to be played on gameplay cue execution.  This should never use a looping effect!
-	UPROPERTY(EditAnywhere, BlueprintReadWrite, Category = GameplayCueNotify)
+	UPROPERTY(EditDefaultsOnly, BlueprintReadOnly, Category = GameplayCueNotify)
 	FGameplayCueNotify_CameraAnimInfo BurstCameraAnim;
 
 	// Camera lens effect to be played on gameplay cue execution.  This should never use a looping effect!
-	UPROPERTY(EditAnywhere, BlueprintReadWrite, Category = GameplayCueNotify)
+	UPROPERTY(EditDefaultsOnly, BlueprintReadOnly, Category = GameplayCueNotify)
 	FGameplayCueNotify_CameraLensEffectInfo BurstCameraLensEffect;
 
 	// Force feedback to be played on gameplay cue execution.  This should never use a looping effect!
-	UPROPERTY(EditAnywhere, BlueprintReadWrite, Category = GameplayCueNotify)
+	UPROPERTY(EditDefaultsOnly, BlueprintReadOnly, Category = GameplayCueNotify)
 	FGameplayCueNotify_ForceFeedbackInfo BurstForceFeedback;
 
 	// Decal to be spawned on gameplay cue execution.  Actor should have fade out time or override should be set so it will clean up properly.
-	UPROPERTY(EditAnywhere, BlueprintReadWrite, Category = GameplayCueNotify)
+	UPROPERTY(EditDefaultsOnly, BlueprintReadOnly, Category = GameplayCueNotify)
 	FGameplayCueNotify_DecalInfo BurstDecal;
 };
 
@@ -773,26 +773,26 @@ public:
 protected:
 
 	// Particle systems to be spawned on gameplay cue loop start.
-	UPROPERTY(EditAnywhere, BlueprintReadWrite, Category = GameplayCueNotify)
+	UPROPERTY(EditDefaultsOnly, BlueprintReadOnly, Category = GameplayCueNotify)
 	TArray<FGameplayCueNotify_ParticleInfo> LoopingParticles;
 
 	// Sound to be played on gameplay cue loop start.
-	UPROPERTY(EditAnywhere, BlueprintReadWrite, Category = GameplayCueNotify)
+	UPROPERTY(EditDefaultsOnly, BlueprintReadOnly, Category = GameplayCueNotify)
 	TArray<FGameplayCueNotify_SoundInfo> LoopingSounds;
 
 	// Camera shake to be played on gameplay cue loop start.
-	UPROPERTY(EditAnywhere, BlueprintReadWrite, Category = GameplayCueNotify)
+	UPROPERTY(EditDefaultsOnly, BlueprintReadOnly, Category = GameplayCueNotify)
 	FGameplayCueNotify_CameraShakeInfo LoopingCameraShake;
 
 	// Camera animation to be played on gameplay cue loop start.
-	UPROPERTY(EditAnywhere, BlueprintReadWrite, Category = GameplayCueNotify)
+	UPROPERTY(EditDefaultsOnly, BlueprintReadOnly, Category = GameplayCueNotify)
 	FGameplayCueNotify_CameraAnimInfo LoopingCameraAnim;
 
 	// Camera lens effect to be played on gameplay cue loop start.
-	UPROPERTY(EditAnywhere, BlueprintReadWrite, Category = GameplayCueNotify)
+	UPROPERTY(EditDefaultsOnly, BlueprintReadOnly, Category = GameplayCueNotify)
 	FGameplayCueNotify_CameraLensEffectInfo LoopingCameraLensEffect;
 
 	// Force feedback to be played on gameplay cue loop start.
-	UPROPERTY(EditAnywhere, BlueprintReadWrite, Category = GameplayCueNotify)
+	UPROPERTY(EditDefaultsOnly, BlueprintReadOnly, Category = GameplayCueNotify)
 	FGameplayCueNotify_ForceFeedbackInfo LoopingForceFeedback;
 };
diff --git a/Engine/Plugins/Runtime/GameplayAbilities/Source/GameplayAbilities/Public/GameplayCueNotify_Burst.h b/Engine/Plugins/Runtime/GameplayAbilities/Source/GameplayAbilities/Public/GameplayCueNotify_Burst.h
--- a/Engine/Plugins/Runtime/GameplayAbilities/Source/GameplayAbilities/Public/GameplayCueNotify_Burst.h
+++ b/Engine/Plugins/Runtime/GameplayAbilities/Source/GameplayAbilities/Public/GameplayCueNotify_Burst.h
@@ -37,14 +37,14 @@ protected:
 protected:
 
 	// Default condition to check before spawning anything.  Applies for all spawns unless overridden.
-	UPROPERTY(EditAnywhere, BlueprintReadOnly, Category = "Defaults")
+	UPROPERTY(EditDefaultsOnly, BlueprintReadOnly, Category = "GCN Defaults")
 	FGameplayCueNotify_SpawnCondition DefaultSpawnCondition;
 
 	// Default placement rules.  Applies for all spawns unless overridden.
-	UPROPERTY(EditAnywhere, BlueprintReadOnly, Category = "Defaults")
+	UPROPERTY(EditDefaultsOnly, BlueprintReadOnly, Category = "GCN Defaults")
 	FGameplayCueNotify_PlacementInfo DefaultPlacementInfo;
 
 	// List of effects to spawn on burst.
-	UPROPERTY(EditAnywhere, BlueprintReadOnly, Category = "Effects")
+	UPROPERTY(EditDefaultsOnly, BlueprintReadOnly, Category = "GCN Effects")
 	FGameplayCueNotify_BurstEffects BurstEffects;
 };
diff --git a/Engine/Plugins/Runtime/GameplayAbilities/Source/GameplayAbilities/Public/GameplayCueNotify_BurstLatent.h b/Engine/Plugins/Runtime/GameplayAbilities/Source/GameplayAbilities/Public/GameplayCueNotify_BurstLatent.h
--- a/Engine/Plugins/Runtime/GameplayAbilities/Source/GameplayAbilities/Public/GameplayCueNotify_BurstLatent.h
+++ b/Engine/Plugins/Runtime/GameplayAbilities/Source/GameplayAbilities/Public/GameplayCueNotify_BurstLatent.h
@@ -8,13 +8,38 @@
 #include "GameplayCueNotify_BurstLatent.generated.h"
 
 
+USTRUCT(BlueprintType)
+struct FGameplayCueNotify_BurstLatentSparseData
+{
+	GENERATED_BODY()
+
+public:
+
+	FGameplayCueNotify_BurstLatentSparseData() {}
+
+public:
+
+	// Default condition to check before spawning anything.  Applies for all spawns unless overridden.
+	UPROPERTY(EditDefaultsOnly, BlueprintReadOnly, Category = "GCN Defaults")
+	FGameplayCueNotify_SpawnCondition DefaultSpawnCondition;
+
+	// Default placement rules.  Applies for all spawns unless overridden.
+	UPROPERTY(EditDefaultsOnly, BlueprintReadOnly, Category = "GCN Defaults")
+	FGameplayCueNotify_PlacementInfo DefaultPlacementInfo;
+
+	// List of effects to spawn on burst.
+	UPROPERTY(EditDefaultsOnly, BlueprintReadOnly, Category = "GCN Effects")
+	FGameplayCueNotify_BurstEffects BurstEffects;
+};
+
+
 /**
  * AGameplayCueNotify_BurstLatent
  *
  *	This is an instanced gameplay cue notify for effects that are one-offs.
  *	Since it is instanced, it can do latent things like time lines or delays.
  */
-UCLASS(Blueprintable, notplaceable, Category = "GameplayCueNotify", Meta = (ShowWorldContextPin, DisplayName = "GCN Burst Latent", ShortTooltip = "A one-off GameplayCueNotify that can use latent actions such as timelines."))
+UCLASS(Blueprintable, notplaceable, SparseClassDataTypes = GameplayCueNotify_BurstLatentSparseData, Category = "GameplayCueNotify", Meta = (ShowWorldContextPin, DisplayName = "GCN Burst Latent", ShortTooltip = "A one-off GameplayCueNotify that can use latent actions such as timelines."))
 class AGameplayCueNotify_BurstLatent : public AGameplayCueNotify_Instanced
 {
 	GENERATED_BODY()
@@ -38,19 +63,7 @@ protected:
 
 protected:
 
-	// Default condition to check before spawning anything.  Applies for all spawns unless overridden.
-	UPROPERTY(EditAnywhere, BlueprintReadOnly, Category = "Defaults")
-	FGameplayCueNotify_SpawnCondition DefaultSpawnCondition;
-
-	// Default placement rules.  Applies for all spawns unless overridden.
-	UPROPERTY(EditAnywhere, BlueprintReadOnly, Category = "Defaults")
-	FGameplayCueNotify_PlacementInfo DefaultPlacementInfo;
-
-	// List of effects to spawn on burst.
-	UPROPERTY(EditAnywhere, BlueprintReadOnly, Category = "Effects")
-	FGameplayCueNotify_BurstEffects BurstEffects;
-
 	// Results of spawned burst effects.
-	UPROPERTY(BlueprintReadOnly, Category = "Effects")
+	UPROPERTY(BlueprintReadOnly, Category = "GCN Effects")
 	FGameplayCueNotify_SpawnResult BurstSpawnResults;
 };
diff --git a/Engine/Plugins/Runtime/GameplayAbilities/Source/GameplayAbilities/Public/GameplayCueNotify_Looping.h b/Engine/Plugins/Runtime/GameplayAbilities/Source/GameplayAbilities/Public/GameplayCueNotify_Looping.h
--- a/Engine/Plugins/Runtime/GameplayAbilities/Source/GameplayAbilities/Public/GameplayCueNotify_Looping.h
+++ b/Engine/Plugins/Runtime/GameplayAbilities/Source/GameplayAbilities/Public/GameplayCueNotify_Looping.h
@@ -8,13 +8,50 @@
 #include "GameplayCueNotify_Looping.generated.h"
 
 
+USTRUCT(BlueprintType)
+struct FGameplayCueNotify_LoopingSparseData
+{
+	GENERATED_BODY()
+
+public:
+
+	FGameplayCueNotify_LoopingSparseData() { DefaultPlacementInfo.AttachPolicy = EGameplayCueNotify_AttachPolicy::AttachToTarget; }
+
+public:
+
+	// Default condition to check before spawning anything.  Applies for all spawns unless overridden.
+	UPROPERTY(EditDefaultsOnly, BlueprintReadOnly, Category = "GCN Defaults")
+	FGameplayCueNotify_SpawnCondition DefaultSpawnCondition;
+
+	// Default placement rules.  Applies for all spawns unless overridden.
+	UPROPERTY(EditDefaultsOnly, BlueprintReadOnly, Category = "GCN Defaults")
+	FGameplayCueNotify_PlacementInfo DefaultPlacementInfo;
+
+	// List of effects to spawn on application.  These should not be looping effects!
+	UPROPERTY(EditDefaultsOnly, BlueprintReadOnly, Category = "GCN Application Effects (On Active)")
+	FGameplayCueNotify_BurstEffects ApplicationEffects;
+
+	// List of effects to spawn on loop start.
+	UPROPERTY(EditDefaultsOnly, BlueprintReadOnly, Category = "GCN Looping Effects (While Active)")
+	FGameplayCueNotify_LoopingEffects LoopingEffects;
+
+	// List of effects to spawn for a recurring gameplay effect (e.g. each time a DOT ticks).  These should not be looping effects!
+	UPROPERTY(EditDefaultsOnly, BlueprintReadOnly, Category = "GCN Recurring Effects (On Execute)")
+	FGameplayCueNotify_BurstEffects RecurringEffects;
+
+	// List of effects to spawn on removal.  These should not be looping effects!
+	UPROPERTY(EditDefaultsOnly, BlueprintReadOnly, Category = "GCN Removal Effects (On Remove)")
+	FGameplayCueNotify_BurstEffects RemovalEffects;
+};
+
+
 /**
  * AGameplayCueNotify_Looping
  *
  *	This is an instanced gameplay cue notify for continuous looping effects.
  *	The game is responsible for defining the start/stop by adding/removing the gameplay cue.
  */
-UCLASS(Blueprintable, notplaceable, Category = "GameplayCueNotify", Meta = (ShowWorldContextPin, DisplayName = "GCN Looping", ShortTooltip = "A GameplayCueNotify that has a duration that is driven by the game."))
+UCLASS(Blueprintable, notplaceable, SparseClassDataTypes = GameplayCueNotify_LoopingSparseData, Category = "GameplayCueNotify", Meta = (ShowWorldContextPin, DisplayName = "GCN Looping", ShortTooltip = "A GameplayCueNotify that has a duration that is driven by the game."))
 class AGameplayCueNotify_Looping : public AGameplayCueNotify_Instanced
 {
 	GENERATED_BODY()
@@ -53,44 +90,20 @@ protected:
 
 protected:
 
-	// Default condition to check before spawning anything.  Applies for all spawns unless overridden.
-	UPROPERTY(EditAnywhere, BlueprintReadOnly, Category = "Defaults")
-	FGameplayCueNotify_SpawnCondition DefaultSpawnCondition;
-
-	// Default placement rules.  Applies for all spawns unless overridden.
-	UPROPERTY(EditAnywhere, BlueprintReadOnly, Category = "Defaults")
-	FGameplayCueNotify_PlacementInfo DefaultPlacementInfo;
-
-	// List of effects to spawn on application.  These should not be looping effects!
-	UPROPERTY(EditAnywhere, BlueprintReadOnly, Category = "Application Effects (On Active)")
-	FGameplayCueNotify_BurstEffects ApplicationEffects;
-
 	// Results of spawned application effects.
-	UPROPERTY(BlueprintReadOnly, Category = "Application Effects (On Active)")
+	UPROPERTY(BlueprintReadOnly, Category = "GCN Application Effects (On Active)")
 	FGameplayCueNotify_SpawnResult ApplicationSpawnResults;
 
-	// List of effects to spawn on loop start.
-	UPROPERTY(EditAnywhere, BlueprintReadOnly, Category = "Looping Effects (While Active)")
-	FGameplayCueNotify_LoopingEffects LoopingEffects;
-
 	// Results of spawned looping effects.
-	UPROPERTY(BlueprintReadOnly, Category = "Looping Effects (While Active)")
+	UPROPERTY(BlueprintReadOnly, Category = "GCN Looping Effects (While Active)")
 	FGameplayCueNotify_SpawnResult LoopingSpawnResults;
 
-	// List of effects to spawn for a recurring gameplay effect (e.g. each time a DOT ticks).  These should not be looping effects!
-	UPROPERTY(EditAnywhere, BlueprintReadOnly, Category = "Recurring Effects (On Execute)")
-	FGameplayCueNotify_BurstEffects RecurringEffects;
-
 	// Results of spawned recurring effects.
-	UPROPERTY(BlueprintReadOnly, Category = "Recurring Effects (On Execute)")
+	UPROPERTY(BlueprintReadOnly, Category = "GCN Recurring Effects (On Execute)")
 	FGameplayCueNotify_SpawnResult RecurringSpawnResults;
 
-	// List of effects to spawn on removal.  These should not be looping effects!
-	UPROPERTY(EditAnywhere, BlueprintReadOnly, Category = "Removal Effects (On Remove)")
-	FGameplayCueNotify_BurstEffects RemovalEffects;
-
 	// Results of spawned removal effects.
-	UPROPERTY(BlueprintReadOnly, Category = "Removal Effects (On Remove)")
+	UPROPERTY(BlueprintReadOnly, Category = "GCN Removal Effects (On Remove)")
 	FGameplayCueNotify_SpawnResult RemovalSpawnResults;
 
 	bool bLoopingEffectsRemoved;
