Reverted BurstLatent and Looping gameplay cue notifies to not use Sparse Class Data.  There is a bug with setting asset properties related to Sparse Class Data.

diff --git a/Engine/Plugins/Runtime/GameplayAbilities/Source/GameplayAbilities/Private/GameplayCueNotify_BurstLatent.cpp b/Engine/Plugins/Runtime/GameplayAbilities/Source/GameplayAbilities/Private/GameplayCueNotify_BurstLatent.cpp
--- a/Engine/Plugins/Runtime/GameplayAbilities/Source/GameplayAbilities/Private/GameplayCueNotify_BurstLatent.cpp
+++ b/Engine/Plugins/Runtime/GameplayAbilities/Source/GameplayAbilities/Private/GameplayCueNotify_BurstLatent.cpp
@@ -29,16 +29,13 @@ bool AGameplayCueNotify_BurstLatent::OnExecute_Implementation(AActor* Target, co
 {
 	UWorld* World = GetWorld();
 
-	const FGameplayCueNotify_BurstLatentSparseData* SparseData = GetGameplayCueNotify_BurstLatentSparseData();
-	check(SparseData);
-
 	FGameplayCueNotify_SpawnContext SpawnContext(World, Target, Parameters);
-	SpawnContext.SetDefaultSpawnCondition(&SparseData->DefaultSpawnCondition);
-	SpawnContext.SetDefaultPlacementInfo(&SparseData->DefaultPlacementInfo);
+	SpawnContext.SetDefaultSpawnCondition(&DefaultSpawnCondition);
+	SpawnContext.SetDefaultPlacementInfo(&DefaultPlacementInfo);
 
-	if (SparseData->DefaultSpawnCondition.ShouldSpawn(SpawnContext))
+	if (DefaultSpawnCondition.ShouldSpawn(SpawnContext))
 	{
-		SparseData->BurstEffects.ExecuteEffects(SpawnContext, BurstSpawnResults);
+		BurstEffects.ExecuteEffects(SpawnContext, BurstSpawnResults);
 
 		OnBurst(Target, Parameters, BurstSpawnResults);
 	}
@@ -58,10 +55,7 @@ bool AGameplayCueNotify_BurstLatent::OnExecute_Implementation(AActor* Target, co
 #if WITH_EDITOR
 EDataValidationResult AGameplayCueNotify_BurstLatent::IsDataValid(TArray<FText>& ValidationErrors)
 {
-	const FGameplayCueNotify_BurstLatentSparseData* SparseData = GetGameplayCueNotify_BurstLatentSparseData();
-	check(SparseData);
-
-	SparseData->BurstEffects.ValidateAssociatedAssets(this, TEXT("BurstEffects"), ValidationErrors);
+	BurstEffects.ValidateAssociatedAssets(this, TEXT("BurstEffects"), ValidationErrors);
 
 	return ((ValidationErrors.Num() > 0) ? EDataValidationResult::Invalid : EDataValidationResult::Valid);
 }
diff --git a/Engine/Plugins/Runtime/GameplayAbilities/Source/GameplayAbilities/Private/GameplayCueNotify_Looping.cpp b/Engine/Plugins/Runtime/GameplayAbilities/Source/GameplayAbilities/Private/GameplayCueNotify_Looping.cpp
--- a/Engine/Plugins/Runtime/GameplayAbilities/Source/GameplayAbilities/Private/GameplayCueNotify_Looping.cpp
+++ b/Engine/Plugins/Runtime/GameplayAbilities/Source/GameplayAbilities/Private/GameplayCueNotify_Looping.cpp
@@ -10,6 +10,8 @@ AGameplayCueNotify_Looping::AGameplayCueNotify_Looping()
 {
 	bAllowMultipleWhileActiveEvents = false;
 
+	DefaultPlacementInfo.AttachPolicy = EGameplayCueNotify_AttachPolicy::AttachToTarget;
+
 	Recycle();
 }
 
@@ -44,16 +46,13 @@ bool AGameplayCueNotify_Looping::OnActive_Implementation(AActor* Target, const F
 {
 	UWorld* World = GetWorld();
 
-	const FGameplayCueNotify_LoopingSparseData* SparseData = GetGameplayCueNotify_LoopingSparseData();
-	check(SparseData);
-
 	FGameplayCueNotify_SpawnContext SpawnContext(World, Target, Parameters);
-	SpawnContext.SetDefaultSpawnCondition(&SparseData->DefaultSpawnCondition);
-	SpawnContext.SetDefaultPlacementInfo(&SparseData->DefaultPlacementInfo);
+	SpawnContext.SetDefaultSpawnCondition(&DefaultSpawnCondition);
+	SpawnContext.SetDefaultPlacementInfo(&DefaultPlacementInfo);
 
-	if (SparseData->DefaultSpawnCondition.ShouldSpawn(SpawnContext))
+	if (DefaultSpawnCondition.ShouldSpawn(SpawnContext))
 	{
-		SparseData->ApplicationEffects.ExecuteEffects(SpawnContext, ApplicationSpawnResults);
+		ApplicationEffects.ExecuteEffects(SpawnContext, ApplicationSpawnResults);
 
 		OnApplication(Target, Parameters, ApplicationSpawnResults);
 	}
@@ -65,18 +64,15 @@ bool AGameplayCueNotify_Looping::WhileActive_Implementation(AActor* Target, cons
 {
 	UWorld* World = GetWorld();
 
-	const FGameplayCueNotify_LoopingSparseData* SparseData = GetGameplayCueNotify_LoopingSparseData();
-	check(SparseData);
-
 	FGameplayCueNotify_SpawnContext SpawnContext(World, Target, Parameters);
-	SpawnContext.SetDefaultSpawnCondition(&SparseData->DefaultSpawnCondition);
-	SpawnContext.SetDefaultPlacementInfo(&SparseData->DefaultPlacementInfo);
+	SpawnContext.SetDefaultSpawnCondition(&DefaultSpawnCondition);
+	SpawnContext.SetDefaultPlacementInfo(&DefaultPlacementInfo);
 
-	if (SparseData->DefaultSpawnCondition.ShouldSpawn(SpawnContext))
+	if (DefaultSpawnCondition.ShouldSpawn(SpawnContext))
 	{
 		bLoopingEffectsRemoved = false;
 
-		SparseData->LoopingEffects.StartEffects(SpawnContext, LoopingSpawnResults);	
+		LoopingEffects.StartEffects(SpawnContext, LoopingSpawnResults);	
 
 		OnLoopingStart(Target, Parameters, LoopingSpawnResults);
 	}
@@ -88,16 +84,13 @@ bool AGameplayCueNotify_Looping::OnExecute_Implementation(AActor* Target, const
 {
 	UWorld* World = GetWorld();
 
-	const FGameplayCueNotify_LoopingSparseData* SparseData = GetGameplayCueNotify_LoopingSparseData();
-	check(SparseData);
-
 	FGameplayCueNotify_SpawnContext SpawnContext(World, Target, Parameters);
-	SpawnContext.SetDefaultSpawnCondition(&SparseData->DefaultSpawnCondition);
-	SpawnContext.SetDefaultPlacementInfo(&SparseData->DefaultPlacementInfo);
+	SpawnContext.SetDefaultSpawnCondition(&DefaultSpawnCondition);
+	SpawnContext.SetDefaultPlacementInfo(&DefaultPlacementInfo);
 
-	if (SparseData->DefaultSpawnCondition.ShouldSpawn(SpawnContext))
+	if (DefaultSpawnCondition.ShouldSpawn(SpawnContext))
 	{
-		SparseData->RecurringEffects.ExecuteEffects(SpawnContext, RecurringSpawnResults);
+		RecurringEffects.ExecuteEffects(SpawnContext, RecurringSpawnResults);
 
 		OnRecurring(Target, Parameters, RecurringSpawnResults);
 	}
@@ -114,16 +107,13 @@ bool AGameplayCueNotify_Looping::OnRemove_Implementation(AActor* Target, const F
 	{
 		UWorld* World = GetWorld();
 
-		const FGameplayCueNotify_LoopingSparseData* SparseData = GetGameplayCueNotify_LoopingSparseData();
-		check(SparseData);
-
 		FGameplayCueNotify_SpawnContext SpawnContext(World, Target, Parameters);
-		SpawnContext.SetDefaultSpawnCondition(&SparseData->DefaultSpawnCondition);
-		SpawnContext.SetDefaultPlacementInfo(&SparseData->DefaultPlacementInfo);
+		SpawnContext.SetDefaultSpawnCondition(&DefaultSpawnCondition);
+		SpawnContext.SetDefaultPlacementInfo(&DefaultPlacementInfo);
 
-		if (SparseData->DefaultSpawnCondition.ShouldSpawn(SpawnContext))
+		if (DefaultSpawnCondition.ShouldSpawn(SpawnContext))
 		{
-			SparseData->RemovalEffects.ExecuteEffects(SpawnContext, RemovalSpawnResults);
+			RemovalEffects.ExecuteEffects(SpawnContext, RemovalSpawnResults);
 		}
 	}
 
@@ -140,24 +130,18 @@ void AGameplayCueNotify_Looping::RemoveLoopingEffects()
 		return;
 	}
 
-	const FGameplayCueNotify_LoopingSparseData* SparseData = GetGameplayCueNotify_LoopingSparseData();
-	check(SparseData);
-
 	bLoopingEffectsRemoved = true;
 
-	SparseData->LoopingEffects.StopEffects(LoopingSpawnResults);
+	LoopingEffects.StopEffects(LoopingSpawnResults);
 }
 
 #if WITH_EDITOR
 EDataValidationResult AGameplayCueNotify_Looping::IsDataValid(TArray<FText>& ValidationErrors)
 {
-	const FGameplayCueNotify_LoopingSparseData* SparseData = GetGameplayCueNotify_LoopingSparseData();
-	check(SparseData);
-
-	SparseData->ApplicationEffects.ValidateAssociatedAssets(this, TEXT("ApplicationEffects"), ValidationErrors);
-	SparseData->LoopingEffects.ValidateAssociatedAssets(this, TEXT("LoopingEffects"), ValidationErrors);
-	SparseData->RecurringEffects.ValidateAssociatedAssets(this, TEXT("RecurringEffects"), ValidationErrors);
-	SparseData->RemovalEffects.ValidateAssociatedAssets(this, TEXT("RemovalEffects"), ValidationErrors);
+	ApplicationEffects.ValidateAssociatedAssets(this, TEXT("ApplicationEffects"), ValidationErrors);
+	LoopingEffects.ValidateAssociatedAssets(this, TEXT("LoopingEffects"), ValidationErrors);
+	RecurringEffects.ValidateAssociatedAssets(this, TEXT("RecurringEffects"), ValidationErrors);
+	RemovalEffects.ValidateAssociatedAssets(this, TEXT("RemovalEffects"), ValidationErrors);
 
 	return ((ValidationErrors.Num() > 0) ? EDataValidationResult::Invalid : EDataValidationResult::Valid);
 }
diff --git a/Engine/Plugins/Runtime/GameplayAbilities/Source/GameplayAbilities/Public/GameplayCueNotify_BurstLatent.h b/Engine/Plugins/Runtime/GameplayAbilities/Source/GameplayAbilities/Public/GameplayCueNotify_BurstLatent.h
--- a/Engine/Plugins/Runtime/GameplayAbilities/Source/GameplayAbilities/Public/GameplayCueNotify_BurstLatent.h
+++ b/Engine/Plugins/Runtime/GameplayAbilities/Source/GameplayAbilities/Public/GameplayCueNotify_BurstLatent.h
@@ -8,38 +8,13 @@
 #include "GameplayCueNotify_BurstLatent.generated.h"
 
 
-USTRUCT(BlueprintType)
-struct FGameplayCueNotify_BurstLatentSparseData
-{
-	GENERATED_BODY()
-
-public:
-
-	FGameplayCueNotify_BurstLatentSparseData() {}
-
-public:
-
-	// Default condition to check before spawning anything.  Applies for all spawns unless overridden.
-	UPROPERTY(EditDefaultsOnly, BlueprintReadOnly, Category = "GCN Defaults")
-	FGameplayCueNotify_SpawnCondition DefaultSpawnCondition;
-
-	// Default placement rules.  Applies for all spawns unless overridden.
-	UPROPERTY(EditDefaultsOnly, BlueprintReadOnly, Category = "GCN Defaults")
-	FGameplayCueNotify_PlacementInfo DefaultPlacementInfo;
-
-	// List of effects to spawn on burst.
-	UPROPERTY(EditDefaultsOnly, BlueprintReadOnly, Category = "GCN Effects")
-	FGameplayCueNotify_BurstEffects BurstEffects;
-};
-
-
 /**
  * AGameplayCueNotify_BurstLatent
  *
  *	This is an instanced gameplay cue notify for effects that are one-offs.
  *	Since it is instanced, it can do latent things like time lines or delays.
  */
-UCLASS(Blueprintable, notplaceable, SparseClassDataTypes = GameplayCueNotify_BurstLatentSparseData, Category = "GameplayCueNotify", Meta = (ShowWorldContextPin, DisplayName = "GCN Burst Latent", ShortTooltip = "A one-off GameplayCueNotify that can use latent actions such as timelines."))
+UCLASS(Blueprintable, notplaceable, Category = "GameplayCueNotify", Meta = (ShowWorldContextPin, DisplayName = "GCN Burst Latent", ShortTooltip = "A one-off GameplayCueNotify that can use latent actions such as timelines."))
 class AGameplayCueNotify_BurstLatent : public AGameplayCueNotify_Instanced
 {
 	GENERATED_BODY()
@@ -63,6 +38,18 @@ protected:
 
 protected:
 
+	// Default condition to check before spawning anything.  Applies for all spawns unless overridden.
+	UPROPERTY(EditDefaultsOnly, BlueprintReadOnly, Category = "GCN Defaults")
+	FGameplayCueNotify_SpawnCondition DefaultSpawnCondition;
+
+	// Default placement rules.  Applies for all spawns unless overridden.
+	UPROPERTY(EditDefaultsOnly, BlueprintReadOnly, Category = "GCN Defaults")
+	FGameplayCueNotify_PlacementInfo DefaultPlacementInfo;
+
+	// List of effects to spawn on burst.
+	UPROPERTY(EditDefaultsOnly, BlueprintReadOnly, Category = "GCN Effects")
+	FGameplayCueNotify_BurstEffects BurstEffects;
+
 	// Results of spawned burst effects.
 	UPROPERTY(BlueprintReadOnly, Category = "GCN Effects")
 	FGameplayCueNotify_SpawnResult BurstSpawnResults;
diff --git a/Engine/Plugins/Runtime/GameplayAbilities/Source/GameplayAbilities/Public/GameplayCueNotify_Looping.h b/Engine/Plugins/Runtime/GameplayAbilities/Source/GameplayAbilities/Public/GameplayCueNotify_Looping.h
--- a/Engine/Plugins/Runtime/GameplayAbilities/Source/GameplayAbilities/Public/GameplayCueNotify_Looping.h
+++ b/Engine/Plugins/Runtime/GameplayAbilities/Source/GameplayAbilities/Public/GameplayCueNotify_Looping.h
@@ -8,50 +8,13 @@
 #include "GameplayCueNotify_Looping.generated.h"
 
 
-USTRUCT(BlueprintType)
-struct FGameplayCueNotify_LoopingSparseData
-{
-	GENERATED_BODY()
-
-public:
-
-	FGameplayCueNotify_LoopingSparseData() { DefaultPlacementInfo.AttachPolicy = EGameplayCueNotify_AttachPolicy::AttachToTarget; }
-
-public:
-
-	// Default condition to check before spawning anything.  Applies for all spawns unless overridden.
-	UPROPERTY(EditDefaultsOnly, BlueprintReadOnly, Category = "GCN Defaults")
-	FGameplayCueNotify_SpawnCondition DefaultSpawnCondition;
-
-	// Default placement rules.  Applies for all spawns unless overridden.
-	UPROPERTY(EditDefaultsOnly, BlueprintReadOnly, Category = "GCN Defaults")
-	FGameplayCueNotify_PlacementInfo DefaultPlacementInfo;
-
-	// List of effects to spawn on application.  These should not be looping effects!
-	UPROPERTY(EditDefaultsOnly, BlueprintReadOnly, Category = "GCN Application Effects (On Active)")
-	FGameplayCueNotify_BurstEffects ApplicationEffects;
-
-	// List of effects to spawn on loop start.
-	UPROPERTY(EditDefaultsOnly, BlueprintReadOnly, Category = "GCN Looping Effects (While Active)")
-	FGameplayCueNotify_LoopingEffects LoopingEffects;
-
-	// List of effects to spawn for a recurring gameplay effect (e.g. each time a DOT ticks).  These should not be looping effects!
-	UPROPERTY(EditDefaultsOnly, BlueprintReadOnly, Category = "GCN Recurring Effects (On Execute)")
-	FGameplayCueNotify_BurstEffects RecurringEffects;
-
-	// List of effects to spawn on removal.  These should not be looping effects!
-	UPROPERTY(EditDefaultsOnly, BlueprintReadOnly, Category = "GCN Removal Effects (On Remove)")
-	FGameplayCueNotify_BurstEffects RemovalEffects;
-};
-
-
 /**
  * AGameplayCueNotify_Looping
  *
  *	This is an instanced gameplay cue notify for continuous looping effects.
  *	The game is responsible for defining the start/stop by adding/removing the gameplay cue.
  */
-UCLASS(Blueprintable, notplaceable, SparseClassDataTypes = GameplayCueNotify_LoopingSparseData, Category = "GameplayCueNotify", Meta = (ShowWorldContextPin, DisplayName = "GCN Looping", ShortTooltip = "A GameplayCueNotify that has a duration that is driven by the game."))
+UCLASS(Blueprintable, notplaceable, Category = "GameplayCueNotify", Meta = (ShowWorldContextPin, DisplayName = "GCN Looping", ShortTooltip = "A GameplayCueNotify that has a duration that is driven by the game."))
 class AGameplayCueNotify_Looping : public AGameplayCueNotify_Instanced
 {
 	GENERATED_BODY()
@@ -90,18 +53,42 @@ protected:
 
 protected:
 
+	// Default condition to check before spawning anything.  Applies for all spawns unless overridden.
+	UPROPERTY(EditDefaultsOnly, BlueprintReadOnly, Category = "GCN Defaults")
+	FGameplayCueNotify_SpawnCondition DefaultSpawnCondition;
+
+	// Default placement rules.  Applies for all spawns unless overridden.
+	UPROPERTY(EditDefaultsOnly, BlueprintReadOnly, Category = "GCN Defaults")
+	FGameplayCueNotify_PlacementInfo DefaultPlacementInfo;
+
+	// List of effects to spawn on application.  These should not be looping effects!
+	UPROPERTY(EditDefaultsOnly, BlueprintReadOnly, Category = "GCN Application Effects (On Active)")
+	FGameplayCueNotify_BurstEffects ApplicationEffects;
+
 	// Results of spawned application effects.
 	UPROPERTY(BlueprintReadOnly, Category = "GCN Application Effects (On Active)")
 	FGameplayCueNotify_SpawnResult ApplicationSpawnResults;
 
+	// List of effects to spawn on loop start.
+	UPROPERTY(EditDefaultsOnly, BlueprintReadOnly, Category = "GCN Looping Effects (While Active)")
+	FGameplayCueNotify_LoopingEffects LoopingEffects;
+
 	// Results of spawned looping effects.
 	UPROPERTY(BlueprintReadOnly, Category = "GCN Looping Effects (While Active)")
 	FGameplayCueNotify_SpawnResult LoopingSpawnResults;
 
+	// List of effects to spawn for a recurring gameplay effect (e.g. each time a DOT ticks).  These should not be looping effects!
+	UPROPERTY(EditDefaultsOnly, BlueprintReadOnly, Category = "GCN Recurring Effects (On Execute)")
+	FGameplayCueNotify_BurstEffects RecurringEffects;
+
 	// Results of spawned recurring effects.
 	UPROPERTY(BlueprintReadOnly, Category = "GCN Recurring Effects (On Execute)")
 	FGameplayCueNotify_SpawnResult RecurringSpawnResults;
 
+	// List of effects to spawn on removal.  These should not be looping effects!
+	UPROPERTY(EditDefaultsOnly, BlueprintReadOnly, Category = "GCN Removal Effects (On Remove)")
+	FGameplayCueNotify_BurstEffects RemovalEffects;
+
 	// Results of spawned removal effects.
 	UPROPERTY(BlueprintReadOnly, Category = "GCN Removal Effects (On Remove)")
 	FGameplayCueNotify_SpawnResult RemovalSpawnResults;
