GameplayCue function library;

* Add a blueprint function library for gameplay cues with the following functions:

    // Builds gameplay cue parameters using data from a hit result.
    MakeGameplayCueParametersFromHitResult

    // Invoke a one time "instant" execute event for a gameplay cue on the target actor.
    ExecuteGameplayCueOnActor

    // Invoke the added event for a gameplay cue on the target actor.
    AddGameplayCueOnActor

    // Invoke the removed event for a gameplay cue on the target actor.
    RemoveGameplayCueOnActor

diff --git a/Engine/Plugins/Runtime/GameplayAbilities/Source/GameplayAbilities/Private/GameplayCueFunctionLibrary.cpp b/Engine/Plugins/Runtime/GameplayAbilities/Source/GameplayAbilities/Private/GameplayCueFunctionLibrary.cpp
new file mode 100644
--- /dev/null
+++ b/Engine/Plugins/Runtime/GameplayAbilities/Source/GameplayAbilities/Private/GameplayCueFunctionLibrary.cpp
@@ -0,0 +1,96 @@
+// Copyright Epic Games, Inc. All Rights Reserved.
+
+#include "GameplayCueFunctionLibrary.h"
+#include "GameplayCueManager.h"
+#include "AbilitySystemInterface.h"
+#include "AbilitySystemComponent.h"
+
+
+//////////////////////////////////////////////////////////////////////////
+// UGameplayCueFunctionLibrary
+//////////////////////////////////////////////////////////////////////////
+FGameplayCueParameters UGameplayCueFunctionLibrary::MakeGameplayCueParametersFromHitResult(const FHitResult& HitResult)
+{
+	FGameplayCueParameters CueParameters;
+
+	CueParameters.Location = (HitResult.bBlockingHit ? HitResult.ImpactPoint : HitResult.TraceEnd);
+	CueParameters.Normal = HitResult.ImpactNormal;
+	CueParameters.PhysicalMaterial = HitResult.PhysMaterial;
+
+	return CueParameters;
+}
+
+void UGameplayCueFunctionLibrary::ExecuteGameplayCueOnActor(AActor* Target, const FGameplayTag GameplayCueTag, const FGameplayCueParameters& Parameters)
+{
+	if (!Target)
+	{
+		return;
+	}
+
+	const IAbilitySystemInterface* ASI = Cast<IAbilitySystemInterface>(Target);
+	UAbilitySystemComponent *ASC = (ASI ? ASI->GetAbilitySystemComponent() : nullptr);
+
+	if (ASC)
+	{
+		// The actor has an ability system so the event will fire on authority only and will be replicated.
+		if (Target->GetLocalRole() == ROLE_Authority)
+		{
+			ASC->ExecuteGameplayCue(GameplayCueTag, Parameters);
+		}
+	}
+	else
+	{
+		// The actor does not have an ability system so the event will only be fired locally.
+		UGameplayCueManager::ExecuteGameplayCue_NonReplicated(Target, GameplayCueTag, Parameters);
+	}
+}
+
+void UGameplayCueFunctionLibrary::AddGameplayCueOnActor(AActor* Target, const FGameplayTag GameplayCueTag, const FGameplayCueParameters& Parameters)
+{
+	if (!Target)
+	{
+		return;
+	}
+
+	const IAbilitySystemInterface* ASI = Cast<IAbilitySystemInterface>(Target);
+	UAbilitySystemComponent *ASC = (ASI ? ASI->GetAbilitySystemComponent() : nullptr);
+
+	if (ASC)
+	{
+		// The actor has an ability system so the event will fire on authority only and will be replicated.
+		if (Target->GetLocalRole() == ROLE_Authority)
+		{
+			ASC->AddGameplayCue(GameplayCueTag, Parameters);
+		}
+	}
+	else
+	{
+		// The actor does not have an ability system so the event will only be fired locally.
+		UGameplayCueManager::AddGameplayCue_NonReplicated(Target, GameplayCueTag, Parameters);
+	}
+}
+
+void UGameplayCueFunctionLibrary::RemoveGameplayCueOnActor(AActor* Target, const FGameplayTag GameplayCueTag, const FGameplayCueParameters& Parameters)
+{
+	if (!Target)
+	{
+		return;
+	}
+
+	const IAbilitySystemInterface* ASI = Cast<IAbilitySystemInterface>(Target);
+	UAbilitySystemComponent *ASC = (ASI ? ASI->GetAbilitySystemComponent() : nullptr);
+
+	if (ASC)
+	{
+		// The actor has an ability system so the event will fire on authority only and will be replicated.
+		if (Target->GetLocalRole() == ROLE_Authority)
+		{
+			ASC->RemoveGameplayCue(GameplayCueTag);
+		}
+	}
+	else
+	{
+		// The actor does not have an ability system so the event will only be fired locally.
+		UGameplayCueManager::RemoveGameplayCue_NonReplicated(Target, GameplayCueTag, Parameters);
+	}
+}
diff --git a/Engine/Plugins/Runtime/GameplayAbilities/Source/GameplayAbilities/Public/GameplayCueFunctionLibrary.h b/Engine/Plugins/Runtime/GameplayAbilities/Source/GameplayAbilities/Public/GameplayCueFunctionLibrary.h
new file mode 100644
--- /dev/null
+++ b/Engine/Plugins/Runtime/GameplayAbilities/Source/GameplayAbilities/Public/GameplayCueFunctionLibrary.h
@@ -0,0 +1,48 @@
+// Copyright Epic Games, Inc. All Rights Reserved.
+
+#pragma once
+
+#include "CoreMinimal.h"
+#include "Kismet/BlueprintFunctionLibrary.h"
+#include "GameplayEffectTypes.h"
+#include "GameplayTagContainer.h"
+#include "GameplayCueFunctionLibrary.generated.h"
+
+
+struct FHitResult;
+
+
+/**
+ * UGameplayCueFunctionLibrary
+ *
+ *	Helpful utility function for working with gameplay cues.
+ */
+UCLASS()
+class UGameplayCueFunctionLibrary : public UBlueprintFunctionLibrary
+{
+	GENERATED_BODY()
+
+public:
+
+	// Builds gameplay cue parameters using data from a hit result.
+	UFUNCTION(BlueprintPure, Category = "GameplayCue")
+	static FGameplayCueParameters MakeGameplayCueParametersFromHitResult(const FHitResult& HitResult);
+
+	// Invoke a one time "instant" execute event for a gameplay cue on the target actor.
+	// * If the actor has an ability system, the event will fire on authority only and will be replicated.
+	// * If the actor does not have an ability system, the event will only be fired locally.
+	UFUNCTION(BlueprintCallable, Category = "GameplayCue", Meta = (GameplayTagFilter = "GameplayCue"), DisplayName = "Execute GameplayCue On Actor (Burst)")
+	static void ExecuteGameplayCueOnActor(AActor* Target, const FGameplayTag GameplayCueTag, const FGameplayCueParameters& Parameters);
+
+	// Invoke the added event for a gameplay cue on the target actor. This should be paired with a RemoveGameplayCueOnActor call.
+	// * If the actor has an ability system, the event will fire on authority only and will be replicated.
+	// * If the actor does not have an ability system, the event will only be fired locally.
+	UFUNCTION(BlueprintCallable, Category = "GameplayCue", Meta = (GameplayTagFilter = "GameplayCue"), DisplayName = "Add GameplayCue On Actor (Looping)")
+	static void AddGameplayCueOnActor(AActor* Target, const FGameplayTag GameplayCueTag, const FGameplayCueParameters& Parameters);
+
+	// Invoke the removed event for a gameplay cue on the target actor. This should be paired with an AddGameplayCueOnActor call.
+	// * If the actor has an ability system, the event will fire on authority only and will be replicated.
+	// * If the actor does not have an ability system, the event will only be fired locally.
+	UFUNCTION(BlueprintCallable, Category = "GameplayCue", Meta = (GameplayTagFilter = "GameplayCue"), DisplayName = "Remove GameplayCue On Actor (Looping)")
+	static void RemoveGameplayCueOnActor(AActor* Target, const FGameplayTag GameplayCueTag, const FGameplayCueParameters& Parameters);
+};
