Removed the GameplayCueNotify_Instanced class.  It's functionality has been incorporated into the parent class GameplayCueNotify_Actor.

diff --git a/Engine/Plugins/Runtime/GameplayAbilities/Source/GameplayAbilities/Private/GameplayCueNotify_Actor.cpp b/Engine/Plugins/Runtime/GameplayAbilities/Source/GameplayAbilities/Private/GameplayCueNotify_Actor.cpp
--- a/Engine/Plugins/Runtime/GameplayAbilities/Source/GameplayAbilities/Private/GameplayCueNotify_Actor.cpp
+++ b/Engine/Plugins/Runtime/GameplayAbilities/Source/GameplayAbilities/Private/GameplayCueNotify_Actor.cpp
@@ -7,6 +7,7 @@
 #include "AbilitySystemStats.h"
 #include "AbilitySystemGlobals.h"
 #include "GameplayCueManager.h"
+#include "GameplayCueNotifyTypes.h"
 
 
 namespace FAbilitySystemTweaks
@@ -43,6 +44,15 @@ AGameplayCueNotify_Actor::AGameplayCueNotify_Actor(const FObjectInitializer& Obj
 
 void AGameplayCueNotify_Actor::EndPlay(const EEndPlayReason::Type EndPlayReason)
 {
+	if (UGameplayCueManager::IsGameplayCueRecylingEnabled() && (EndPlayReason == EEndPlayReason::Destroyed))
+	{
+		UWorld* World = GetWorld();
+		if (World && !World->IsPlayingReplay() && (World->WorldType == EWorldType::Game))
+		{
+			UE_LOG(LogGameplayCueNotify, Error, TEXT("GameplayCueNotify [%s] is calling EndPlay(). This should not happen since they are recycled through the GameplayCueManager."), *GetName());
+		}
+	}
+
 	if (EndPlayReason == EEndPlayReason::Destroyed)
 	{
 		UAbilitySystemGlobals::Get().GetGameplayCueManager()->NotifyGameplayCueActorEndPlay(this);
@@ -51,6 +61,32 @@ void AGameplayCueNotify_Actor::EndPlay(const EEndPlayReason::Type EndPlayReason)
 	Super::EndPlay(EndPlayReason);
 }
 
+void AGameplayCueNotify_Actor::K2_DestroyActor()
+{
+	if (UGameplayCueManager::IsGameplayCueRecylingEnabled())
+	{
+		UE_LOG(LogGameplayCueNotify, Warning, TEXT("GameplayCueNotify [%s] is calling DesotryActor(). This is not necessary as GCs will be cleaned up and recycled automatically."), *GetName());
+	}
+	else
+	{
+		Super::K2_DestroyActor();
+	}
+}
+
+void AGameplayCueNotify_Actor::Destroyed()
+{
+	if (UGameplayCueManager::IsGameplayCueRecylingEnabled())
+	{
+		UWorld* World = GetWorld();
+		if (World && !World->IsPlayingReplay() && (World->WorldType == EWorldType::Game))
+		{
+			UE_LOG(LogGameplayCueNotify, Error, TEXT("GameplayCueNotify [%s] is calling Destroyed(). This should not happen since they are recycled through the GameplayCueManager."), *GetName());
+		}
+	}
+
+	Super::Destroyed();
+}
+
 #if WITH_EDITOR
 void AGameplayCueNotify_Actor::PostEditChangeProperty(FPropertyChangedEvent& PropertyChangedEvent)
 {
@@ -261,11 +297,21 @@ bool AGameplayCueNotify_Actor::OnActive_Implementation(AActor* MyTarget, const F
 
 bool AGameplayCueNotify_Actor::WhileActive_Implementation(AActor* MyTarget, const FGameplayCueParameters& Parameters)
 {
+	if (IsHidden())
+	{
+		SetActorHiddenInGame(false);
+	}
+
 	return false;
 }
 
 bool AGameplayCueNotify_Actor::OnRemove_Implementation(AActor* MyTarget, const FGameplayCueParameters& Parameters)
 {
+	if (!IsHidden())
+	{
+		SetActorHiddenInGame(true);
+	}
+
 	return false;
 }
 
diff --git a/Engine/Plugins/Runtime/GameplayAbilities/Source/GameplayAbilities/Private/GameplayCueNotify_BurstLatent.cpp b/Engine/Plugins/Runtime/GameplayAbilities/Source/GameplayAbilities/Private/GameplayCueNotify_BurstLatent.cpp
--- a/Engine/Plugins/Runtime/GameplayAbilities/Source/GameplayAbilities/Private/GameplayCueNotify_BurstLatent.cpp
+++ b/Engine/Plugins/Runtime/GameplayAbilities/Source/GameplayAbilities/Private/GameplayCueNotify_BurstLatent.cpp
@@ -13,6 +13,12 @@ const float DefaultBurstLatentLifetime = 5.0f;
 //////////////////////////////////////////////////////////////////////////
 AGameplayCueNotify_BurstLatent::AGameplayCueNotify_BurstLatent()
 {
+	RootComponent = CreateDefaultSubobject<USceneComponent>(TEXT("RootComponent"));
+
+	PrimaryActorTick.bStartWithTickEnabled = false;
+	bAutoDestroyOnRemove = true;
+	NumPreallocatedInstances = 3;
+
 	Recycle();
 }
 
diff --git a/Engine/Plugins/Runtime/GameplayAbilities/Source/GameplayAbilities/Private/GameplayCueNotify_Instanced.cpp b/Engine/Plugins/Runtime/GameplayAbilities/Source/GameplayAbilities/Private/GameplayCueNotify_Instanced.cpp
deleted file mode 100644
diff --git a/Engine/Plugins/Runtime/GameplayAbilities/Source/GameplayAbilities/Private/GameplayCueNotify_Looping.cpp b/Engine/Plugins/Runtime/GameplayAbilities/Source/GameplayAbilities/Private/GameplayCueNotify_Looping.cpp
--- a/Engine/Plugins/Runtime/GameplayAbilities/Source/GameplayAbilities/Private/GameplayCueNotify_Looping.cpp
+++ b/Engine/Plugins/Runtime/GameplayAbilities/Source/GameplayAbilities/Private/GameplayCueNotify_Looping.cpp
@@ -8,7 +8,12 @@
 //////////////////////////////////////////////////////////////////////////
 AGameplayCueNotify_Looping::AGameplayCueNotify_Looping()
 {
+	RootComponent = CreateDefaultSubobject<USceneComponent>(TEXT("RootComponent"));
+
+	PrimaryActorTick.bStartWithTickEnabled = false;
+	bAutoDestroyOnRemove = true;
 	bAllowMultipleWhileActiveEvents = false;
+	NumPreallocatedInstances = 3;
 
 	DefaultPlacementInfo.AttachPolicy = EGameplayCueNotify_AttachPolicy::AttachToTarget;
 
diff --git a/Engine/Plugins/Runtime/GameplayAbilities/Source/GameplayAbilities/Public/GameplayCueNotify_Actor.h b/Engine/Plugins/Runtime/GameplayAbilities/Source/GameplayAbilities/Public/GameplayCueNotify_Actor.h
--- a/Engine/Plugins/Runtime/GameplayAbilities/Source/GameplayAbilities/Public/GameplayCueNotify_Actor.h
+++ b/Engine/Plugins/Runtime/GameplayAbilities/Source/GameplayAbilities/Public/GameplayCueNotify_Actor.h
@@ -30,6 +30,9 @@ class GAMEPLAYABILITIES_API AGameplayCueNotify_Actor : public AActor
 protected:
 	virtual void BeginPlay() override;
 
+	virtual void K2_DestroyActor() override;
+	virtual void Destroyed() override;
+
 public:
 	virtual void SetOwner( AActor* NewOwner ) override;
 
diff --git a/Engine/Plugins/Runtime/GameplayAbilities/Source/GameplayAbilities/Public/GameplayCueNotify_BurstLatent.h b/Engine/Plugins/Runtime/GameplayAbilities/Source/GameplayAbilities/Public/GameplayCueNotify_BurstLatent.h
--- a/Engine/Plugins/Runtime/GameplayAbilities/Source/GameplayAbilities/Public/GameplayCueNotify_BurstLatent.h
+++ b/Engine/Plugins/Runtime/GameplayAbilities/Source/GameplayAbilities/Public/GameplayCueNotify_BurstLatent.h
@@ -3,7 +3,7 @@
 #pragma once
 
 #include "CoreMinimal.h"
-#include "GameplayCueNotify_Instanced.h"
+#include "GameplayCueNotify_Actor.h"
 #include "GameplayCueNotifyTypes.h"
 #include "GameplayCueNotify_BurstLatent.generated.h"
 
@@ -15,7 +15,7 @@
  *	Since it is instanced, it can do latent things like time lines or delays.
  */
 UCLASS(Blueprintable, notplaceable, Category = "GameplayCueNotify", Meta = (ShowWorldContextPin, DisplayName = "GCN Burst Latent", ShortTooltip = "A one-off GameplayCueNotify that can use latent actions such as timelines."))
-class AGameplayCueNotify_BurstLatent : public AGameplayCueNotify_Instanced
+class AGameplayCueNotify_BurstLatent : public AGameplayCueNotify_Actor
 {
 	GENERATED_BODY()
 
diff --git a/Engine/Plugins/Runtime/GameplayAbilities/Source/GameplayAbilities/Public/GameplayCueNotify_Instanced.h b/Engine/Plugins/Runtime/GameplayAbilities/Source/GameplayAbilities/Public/GameplayCueNotify_Instanced.h
deleted file mode 100644
diff --git a/Engine/Plugins/Runtime/GameplayAbilities/Source/GameplayAbilities/Public/GameplayCueNotify_Looping.h b/Engine/Plugins/Runtime/GameplayAbilities/Source/GameplayAbilities/Public/GameplayCueNotify_Looping.h
--- a/Engine/Plugins/Runtime/GameplayAbilities/Source/GameplayAbilities/Public/GameplayCueNotify_Looping.h
+++ b/Engine/Plugins/Runtime/GameplayAbilities/Source/GameplayAbilities/Public/GameplayCueNotify_Looping.h
@@ -3,7 +3,7 @@
 #pragma once
 
 #include "CoreMinimal.h"
-#include "GameplayCueNotify_Instanced.h"
+#include "GameplayCueNotify_Actor.h"
 #include "GameplayCueNotifyTypes.h"
 #include "GameplayCueNotify_Looping.generated.h"
 
@@ -15,7 +15,7 @@
  *	The game is responsible for defining the start/stop by adding/removing the gameplay cue.
  */
 UCLASS(Blueprintable, notplaceable, Category = "GameplayCueNotify", Meta = (ShowWorldContextPin, DisplayName = "GCN Looping", ShortTooltip = "A GameplayCueNotify that has a duration that is driven by the game."))
-class AGameplayCueNotify_Looping : public AGameplayCueNotify_Instanced
+class AGameplayCueNotify_Looping : public AGameplayCueNotify_Actor
 {
 	GENERATED_BODY()
 
