Added Niagara FX support to the gameplay cue notifies;

diff --git a/Engine/Plugins/Runtime/GameplayAbilities/GameplayAbilities.uplugin b/Engine/Plugins/Runtime/GameplayAbilities/GameplayAbilities.uplugin
--- a/Engine/Plugins/Runtime/GameplayAbilities/GameplayAbilities.uplugin
+++ b/Engine/Plugins/Runtime/GameplayAbilities/GameplayAbilities.uplugin
@@ -31,6 +31,10 @@
 		{
 			"Name": "GameplayTagsEditor",
 			"Enabled": true
+		},
+		{
+			"Name": "Niagara",
+			"Enabled": true
 		}
 	]
 
diff --git a/Engine/Plugins/Runtime/GameplayAbilities/Source/GameplayAbilities/GameplayAbilities.Build.cs b/Engine/Plugins/Runtime/GameplayAbilities/Source/GameplayAbilities/GameplayAbilities.Build.cs
--- a/Engine/Plugins/Runtime/GameplayAbilities/Source/GameplayAbilities/GameplayAbilities.Build.cs
+++ b/Engine/Plugins/Runtime/GameplayAbilities/Source/GameplayAbilities/GameplayAbilities.Build.cs
@@ -22,6 +22,12 @@ namespace UnrealBuildTool.Rules
 				}
 				);
 
+			// Niagara support for gameplay cue notifies.
+			{
+				PrivateDependencyModuleNames.Add("Niagara");
+				PublicDefinitions.Add("WITH_NIAGARA=1");
+			}
+
 			if (Target.bBuildEditor == true)
 			{
 				PrivateDependencyModuleNames.Add("EditorFramework");
diff --git a/Engine/Plugins/Runtime/GameplayAbilities/Source/GameplayAbilities/Private/GameplayCueNotifyTypes.cpp b/Engine/Plugins/Runtime/GameplayAbilities/Source/GameplayAbilities/Private/GameplayCueNotifyTypes.cpp
--- a/Engine/Plugins/Runtime/GameplayAbilities/Source/GameplayAbilities/Private/GameplayCueNotifyTypes.cpp
+++ b/Engine/Plugins/Runtime/GameplayAbilities/Source/GameplayAbilities/Private/GameplayCueNotifyTypes.cpp
@@ -2,8 +2,10 @@
 
 #include "GameplayCueNotifyTypes.h"
 #include "Components/SkeletalMeshComponent.h"
-#include "Particles/ParticleSystemComponent.h"
 #include "Particles/EmitterCameraLensEffectBase.h"
+#include "NiagaraSystem.h"
+#include "NiagaraComponent.h"
+#include "NiagaraFunctionLibrary.h"
 #include "Components/AudioComponent.h"
 #include "Sound/SoundBase.h"
 #include "Camera/CameraShakeBase.h"
@@ -376,7 +378,7 @@ APlayerController* FGameplayCueNotify_SpawnContext::FindLocalPlayerController(EG
 //////////////////////////////////////////////////////////////////////////
 FGameplayCueNotify_ParticleInfo::FGameplayCueNotify_ParticleInfo()
 {
-	ParticleSystem = nullptr;
+	NiagaraSystem = nullptr;
 	bOverrideSpawnCondition = false;
 	bOverridePlacementInfo = false;
 	bCastShadow = false;
@@ -384,9 +386,9 @@ FGameplayCueNotify_ParticleInfo::FGameplayCueNotify_ParticleInfo()
 
 bool FGameplayCueNotify_ParticleInfo::PlayParticleEffect(const FGameplayCueNotify_SpawnContext& SpawnContext, FGameplayCueNotify_SpawnResult& OutSpawnResult) const
 {
-	UParticleSystemComponent* SpawnedPSC = nullptr;
+	UFXSystemComponent* SpawnedFXSC = nullptr;
 
-	if (ParticleSystem != nullptr)
+	if (NiagaraSystem != nullptr)
 	{
 		const FGameplayCueNotify_SpawnCondition& SpawnCondition = SpawnContext.GetSpawnCondition(bOverrideSpawnCondition, SpawnConditionOverride);
 		const FGameplayCueNotify_PlacementInfo& PlacementInfo = SpawnContext.GetPlacementInfo(bOverridePlacementInfo, PlacementInfoOverride);
@@ -399,36 +401,49 @@ bool FGameplayCueNotify_ParticleInfo::PlayParticleEffect(const FGameplayCueNotif
 			{
 				check(SpawnContext.World);
 
-				SpawnedPSC = SpawnContext.World->GetPSCPool().CreateWorldParticleSystem(ParticleSystem, SpawnContext.World, EPSCPoolMethod::AutoRelease);
-				if (ensure(SpawnedPSC))
+				const FVector SpawnLocation = SpawnTransform.GetLocation();
+				const FRotator SpawnRotation = SpawnTransform.Rotator();
+				const FVector SpawnScale = SpawnTransform.GetScale3D();
+				const bool bAutoDestroy = false;
+				const bool bAutoActivate = true;
+
+				if (SpawnContext.TargetComponent && (PlacementInfo.AttachPolicy == EGameplayCueNotify_AttachPolicy::AttachToTarget))
 				{
-					PlacementInfo.SetComponentTransform(SpawnedPSC, SpawnTransform);
-					PlacementInfo.TryComponentAttachment(SpawnedPSC, SpawnContext.TargetComponent);
+					const EAttachLocation::Type AttachLocationType = GetAttachLocationTypeFromRule(PlacementInfo.AttachmentRule);
 
-					SpawnedPSC->RegisterComponentWithWorld(SpawnContext.World);
-					SpawnedPSC->SetCastShadow(bCastShadow);
-					SpawnedPSC->Activate();
+					SpawnedFXSC = UNiagaraFunctionLibrary::SpawnSystemAttached(NiagaraSystem, SpawnContext.TargetComponent, PlacementInfo.SocketName,
+						SpawnLocation, SpawnRotation, SpawnScale, AttachLocationType, bAutoDestroy, ENCPoolMethod::AutoRelease, bAutoActivate);
+				}
+				else
+				{
+					SpawnedFXSC = UNiagaraFunctionLibrary::SpawnSystemAtLocation(SpawnContext.World, NiagaraSystem,
+						SpawnLocation, SpawnRotation, SpawnScale, bAutoDestroy, bAutoActivate, ENCPoolMethod::AutoRelease);
+				}
+
+				if (ensure(SpawnedFXSC))
+				{
+					SpawnedFXSC->SetCastShadow(bCastShadow);
 				}
 			}
 		}
 	}
 
 	// Always add to the list, even if null, so that the list is stable and in order for blueprint users.
-	OutSpawnResult.ParticleComponents.Add(SpawnedPSC);
+	OutSpawnResult.FxSystemComponents.Add(SpawnedFXSC);
 
-	return (SpawnedPSC != nullptr);
+	return (SpawnedFXSC != nullptr);
 }
 
 void FGameplayCueNotify_ParticleInfo::ValidateBurstAssets(UObject* ContainingAsset, const FString& Context, TArray<FText>& ValidationErrors) const
 {
 #if WITH_EDITORONLY_DATA
-	if (ParticleSystem != nullptr)
+	if (NiagaraSystem != nullptr)
 	{
-		if (ParticleSystem->IsLooping())
+		if (NiagaraSystem->IsLooping())
 		{
 			ValidationErrors.Add(FText::Format(
-				LOCTEXT("ParticleSystem_ShouldNotLoop", "Particle system [{0}] used in slot [{1}] for asset [{2}] is set to looping, but the slot is a one-shot (the instance will leak)."),
-				FText::AsCultureInvariant(ParticleSystem->GetPathName()),
+				LOCTEXT("NiagaraSystem_ShouldNotLoop", "Niagara system [{0}] used in slot [{1}] for asset [{2}] is set to looping, but the slot is a one-shot (the instance will leak)."),
+				FText::AsCultureInvariant(NiagaraSystem->GetPathName()),
 				FText::AsCultureInvariant(Context),
 				FText::AsCultureInvariant(ContainingAsset->GetPathName())));
 		}
@@ -1051,11 +1066,11 @@ void FGameplayCueNotify_LoopingEffects::StartEffects(const FGameplayCueNotify_Sp
 void FGameplayCueNotify_LoopingEffects::StopEffects(FGameplayCueNotify_SpawnResult& SpawnResult) const
 {
 	// Stop all particle effects.
-	for (UParticleSystemComponent* PSC : SpawnResult.ParticleComponents)
+	for (UFXSystemComponent* FXSC : SpawnResult.FxSystemComponents)
 	{
-		if (PSC)
+		if (FXSC)
 		{
-			PSC->Deactivate();
+			FXSC->Deactivate();
 		}
 	}
 
diff --git a/Engine/Plugins/Runtime/GameplayAbilities/Source/GameplayAbilities/Public/GameplayCueNotifyTypes.h b/Engine/Plugins/Runtime/GameplayAbilities/Source/GameplayAbilities/Public/GameplayCueNotifyTypes.h
--- a/Engine/Plugins/Runtime/GameplayAbilities/Source/GameplayAbilities/Public/GameplayCueNotifyTypes.h
+++ b/Engine/Plugins/Runtime/GameplayAbilities/Source/GameplayAbilities/Public/GameplayCueNotifyTypes.h
@@ -10,8 +10,8 @@
 #include "GameplayCueNotifyTypes.generated.h"
 
 
-class UParticleSystem;
-class UParticleSystemComponent;
+class UNiagaraSystem;
+class UFXSystemComponent;
 class UAudioComponent;
 class USoundBase;
 class UCameraShakeBase;
@@ -260,7 +260,7 @@ public:
 
 	void Reset()
 	{
-		ParticleComponents.Reset();
+		FxSystemComponents.Reset();
 		AudioComponents.Reset();
 		CameraShakes.Reset();
 		CameraAnim = nullptr;
@@ -272,9 +272,9 @@ public:
 
 public:
 
-	// List of particle components spawned.  There may be null pointers here as it matches the defined order.
+	// List of FX components spawned.  There may be null pointers here as it matches the defined order.
 	UPROPERTY(BlueprintReadOnly, Transient, Category = GameplayCueNotify)
-	TArray<UParticleSystemComponent*> ParticleComponents;
+	TArray<UFXSystemComponent*> FxSystemComponents;
 
 	// List of audio components spawned.  There may be null pointers here as it matches the defined order.
 	UPROPERTY(BlueprintReadOnly, Transient, Category = GameplayCueNotify)
@@ -334,9 +334,9 @@ public:
 	UPROPERTY(EditDefaultsOnly, BlueprintReadOnly, Category = GameplayCueNotify, Meta = (EditCondition = "bOverridePlacementInfo"))
 	FGameplayCueNotify_PlacementInfo PlacementInfoOverride;
 
-	// Particle system to spawn.
+	// Niagara FX system to spawn.
 	UPROPERTY(EditDefaultsOnly, BlueprintReadOnly, Category = GameplayCueNotify)
-	UParticleSystem* ParticleSystem;
+	UNiagaraSystem* NiagaraSystem;
 
 	// If enabled, use the spawn condition override and not the default one.
 	UPROPERTY(EditDefaultsOnly, BlueprintReadOnly, Category = GameplayCueNotify, Meta = (InlineEditConditionToggle))
