Converted gameplay cue notifies to use decal components instead of decal actors;

diff --git a/Engine/Plugins/Runtime/GameplayAbilities/Source/GameplayAbilities/Private/GameplayCueNotifyTypes.cpp b/Engine/Plugins/Runtime/GameplayAbilities/Source/GameplayAbilities/Private/GameplayCueNotifyTypes.cpp
--- a/Engine/Plugins/Runtime/GameplayAbilities/Source/GameplayAbilities/Private/GameplayCueNotifyTypes.cpp
+++ b/Engine/Plugins/Runtime/GameplayAbilities/Source/GameplayAbilities/Private/GameplayCueNotifyTypes.cpp
@@ -12,7 +12,7 @@
 #include "Camera/CameraAnim.h"
 #include "Camera/CameraAnimInst.h"
 #include "Components/ForceFeedbackComponent.h"
-#include "Engine/DecalActor.h"
+#include "Materials/MaterialInterface.h"
 #include "Components/DecalComponent.h"
 #include "Kismet/GameplayStatics.h"
 #include "GameFramework/PlayerController.h"
@@ -862,10 +862,10 @@ bool FGameplayCueNotify_ForceFeedbackInfo::PlayForceFeedback(const FGameplayCueN
 		const bool bAutoDestroy = true;
 		const float StartTime = 0.0f;
 
-		const EAttachLocation::Type AttachLocationType = GetAttachLocationTypeFromRule(PlacementInfo.AttachmentRule);
-
 		if (SpawnContext.TargetComponent && (PlacementInfo.AttachPolicy == EGameplayCueNotify_AttachPolicy::AttachToTarget))
 		{
+			const EAttachLocation::Type AttachLocationType = GetAttachLocationTypeFromRule(PlacementInfo.AttachmentRule);
+
 			OutSpawnResult.ForceFeedbackComponent = UGameplayStatics::SpawnForceFeedbackAttached(ForceFeedbackEffect, SpawnContext.TargetComponent, PlacementInfo.SocketName,
 				SpawnLocation, SpawnRotation, AttachLocationType, bStopWhenAttachedToDestroyed,
 				bIsLooping, WorldIntensity, StartTime, WorldAttenuation, bAutoDestroy);
@@ -921,7 +921,8 @@ void FGameplayCueNotify_ForceFeedbackInfo::ValidateBurstAssets(UObject* Containi
 //////////////////////////////////////////////////////////////////////////
 FGameplayCueNotify_DecalInfo::FGameplayCueNotify_DecalInfo()
 {
-	Decal = nullptr;
+	DecalMaterial = nullptr;
+	DecalSize = FVector(128.0f, 256.0f, 256.0f);
 
 	bOverrideSpawnCondition = false;
 	bOverridePlacementInfo = false;
@@ -933,7 +934,7 @@ FGameplayCueNotify_DecalInfo::FGameplayCueNotify_DecalInfo()
 
 bool FGameplayCueNotify_DecalInfo::SpawnDecal(const FGameplayCueNotify_SpawnContext& SpawnContext, FGameplayCueNotify_SpawnResult& OutSpawnResult) const
 {
-	if (!Decal)
+	if (!DecalMaterial)
 	{
 		return false;
 	}
@@ -959,25 +960,35 @@ bool FGameplayCueNotify_DecalInfo::SpawnDecal(const FGameplayCueNotify_SpawnCont
 		SpawnTransform.SetRotation(InverseRotation);
 	}
 
-	ADecalActor* SpawnedDecal = SpawnContext.World->SpawnActor<ADecalActor>(Decal, SpawnTransform);
-	if (!SpawnedDecal)
-	{
-		return false;
-	}
+	UDecalComponent* SpawnedDecalComponent = nullptr;
 
-	UDecalComponent* DecalComponent = SpawnedDecal->GetDecal();
-	USceneComponent* DecalRootComponent = SpawnedDecal->GetRootComponent();
+	const FVector SpawnLocation = SpawnTransform.GetLocation();
+	const FRotator SpawnRotation = SpawnTransform.Rotator();
+	const float LifeSpan = 0.0f;
 
-	PlacementInfo.SetComponentTransform(DecalRootComponent, SpawnTransform);
-	PlacementInfo.TryComponentAttachment(DecalRootComponent, SpawnContext.TargetComponent);
+	if (SpawnContext.TargetComponent && (PlacementInfo.AttachPolicy == EGameplayCueNotify_AttachPolicy::AttachToTarget))
+	{
+		const EAttachLocation::Type AttachLocationType = GetAttachLocationTypeFromRule(PlacementInfo.AttachmentRule);
 
-	if (DecalComponent && bOverrideFadeOut)
+		SpawnedDecalComponent = UGameplayStatics::SpawnDecalAttached(DecalMaterial, DecalSize, SpawnContext.TargetComponent, PlacementInfo.SocketName,
+			SpawnLocation, SpawnRotation, AttachLocationType, LifeSpan);
+	}
+	else
 	{
-		DecalComponent->SetFadeOut(FadeOutStartDelay, FadeOutDuration);
+		SpawnedDecalComponent = UGameplayStatics::SpawnDecalAtLocation(SpawnContext.World, DecalMaterial, DecalSize, SpawnLocation, SpawnRotation, LifeSpan);
 	}
 
-	OutSpawnResult.DecalActor = SpawnedDecal;
-	return true;
+	if (ensure(SpawnedDecalComponent))
+	{
+		if (bOverrideFadeOut)
+		{
+			SpawnedDecalComponent->SetFadeOut(FadeOutStartDelay, FadeOutDuration);
+		}
+	}
+
+	OutSpawnResult.DecalComponent = SpawnedDecalComponent;
+
+	return (OutSpawnResult.DecalComponent != nullptr);
 }
 
 
@@ -1136,7 +1147,7 @@ void FGameplayCueNotify_LoopingEffects::StopEffects(FGameplayCueNotify_SpawnResu
 	}
 
 	// There should be no decal on looping gameplay cues.
-	ensure(SpawnResult.DecalActor == nullptr);
+	ensure(SpawnResult.DecalComponent == nullptr);
 
 	// Clear the spawn results.
 	SpawnResult.Reset();
diff --git a/Engine/Plugins/Runtime/GameplayAbilities/Source/GameplayAbilities/Public/GameplayCueNotifyTypes.h b/Engine/Plugins/Runtime/GameplayAbilities/Source/GameplayAbilities/Public/GameplayCueNotifyTypes.h
--- a/Engine/Plugins/Runtime/GameplayAbilities/Source/GameplayAbilities/Public/GameplayCueNotifyTypes.h
+++ b/Engine/Plugins/Runtime/GameplayAbilities/Source/GameplayAbilities/Public/GameplayCueNotifyTypes.h
@@ -21,7 +21,8 @@ class AEmitterCameraLensEffectBase;
 class UForceFeedbackEffect;
 class UForceFeedbackAttenuation;
 class UForceFeedbackComponent;
-class ADecalActor;
+class UDecalComponent;
+class UMaterialInterface;
 class APlayerController;
 struct FGameplayCueParameters;
 struct FGameplayCueNotify_SpawnContext;
@@ -267,7 +268,7 @@ public:
 		CameraLensEffects.Reset();
 		ForceFeedbackComponent = nullptr;
 		ForceFeedbackTargetPC = nullptr;
-		DecalActor = nullptr;
+		DecalComponent = nullptr;
 	}
 
 public:
@@ -300,9 +301,9 @@ public:
 	UPROPERTY(Transient)
 	APlayerController* ForceFeedbackTargetPC;
 
-	// Spawned decal actor.  This may be null.
+	// Spawned decal component.  This may be null.
 	UPROPERTY(BlueprintReadOnly, Transient, Category = GameplayCueNotify)
-	ADecalActor* DecalActor;
+	UDecalComponent* DecalComponent;
 };
 
 
@@ -675,9 +676,13 @@ public:
 	UPROPERTY(EditDefaultsOnly, BlueprintReadOnly, Category = GameplayCueNotify, Meta = (EditCondition = "bOverridePlacementInfo"))
 	FGameplayCueNotify_PlacementInfo PlacementInfoOverride;
 
-	// Decal actor class to spawn.
+	// Decal material to spawn.
 	UPROPERTY(EditDefaultsOnly, BlueprintReadOnly, Category = GameplayCueNotify)
-	TSubclassOf<ADecalActor> Decal;
+	UMaterialInterface* DecalMaterial;
+
+	// Decal size in local space (does not include the component scale).
+	UPROPERTY(EditDefaultsOnly, BlueprintReadOnly, Category = GameplayCueNotify, meta = (AllowPreserveRatio = "true"))
+	FVector DecalSize;
 
 	// If enabled, use the spawn condition override and not the default one.
 	UPROPERTY(EditDefaultsOnly, BlueprintReadOnly, Category = GameplayCueNotify, Meta = (InlineEditConditionToggle))
