Lyra: Setup project to use new GameplayCutNotify classes in the GameplayCue Editor.

* Implemented the GameplayCue delegates used to setup the GameplayCue Editor.

diff --git a/Engine/Plugins/Runtime/GameplayAbilities/Source/GameplayAbilities/Public/GameplayCueNotify_Burst.h b/Engine/Plugins/Runtime/GameplayAbilities/Source/GameplayAbilities/Public/GameplayCueNotify_Burst.h
--- a/Engine/Plugins/Runtime/GameplayAbilities/Source/GameplayAbilities/Public/GameplayCueNotify_Burst.h
+++ b/Engine/Plugins/Runtime/GameplayAbilities/Source/GameplayAbilities/Public/GameplayCueNotify_Burst.h
@@ -15,7 +15,7 @@
  *	Since it is not instanced, it cannot do latent actions such as delays and time lines.
  */
 UCLASS(Blueprintable, notplaceable, Category = "GameplayCueNotify", Meta = (ShowWorldContextPin, DisplayName = "GCN Burst", ShortTooltip = "A one-off GameplayCueNotify that is never spawned into the world."))
-class UGameplayCueNotify_Burst : public UGameplayCueNotify_Static
+class GAMEPLAYABILITIES_API UGameplayCueNotify_Burst : public UGameplayCueNotify_Static
 {
 	GENERATED_BODY()
 
diff --git a/Engine/Plugins/Runtime/GameplayAbilities/Source/GameplayAbilities/Public/GameplayCueNotify_BurstLatent.h b/Engine/Plugins/Runtime/GameplayAbilities/Source/GameplayAbilities/Public/GameplayCueNotify_BurstLatent.h
--- a/Engine/Plugins/Runtime/GameplayAbilities/Source/GameplayAbilities/Public/GameplayCueNotify_BurstLatent.h
+++ b/Engine/Plugins/Runtime/GameplayAbilities/Source/GameplayAbilities/Public/GameplayCueNotify_BurstLatent.h
@@ -15,7 +15,7 @@
  *	Since it is instanced, it can do latent things like time lines or delays.
  */
 UCLASS(Blueprintable, notplaceable, Category = "GameplayCueNotify", Meta = (ShowWorldContextPin, DisplayName = "GCN Burst Latent", ShortTooltip = "A one-off GameplayCueNotify that can use latent actions such as timelines."))
-class AGameplayCueNotify_BurstLatent : public AGameplayCueNotify_Actor
+class GAMEPLAYABILITIES_API AGameplayCueNotify_BurstLatent : public AGameplayCueNotify_Actor
 {
 	GENERATED_BODY()
 
diff --git a/Engine/Plugins/Runtime/GameplayAbilities/Source/GameplayAbilities/Public/GameplayCueNotify_Looping.h b/Engine/Plugins/Runtime/GameplayAbilities/Source/GameplayAbilities/Public/GameplayCueNotify_Looping.h
--- a/Engine/Plugins/Runtime/GameplayAbilities/Source/GameplayAbilities/Public/GameplayCueNotify_Looping.h
+++ b/Engine/Plugins/Runtime/GameplayAbilities/Source/GameplayAbilities/Public/GameplayCueNotify_Looping.h
@@ -15,7 +15,7 @@
  *	The game is responsible for defining the start/stop by adding/removing the gameplay cue.
  */
 UCLASS(Blueprintable, notplaceable, Category = "GameplayCueNotify", Meta = (ShowWorldContextPin, DisplayName = "GCN Looping", ShortTooltip = "A GameplayCueNotify that has a duration that is driven by the game."))
-class AGameplayCueNotify_Looping : public AGameplayCueNotify_Actor
+class GAMEPLAYABILITIES_API AGameplayCueNotify_Looping : public AGameplayCueNotify_Actor
 {
 	GENERATED_BODY()
 
