'UAttributeSet::PostAttributeChange' is no longer const.  This allows the ability set to update internal state and it matches the call to 'PreAttributeChange'.

diff --git a/Engine/Plugins/Runtime/GameplayAbilities/Source/GameplayAbilities/Public/AttributeSet.h b/Engine/Plugins/Runtime/GameplayAbilities/Source/GameplayAbilities/Public/AttributeSet.h
--- a/Engine/Plugins/Runtime/GameplayAbilities/Source/GameplayAbilities/Public/AttributeSet.h
+++ b/Engine/Plugins/Runtime/GameplayAbilities/Source/GameplayAbilities/Public/AttributeSet.h
@@ -217,7 +217,7 @@ public:
 	virtual void PreAttributeChange(const FGameplayAttribute& Attribute, float& NewValue) { }
 
 	/** Called just after any modification happens to an attribute. */
-	virtual void PostAttributeChange(const FGameplayAttribute& Attribute, float OldValue, float NewValue) const { }
+	virtual void PostAttributeChange(const FGameplayAttribute& Attribute, float OldValue, float NewValue) { }
 
 	/**
 	 *	This is called just before any modification happens to an attribute's base value when an attribute aggregator exists.
