Fix missing toolbar for gameplay abilities blueprints

diff --git a/Engine/Plugins/Runtime/GameplayAbilities/Source/GameplayAbilitiesEditor/Public/GameplayAbilitiesEditor.h b/Engine/Plugins/Runtime/GameplayAbilities/Source/GameplayAbilitiesEditor/Public/GameplayAbilitiesEditor.h
--- a/Engine/Plugins/Runtime/GameplayAbilities/Source/GameplayAbilitiesEditor/Public/GameplayAbilitiesEditor.h
+++ b/Engine/Plugins/Runtime/GameplayAbilities/Source/GameplayAbilitiesEditor/Public/GameplayAbilitiesEditor.h
@@ -48,6 +48,7 @@ public:
 	virtual FText GetToolkitToolTipText() const override;
 	virtual FString GetWorldCentricTabPrefix() const override;
 	virtual FLinearColor GetWorldCentricTabColorScale() const override;
+	virtual bool UsesCustomToolbarPlacement() const override { return false; }
 	// End of IToolkit interface
 
 	/** @return the documentation location for this editor */
diff --git a/Engine/Source/Editor/UMGEditor/Public/WidgetBlueprintEditor.h b/Engine/Source/Editor/UMGEditor/Public/WidgetBlueprintEditor.h
--- a/Engine/Source/Editor/UMGEditor/Public/WidgetBlueprintEditor.h
+++ b/Engine/Source/Editor/UMGEditor/Public/WidgetBlueprintEditor.h
@@ -216,6 +216,7 @@ protected:
 	virtual void RegisterApplicationModes(const TArray<UBlueprint*>& InBlueprints, bool bShouldOpenInDefaultsMode, bool bNewlyCreated = false) override;
 	virtual FGraphAppearanceInfo GetGraphAppearance(class UEdGraph* InGraph) const override;
 	virtual TSubclassOf<UEdGraphSchema> GetDefaultSchemaClass() const override;
+	virtual bool UsesCustomToolbarPlacement() const override { return true; }
 	// End FBlueprintEditor
 
 private:
