Updated FGameplayTagBlueprintPropertyMap to not initialize again if the owner and ability system hasn't changed.

diff --git a/Engine/Plugins/Runtime/GameplayAbilities/Source/GameplayAbilities/Private/GameplayEffectTypes.cpp b/Engine/Plugins/Runtime/GameplayAbilities/Source/GameplayAbilities/Private/GameplayEffectTypes.cpp
--- a/Engine/Plugins/Runtime/GameplayAbilities/Source/GameplayAbilities/Private/GameplayEffectTypes.cpp
+++ b/Engine/Plugins/Runtime/GameplayAbilities/Source/GameplayAbilities/Private/GameplayEffectTypes.cpp
@@ -601,6 +601,12 @@ void FGameplayTagBlueprintPropertyMap::Initialize(UObject* Owner, UAbilitySystem
 		return;
 	}
 
+	if ((CachedOwner == Owner) && (CachedASC == ASC))
+	{
+		// Already initialized.
+		return;
+	}
+
 	if (CachedOwner.IsValid())
 	{
 		Unregister();
