Stubbing out the full light-dart ability.

diff --git a/Engine/Plugins/Runtime/GameplayAbilities/Source/GameplayAbilities/Public/Abilities/Tasks/AbilityTask_PlayMontageAndWait.h b/Engine/Plugins/Runtime/GameplayAbilities/Source/GameplayAbilities/Public/Abilities/Tasks/AbilityTask_PlayMontageAndWait.h
--- a/Engine/Plugins/Runtime/GameplayAbilities/Source/GameplayAbilities/Public/Abilities/Tasks/AbilityTask_PlayMontageAndWait.h
+++ b/Engine/Plugins/Runtime/GameplayAbilities/Source/GameplayAbilities/Public/Abilities/Tasks/AbilityTask_PlayMontageAndWait.h
@@ -62,7 +62,7 @@ class GAMEPLAYABILITIES_API UAbilityTask_PlayMontageAndWait : public UAbilityTas
 
 	virtual FString GetDebugString() const override;
 
-private:
+protected:
 
 	virtual void OnDestroy(bool AbilityEnded) override;
 
