Replace all engine-level manual casts to IAbilitySystemInterface with calls to GetAbilitySystemComponentFromActor wrapper
Change behavior of that wrapper to do a backup component search by default as it isn't actually that slow
These two changes combined make basic ability system operations possible on AActor blueprints that don't implement the interface

diff --git a/Engine/Plugins/Runtime/GameplayAbilities/Source/GameplayAbilities/Private/AbilitySystemBlueprintLibrary.cpp b/Engine/Plugins/Runtime/GameplayAbilities/Source/GameplayAbilities/Private/AbilitySystemBlueprintLibrary.cpp
--- a/Engine/Plugins/Runtime/GameplayAbilities/Source/GameplayAbilities/Private/AbilitySystemBlueprintLibrary.cpp
+++ b/Engine/Plugins/Runtime/GameplayAbilities/Source/GameplayAbilities/Private/AbilitySystemBlueprintLibrary.cpp
@@ -5,7 +5,6 @@
 #include "GameplayEffect.h"
 #include "AbilitySystemGlobals.h"
 #include "AbilitySystemComponent.h"
-#include "AbilitySystemInterface.h"
 
 UAbilitySystemBlueprintLibrary::UAbilitySystemBlueprintLibrary(const FObjectInitializer& ObjectInitializer)
 : Super(ObjectInitializer)
@@ -21,19 +20,15 @@ void UAbilitySystemBlueprintLibrary::SendGameplayEventToActor(AActor* Actor, FGa
 {
 	if (Actor && !Actor->IsPendingKill())
 	{
-		IAbilitySystemInterface* AbilitySystemInterface = Cast<IAbilitySystemInterface>(Actor);
-		if (AbilitySystemInterface != nullptr)
+		UAbilitySystemComponent* AbilitySystemComponent = GetAbilitySystemComponent(Actor);
+		if (AbilitySystemComponent != nullptr && !AbilitySystemComponent->IsPendingKill())
 		{
-			UAbilitySystemComponent* AbilitySystemComponent = AbilitySystemInterface->GetAbilitySystemComponent();
-			if (AbilitySystemComponent != nullptr && !AbilitySystemComponent->IsPendingKill())
-			{
-				FScopedPredictionWindow NewScopedWindow(AbilitySystemComponent, true);
-				AbilitySystemComponent->HandleGameplayEvent(EventTag, &Payload);
-			}
-			else
-			{
-				ABILITY_LOG(Error, TEXT("UAbilitySystemBlueprintLibrary::SendGameplayEventToActor: Invalid ability system component retrieved from Actor %s. EventTag was %s"), *Actor->GetName(), *EventTag.ToString());
-			}
+			FScopedPredictionWindow NewScopedWindow(AbilitySystemComponent, true);
+			AbilitySystemComponent->HandleGameplayEvent(EventTag, &Payload);
+		}
+		else
+		{
+			ABILITY_LOG(Error, TEXT("UAbilitySystemBlueprintLibrary::SendGameplayEventToActor: Invalid ability system component retrieved from Actor %s. EventTag was %s"), *Actor->GetName(), *EventTag.ToString());
 		}
 	}
 }
diff --git a/Engine/Plugins/Runtime/GameplayAbilities/Source/GameplayAbilities/Private/AbilitySystemDebugHUD.cpp b/Engine/Plugins/Runtime/GameplayAbilities/Source/GameplayAbilities/Private/AbilitySystemDebugHUD.cpp
--- a/Engine/Plugins/Runtime/GameplayAbilities/Source/GameplayAbilities/Private/AbilitySystemDebugHUD.cpp
+++ b/Engine/Plugins/Runtime/GameplayAbilities/Source/GameplayAbilities/Private/AbilitySystemDebugHUD.cpp
@@ -10,7 +10,7 @@
 #include "Engine/LocalPlayer.h"
 #include "EngineUtils.h"
 #include "AbilitySystemComponent.h"
-#include "AbilitySystemInterface.h"
+#include "AbilitySystemGlobals.h"
 
 AAbilitySystemDebugHUD::AAbilitySystemDebugHUD(const FObjectInitializer& ObjectInitializer) : Super(ObjectInitializer)
 {
@@ -66,13 +66,10 @@ void AAbilitySystemDebugHUD::DrawDebugHUD(UCanvas* InCanvas, APlayerController*)
 		return;
 	}
 
-	if (IAbilitySystemInterface* Pawn = Cast<IAbilitySystemInterface>(PC->GetPawn()))
+	if (UAbilitySystemComponent* AbilitySystemComponent = UAbilitySystemGlobals::GetAbilitySystemComponentFromActor(PC->GetPawn()))
 	{
-		if (UAbilitySystemComponent* AbilitySystemComponent = Pawn->GetAbilitySystemComponent())
-		{
-			DrawDebugAbilitySystemComponent(AbilitySystemComponent);
-			return;
-		}
+		DrawDebugAbilitySystemComponent(AbilitySystemComponent);
+		return;
 	}
 
 	if (PC->GetPawn())
diff --git a/Engine/Plugins/Runtime/GameplayAbilities/Source/GameplayAbilities/Private/AbilitySystemGlobals.cpp b/Engine/Plugins/Runtime/GameplayAbilities/Source/GameplayAbilities/Private/AbilitySystemGlobals.cpp
--- a/Engine/Plugins/Runtime/GameplayAbilities/Source/GameplayAbilities/Private/AbilitySystemGlobals.cpp
+++ b/Engine/Plugins/Runtime/GameplayAbilities/Source/GameplayAbilities/Private/AbilitySystemGlobals.cpp
@@ -179,9 +179,7 @@ UAbilitySystemComponent* UAbilitySystemGlobals::GetAbilitySystemComponentFromAct
 
 	if (LookForComponent)
 	{
-		/** This is slow and not desirable */
-		ABILITY_LOG(Warning, TEXT("GetAbilitySystemComponentFromActor called on %s that is not IAbilitySystemInterface. This slow!"), *Actor->GetName());
-
+		// Fall back to a component search to better support BP-only actors
 		return Actor->FindComponentByClass<UAbilitySystemComponent>();
 	}
 
@@ -468,8 +466,8 @@ void UAbilitySystemGlobals::ListPlayerAbilities()
 {
 #if WITH_ABILITY_CHEATS	
 	APlayerController* PC = GWorld->GetFirstPlayerController();
-	IAbilitySystemInterface* AbilitySystem = PC ? Cast<IAbilitySystemInterface>(PC->GetPawn()) : nullptr;
-	if(AbilitySystem)
+	UAbilitySystemComponent* AbilityComponent = PC ? GetAbilitySystemComponentFromActor(PC->GetPawn()) : nullptr;
+	if(AbilityComponent)
 	{
 		const UEnum* ExecutionEnumPtr = FindObject<UEnum>(ANY_PACKAGE, TEXT("EGameplayAbilityNetExecutionPolicy"), true);
 		check(ExecutionEnumPtr && TEXT("Couldn't locate EGameplayAbilityNetExecutionPolicy enum!"));
@@ -478,7 +476,6 @@ void UAbilitySystemGlobals::ListPlayerAbilities()
 
 		PC->ClientMessage(TEXT("Available abilities:"));
 
-		UAbilitySystemComponent* AbilityComponent = AbilitySystem->GetAbilitySystemComponent();
 		check(AbilityComponent && TEXT("Failed to find ability component on player pawn."));
 		for (FGameplayAbilitySpec &Activatable : AbilityComponent->GetActivatableAbilities())
 		{
@@ -498,10 +495,9 @@ void UAbilitySystemGlobals::ServerActivatePlayerAbility(FString AbilityNameMatch
 	}
 
 	APlayerController* PC = GWorld->GetFirstPlayerController();
-	IAbilitySystemInterface* AbilitySystem = PC ? Cast<IAbilitySystemInterface>(PC->GetPawn()) : nullptr;
-	if(AbilitySystem)
+	UAbilitySystemComponent* AbilityComponent = PC ? GetAbilitySystemComponentFromActor(PC->GetPawn()) : nullptr;
+	if(AbilityComponent)
 	{
-		UAbilitySystemComponent* AbilityComponent = AbilitySystem->GetAbilitySystemComponent();
 		for (FGameplayAbilitySpec &Activatable : AbilityComponent->GetActivatableAbilities())
 		{
 			// Trigger on first match only
@@ -523,10 +519,9 @@ void UAbilitySystemGlobals::ServerActivatePlayerAbility(FString AbilityNameMatch
 static void TerminatePlayerAbility(const FString &AbilityNameMatch, TFunction<FString(UAbilitySystemComponent*, FGameplayAbilitySpec&)> TerminateFn)
 {
 	APlayerController* PC = GWorld->GetFirstPlayerController();
-	IAbilitySystemInterface* AbilitySystem = PC ? Cast<IAbilitySystemInterface>(PC->GetPawn()) : nullptr;
-	if (AbilitySystem)
+	UAbilitySystemComponent* AbilityComponent = PC ? UAbilitySystemGlobals::GetAbilitySystemComponentFromActor(PC->GetPawn()) : nullptr;
+	if (AbilityComponent)
 	{
-		UAbilitySystemComponent* AbilityComponent = AbilitySystem->GetAbilitySystemComponent();
 		for (FGameplayAbilitySpec &ActivatableSpec : AbilityComponent->GetActivatableAbilities())
 		{
 			if (ActivatableSpec.Ability->IsActive() &&
diff --git a/Engine/Plugins/Runtime/GameplayAbilities/Source/GameplayAbilities/Private/AnimNotify_GameplayCue.cpp b/Engine/Plugins/Runtime/GameplayAbilities/Source/GameplayAbilities/Private/AnimNotify_GameplayCue.cpp
--- a/Engine/Plugins/Runtime/GameplayAbilities/Source/GameplayAbilities/Private/AnimNotify_GameplayCue.cpp
+++ b/Engine/Plugins/Runtime/GameplayAbilities/Source/GameplayAbilities/Private/AnimNotify_GameplayCue.cpp
@@ -3,7 +3,6 @@
 #include "AnimNotify_GameplayCue.h"
 #include "Components/SkeletalMeshComponent.h"
 #include "AbilitySystemGlobals.h"
-#include "AbilitySystemInterface.h"
 #include "AbilitySystemComponent.h"
 #include "Animation/AnimMontage.h"
 #include "GameplayCueManager.h"
@@ -37,21 +36,18 @@ static void ProcessGameplayCue(GameplayCueFunc Func, USkeletalMeshComponent* Mes
 
 		// Try to get the ability level. This may not be able to find the ability level
 		// in cases where a blend out is happening or two abilities are trying to play animations
-		if (const IAbilitySystemInterface* ASI = Cast<IAbilitySystemInterface>(OwnerActor))
+		if (UAbilitySystemComponent* ASC = UAbilitySystemGlobals::GetAbilitySystemComponentFromActor(OwnerActor))
 		{
-			if (UAbilitySystemComponent* ASC = ASI->GetAbilitySystemComponent())
+			if (ASC->GetCurrentMontage() == Animation)
 			{
-				if (ASC->GetCurrentMontage() == Animation)
+				if (const UGameplayAbility* Ability = ASC->GetAnimatingAbility())
 				{
-					if (const UGameplayAbility* Ability = ASC->GetAnimatingAbility())
-					{
-						Parameters.AbilityLevel = Ability->GetAbilityLevel();
-					}
+					Parameters.AbilityLevel = Ability->GetAbilityLevel();
 				}
-
-				// Always use ASC's owner for instigator
-				Parameters.Instigator = ASC->GetOwner();
 			}
+
+			// Always use ASC's owner for instigator
+			Parameters.Instigator = ASC->GetOwner();
 		}
 
 		Parameters.TargetAttachComponent = MeshComp;
diff --git a/Engine/Plugins/Runtime/GameplayAbilities/Source/GameplayAbilities/Private/GameplayCueFunctionLibrary.cpp b/Engine/Plugins/Runtime/GameplayAbilities/Source/GameplayAbilities/Private/GameplayCueFunctionLibrary.cpp
--- a/Engine/Plugins/Runtime/GameplayAbilities/Source/GameplayAbilities/Private/GameplayCueFunctionLibrary.cpp
+++ b/Engine/Plugins/Runtime/GameplayAbilities/Source/GameplayAbilities/Private/GameplayCueFunctionLibrary.cpp
@@ -2,7 +2,7 @@
 
 #include "GameplayCueFunctionLibrary.h"
 #include "GameplayCueManager.h"
-#include "AbilitySystemInterface.h"
+#include "AbilitySystemGlobals.h"
 #include "AbilitySystemComponent.h"
 
 
@@ -27,8 +27,7 @@ void UGameplayCueFunctionLibrary::ExecuteGameplayCueOnActor(AActor* Target, cons
 		return;
 	}
 
-	const IAbilitySystemInterface* ASI = Cast<IAbilitySystemInterface>(Target);
-	UAbilitySystemComponent *ASC = (ASI ? ASI->GetAbilitySystemComponent() : nullptr);
+	UAbilitySystemComponent* ASC = UAbilitySystemGlobals::GetAbilitySystemComponentFromActor(Target);
 
 	if (ASC)
 	{
@@ -52,8 +51,7 @@ void UGameplayCueFunctionLibrary::AddGameplayCueOnActor(AActor* Target, const FG
 		return;
 	}
 
-	const IAbilitySystemInterface* ASI = Cast<IAbilitySystemInterface>(Target);
-	UAbilitySystemComponent *ASC = (ASI ? ASI->GetAbilitySystemComponent() : nullptr);
+	UAbilitySystemComponent* ASC = UAbilitySystemGlobals::GetAbilitySystemComponentFromActor(Target);
 
 	if (ASC)
 	{
@@ -77,8 +75,7 @@ void UGameplayCueFunctionLibrary::RemoveGameplayCueOnActor(AActor* Target, const
 		return;
 	}
 
-	const IAbilitySystemInterface* ASI = Cast<IAbilitySystemInterface>(Target);
-	UAbilitySystemComponent *ASC = (ASI ? ASI->GetAbilitySystemComponent() : nullptr);
+	UAbilitySystemComponent* ASC = UAbilitySystemGlobals::GetAbilitySystemComponentFromActor(Target);
 
 	if (ASC)
 	{
diff --git a/Engine/Plugins/Runtime/GameplayAbilities/Source/GameplayAbilities/Private/GameplayCueManager.cpp b/Engine/Plugins/Runtime/GameplayAbilities/Source/GameplayAbilities/Private/GameplayCueManager.cpp
--- a/Engine/Plugins/Runtime/GameplayAbilities/Source/GameplayAbilities/Private/GameplayCueManager.cpp
+++ b/Engine/Plugins/Runtime/GameplayAbilities/Source/GameplayAbilities/Private/GameplayCueManager.cpp
@@ -16,7 +16,6 @@
 #include "GameplayCueSet.h"
 #include "GameplayCueNotify_Static.h"
 #include "AbilitySystemComponent.h"
-#include "AbilitySystemInterface.h"
 #include "Net/DataReplication.h"
 #include "Engine/ActorChannel.h"
 #include "Engine/NetConnection.h"
@@ -235,12 +234,9 @@ void UGameplayCueManager::TranslateGameplayCue(FGameplayTag& Tag, AActor* Target
 
 void UGameplayCueManager::AddGameplayCue_NonReplicated(AActor* Target, const FGameplayTag GameplayCueTag, const FGameplayCueParameters& Parameters)
 {
-	if (const IAbilitySystemInterface* ASI = Cast<IAbilitySystemInterface>(Target))
+	if (UAbilitySystemComponent* ASC = UAbilitySystemGlobals::GetAbilitySystemComponentFromActor(Target))
 	{
-		if (UAbilitySystemComponent* ASC = ASI->GetAbilitySystemComponent())
-		{
-			ASC->AddLooseGameplayTag(GameplayCueTag);
-		}
+		ASC->AddLooseGameplayTag(GameplayCueTag);
 	}
 
 	if (UGameplayCueManager* GCM = UAbilitySystemGlobals::Get().GetGameplayCueManager())
@@ -252,12 +248,9 @@ void UGameplayCueManager::AddGameplayCue_NonReplicated(AActor* Target, const FGa
 
 void UGameplayCueManager::RemoveGameplayCue_NonReplicated(AActor* Target, const FGameplayTag GameplayCueTag, const FGameplayCueParameters& Parameters)
 {
-	if (const IAbilitySystemInterface* ASI = Cast<IAbilitySystemInterface>(Target))
+	if (UAbilitySystemComponent* ASC = UAbilitySystemGlobals::GetAbilitySystemComponentFromActor(Target))
 	{
-		if (UAbilitySystemComponent* ASC = ASI->GetAbilitySystemComponent())
-		{
-			ASC->RemoveLooseGameplayTag(GameplayCueTag);
-		}
+		ASC->RemoveLooseGameplayTag(GameplayCueTag);
 	}
 
 	if (UGameplayCueManager* GCM = UAbilitySystemGlobals::Get().GetGameplayCueManager())
diff --git a/Engine/Plugins/Runtime/GameplayAbilities/Source/GameplayAbilities/Private/GameplayEffectTypes.cpp b/Engine/Plugins/Runtime/GameplayAbilities/Source/GameplayAbilities/Private/GameplayEffectTypes.cpp
--- a/Engine/Plugins/Runtime/GameplayAbilities/Source/GameplayAbilities/Private/GameplayEffectTypes.cpp
+++ b/Engine/Plugins/Runtime/GameplayAbilities/Source/GameplayAbilities/Private/GameplayEffectTypes.cpp
@@ -8,7 +8,6 @@
 #include "Misc/ConfigCacheIni.h"
 #include "AbilitySystemGlobals.h"
 #include "AbilitySystemComponent.h"
-#include "AbilitySystemInterface.h"
 #include "Engine/PackageMapClient.h"
 
 
@@ -120,11 +119,7 @@ void FGameplayEffectContext::AddInstigator(class AActor *InInstigator, class AAc
 	InstigatorAbilitySystemComponent = NULL;
 
 	// Cache off his AbilitySystemComponent.
-	IAbilitySystemInterface* AbilitySystemInterface = Cast<IAbilitySystemInterface>(Instigator.Get());
-	if (AbilitySystemInterface)
-	{
-		InstigatorAbilitySystemComponent = AbilitySystemInterface->GetAbilitySystemComponent();
-	}
+	InstigatorAbilitySystemComponent = UAbilitySystemGlobals::GetAbilitySystemComponentFromActor(Instigator.Get());
 }
 
 void FGameplayEffectContext::SetAbility(const UGameplayAbility* InGameplayAbility)
diff --git a/Engine/Plugins/Runtime/GameplayAbilities/Source/GameplayAbilities/Public/AbilitySystemBlueprintLibrary.h b/Engine/Plugins/Runtime/GameplayAbilities/Source/GameplayAbilities/Public/AbilitySystemBlueprintLibrary.h
--- a/Engine/Plugins/Runtime/GameplayAbilities/Source/GameplayAbilities/Public/AbilitySystemBlueprintLibrary.h
+++ b/Engine/Plugins/Runtime/GameplayAbilities/Source/GameplayAbilities/Public/AbilitySystemBlueprintLibrary.h
@@ -25,14 +25,13 @@ class GAMEPLAYABILITIES_API UAbilitySystemBlueprintLibrary : public UBlueprintFu
 {
 	GENERATED_UCLASS_BODY()
 
-	/** Tries to find an ability system component on the actor, will use AbilitySystemInterface */
+	/** Tries to find an ability system component on the actor, will use AbilitySystemInterface or fall back to a component search */
 	UFUNCTION(BlueprintPure, Category = Ability)
 	static UAbilitySystemComponent* GetAbilitySystemComponent(AActor *Actor);
 
 	/**
 	 * This function can be used to trigger an ability on the actor in question with useful payload data.
-	 * NOTE: The Actor passed in must implement IAbilitySystemInterface! or else this function will silently fail to
-	 * send the event.  The actor needs the interface to find the UAbilitySystemComponent, and if the component isn't
+	 * NOTE: GetAbilitySystemComponent is called on the actor to find a good component, and if the component isn't
 	 * found, the event will not be sent.
 	 */
 	UFUNCTION(BlueprintCallable, Category = Ability, Meta = (Tooltip = "This function can be used to trigger an ability on the actor in question with useful payload data."))
diff --git a/Engine/Plugins/Runtime/GameplayAbilities/Source/GameplayAbilities/Public/AbilitySystemGlobals.h b/Engine/Plugins/Runtime/GameplayAbilities/Source/GameplayAbilities/Public/AbilitySystemGlobals.h
--- a/Engine/Plugins/Runtime/GameplayAbilities/Source/GameplayAbilities/Public/AbilitySystemGlobals.h
+++ b/Engine/Plugins/Runtime/GameplayAbilities/Source/GameplayAbilities/Public/AbilitySystemGlobals.h
@@ -68,8 +68,8 @@ class GAMEPLAYABILITIES_API UAbilitySystemGlobals : public UObject
 	/** Returns data used to initialize attributes to their default values */
 	FAttributeSetInitter* GetAttributeSetInitter() const;
 
-	/** Searches the passed in actor for an ability system component, will use the AbilitySystemInterface */
-	static UAbilitySystemComponent* GetAbilitySystemComponentFromActor(const AActor* Actor, bool LookForComponent=false);
+	/** Searches the passed in actor for an ability system component, will use IAbilitySystemInterface or fall back to a component search */
+	static UAbilitySystemComponent* GetAbilitySystemComponentFromActor(const AActor* Actor, bool LookForComponent=true);
 
 	/** Should allocate a project specific AbilityActorInfo struct. Caller is responsible for deallocation */
 	virtual FGameplayAbilityActorInfo* AllocAbilityActorInfo() const;
