Add actor-targeted async actions to wait for gameplay effect, event, and tag changes so they can be used in normal blueprints
These are versions of existing ability tasks and have the same behavior, can replace game-specific actions
Change base async action to new-style constructor

diff --git a/Engine/Plugins/Runtime/GameplayAbilities/Source/GameplayAbilities/Private/Abilities/Async/AbilityAsync.cpp b/Engine/Plugins/Runtime/GameplayAbilities/Source/GameplayAbilities/Private/Abilities/Async/AbilityAsync.cpp
new file mode 100644
--- /dev/null
+++ b/Engine/Plugins/Runtime/GameplayAbilities/Source/GameplayAbilities/Private/Abilities/Async/AbilityAsync.cpp
@@ -0,0 +1,38 @@
+// Copyright Epic Games, Inc. All Rights Reserved.
+
+#include "Abilities/Async/AbilityAsync.h"
+#include "AbilitySystemGlobals.h"
+#include "AbilitySystemComponent.h"
+
+void UAbilityAsync::EndAction()
+{
+	// Clear our ASC so it won't broadcast delegates
+	SetAbilitySystemComponent(nullptr);
+	SetReadyToDestroy();
+}
+
+bool UAbilityAsync::ShouldBroadcastDelegates() const
+{
+	// By default, broadcast if our ASC is valid
+	if (GetAbilitySystemComponent() != nullptr)
+	{
+		return true;
+	}
+
+	return false;
+}
+
+UAbilitySystemComponent* UAbilityAsync::GetAbilitySystemComponent() const
+{
+	return AbilitySystemComponent.Get();
+}
+
+void UAbilityAsync::SetAbilitySystemComponent(UAbilitySystemComponent* InAbilitySystemComponent)
+{
+	AbilitySystemComponent = InAbilitySystemComponent;
+}
+
+void UAbilityAsync::SetAbilityActor(AActor* InActor)
+{
+	AbilitySystemComponent = UAbilitySystemGlobals::GetAbilitySystemComponentFromActor(InActor);
+}
diff --git a/Engine/Plugins/Runtime/GameplayAbilities/Source/GameplayAbilities/Private/Abilities/Async/AbilityAsync_WaitGameplayEffectApplied.cpp b/Engine/Plugins/Runtime/GameplayAbilities/Source/GameplayAbilities/Private/Abilities/Async/AbilityAsync_WaitGameplayEffectApplied.cpp
new file mode 100644
--- /dev/null
+++ b/Engine/Plugins/Runtime/GameplayAbilities/Source/GameplayAbilities/Private/Abilities/Async/AbilityAsync_WaitGameplayEffectApplied.cpp
@@ -0,0 +1,89 @@
+// Copyright Epic Games, Inc. All Rights Reserved.
+
+#include "Abilities/Async/AbilityAsync_WaitGameplayEffectApplied.h"
+#include "AbilitySystemGlobals.h"
+#include "AbilitySystemComponent.h"
+
+UAbilityAsync_WaitGameplayEffectApplied* UAbilityAsync_WaitGameplayEffectApplied::WaitGameplayEffectAppliedToActor(AActor* TargetActor, const FGameplayTargetDataFilterHandle SourceFilter, FGameplayTagRequirements SourceTagRequirements, FGameplayTagRequirements TargetTagRequirements, bool TriggerOnce, bool ListenForPeriodicEffect)
+{
+	UAbilityAsync_WaitGameplayEffectApplied* MyObj = NewObject<UAbilityAsync_WaitGameplayEffectApplied>();
+	MyObj->SetAbilityActor(TargetActor);
+	MyObj->Filter = SourceFilter;
+	MyObj->SourceTagRequirements = SourceTagRequirements;
+	MyObj->TargetTagRequirements = TargetTagRequirements;
+	MyObj->TriggerOnce = TriggerOnce;
+	MyObj->ListenForPeriodicEffects = ListenForPeriodicEffect;
+	return MyObj;
+}
+
+void UAbilityAsync_WaitGameplayEffectApplied::Activate()
+{
+	Super::Activate();
+
+	if (UAbilitySystemComponent* ASC = GetAbilitySystemComponent())
+	{
+		OnApplyGameplayEffectCallbackDelegateHandle = ASC->OnGameplayEffectAppliedDelegateToSelf.AddUObject(this, &UAbilityAsync_WaitGameplayEffectApplied::OnApplyGameplayEffectCallback);
+		if (ListenForPeriodicEffects)
+		{
+			OnPeriodicGameplayEffectExecuteCallbackDelegateHandle = ASC->OnPeriodicGameplayEffectExecuteDelegateOnSelf.AddUObject(this, &UAbilityAsync_WaitGameplayEffectApplied::OnApplyGameplayEffectCallback);
+		}
+	}
+	else
+	{
+		EndAction();
+	}
+}
+
+void UAbilityAsync_WaitGameplayEffectApplied::OnApplyGameplayEffectCallback(UAbilitySystemComponent* Target, const FGameplayEffectSpec& SpecApplied, FActiveGameplayEffectHandle ActiveHandle)
+{
+	AActor* AvatarActor = Target ? Target->GetAvatarActor_Direct() : nullptr;
+
+	if (!Filter.FilterPassesForActor(AvatarActor))
+	{
+		return;
+	}
+	if (!SourceTagRequirements.RequirementsMet(*SpecApplied.CapturedSourceTags.GetAggregatedTags()))
+	{
+		return;
+	}
+	if (!TargetTagRequirements.RequirementsMet(*SpecApplied.CapturedTargetTags.GetAggregatedTags()))
+	{
+		return;
+	}
+
+	if (bLocked)
+	{
+		ABILITY_LOG(Error, TEXT("WaitGameplayEffectApplied recursion detected. Action: %s. Applied Spec: %s. This could cause an infinite loop! Ignoring"), *GetPathName(), *SpecApplied.ToSimpleString());
+		return;
+	}
+
+	FGameplayEffectSpecHandle SpecHandle(new FGameplayEffectSpec(SpecApplied));
+
+	if (ShouldBroadcastDelegates())
+	{
+		TGuardValue<bool> GuardValue(bLocked, true);
+		OnApplied.Broadcast(AvatarActor, SpecHandle, ActiveHandle);
+
+		if (TriggerOnce)
+		{
+			EndAction();
+		}
+	}
+	else
+	{
+		EndAction();
+	}
+}
+
+void UAbilityAsync_WaitGameplayEffectApplied::EndAction()
+{
+	if (UAbilitySystemComponent* ASC = GetAbilitySystemComponent())
+	{
+		if (OnPeriodicGameplayEffectExecuteCallbackDelegateHandle.IsValid())
+		{
+			ASC->OnGameplayEffectAppliedDelegateToTarget.Remove(OnPeriodicGameplayEffectExecuteCallbackDelegateHandle);
+		}
+		ASC->OnGameplayEffectAppliedDelegateToSelf.Remove(OnApplyGameplayEffectCallbackDelegateHandle);
+	}
+	Super::EndAction();
+}
diff --git a/Engine/Plugins/Runtime/GameplayAbilities/Source/GameplayAbilities/Private/Abilities/Async/AbilityAsync_WaitGameplayEvent.cpp b/Engine/Plugins/Runtime/GameplayAbilities/Source/GameplayAbilities/Private/Abilities/Async/AbilityAsync_WaitGameplayEvent.cpp
new file mode 100644
--- /dev/null
+++ b/Engine/Plugins/Runtime/GameplayAbilities/Source/GameplayAbilities/Private/Abilities/Async/AbilityAsync_WaitGameplayEvent.cpp
@@ -0,0 +1,79 @@
+// Copyright Epic Games, Inc. All Rights Reserved.
+
+#include "Abilities/Async/AbilityAsync_WaitGameplayEvent.h"
+#include "AbilitySystemGlobals.h"
+#include "AbilitySystemComponent.h"
+
+UAbilityAsync_WaitGameplayEvent* UAbilityAsync_WaitGameplayEvent::WaitGameplayEventToActor(AActor* TargetActor, FGameplayTag EventTag, bool OnlyTriggerOnce, bool OnlyMatchExact)
+{
+	UAbilityAsync_WaitGameplayEvent* MyObj = NewObject<UAbilityAsync_WaitGameplayEvent>();
+	MyObj->SetAbilityActor(TargetActor);
+	MyObj->Tag = EventTag;
+	MyObj->OnlyTriggerOnce = OnlyTriggerOnce;
+	MyObj->OnlyMatchExact = OnlyMatchExact;
+	return MyObj;
+}
+
+void UAbilityAsync_WaitGameplayEvent::Activate()
+{
+	Super::Activate();
+
+	UAbilitySystemComponent* ASC = GetAbilitySystemComponent();
+	if (ASC)
+	{
+		if (OnlyMatchExact)
+		{
+			MyHandle = ASC->GenericGameplayEventCallbacks.FindOrAdd(Tag).AddUObject(this, &UAbilityAsync_WaitGameplayEvent::GameplayEventCallback);
+		}
+		else
+		{
+			MyHandle = ASC->AddGameplayEventTagContainerDelegate(FGameplayTagContainer(Tag), FGameplayEventTagMulticastDelegate::FDelegate::CreateUObject(this, &UAbilityAsync_WaitGameplayEvent::GameplayEventContainerCallback));
+		}
+	}
+	else
+	{
+		EndAction();
+	}
+}
+
+void UAbilityAsync_WaitGameplayEvent::GameplayEventCallback(const FGameplayEventData* Payload)
+{
+	GameplayEventContainerCallback(Tag, Payload);
+}
+
+void UAbilityAsync_WaitGameplayEvent::GameplayEventContainerCallback(FGameplayTag MatchingTag, const FGameplayEventData* Payload)
+{
+	if (ShouldBroadcastDelegates())
+	{
+		FGameplayEventData TempPayload = *Payload;
+		TempPayload.EventTag = MatchingTag;
+		EventReceived.Broadcast(MoveTemp(TempPayload));
+		
+		if (OnlyTriggerOnce)
+		{
+			EndAction();
+		}
+	}
+	else
+	{
+		EndAction();
+	}
+}
+
+void UAbilityAsync_WaitGameplayEvent::EndAction()
+{
+	UAbilitySystemComponent* ASC = GetAbilitySystemComponent();
+	if (ASC && MyHandle.IsValid())
+	{
+		if (OnlyMatchExact)
+		{
+			ASC->GenericGameplayEventCallbacks.FindOrAdd(Tag).Remove(MyHandle);
+		}
+		else
+		{
+			ASC->RemoveGameplayEventTagContainerDelegate(FGameplayTagContainer(Tag), MyHandle);
+		}
+
+	}
+	Super::EndAction();
+}
diff --git a/Engine/Plugins/Runtime/GameplayAbilities/Source/GameplayAbilities/Private/Abilities/Async/AbilityAsync_WaitGameplayTag.cpp b/Engine/Plugins/Runtime/GameplayAbilities/Source/GameplayAbilities/Private/Abilities/Async/AbilityAsync_WaitGameplayTag.cpp
new file mode 100644
--- /dev/null
+++ b/Engine/Plugins/Runtime/GameplayAbilities/Source/GameplayAbilities/Private/Abilities/Async/AbilityAsync_WaitGameplayTag.cpp
@@ -0,0 +1,97 @@
+// Copyright Epic Games, Inc. All Rights Reserved.
+
+#include "Abilities/Async/AbilityAsync_WaitGameplayTag.h"
+#include "AbilitySystemGlobals.h"
+#include "AbilitySystemComponent.h"
+
+void UAbilityAsync_WaitGameplayTag::Activate()
+{
+	Super::Activate();
+
+	check(TargetCount >= 0);
+	UAbilitySystemComponent* ASC = GetAbilitySystemComponent();
+	if (ASC && ShouldBroadcastDelegates())
+	{
+		MyHandle = ASC->RegisterGameplayTagEvent(Tag).AddUObject(this, &UAbilityAsync_WaitGameplayTag::GameplayTagCallback);
+
+		if ((TargetCount == 0 && ASC->GetTagCount(Tag) == 0) || (TargetCount > 0 && ASC->GetTagCount(Tag) >= TargetCount))
+		{
+			BroadcastDelegate();
+			if (OnlyTriggerOnce)
+			{
+				EndAction();
+			}
+		}
+	}
+	else
+	{
+		EndAction();
+	}
+}
+
+void UAbilityAsync_WaitGameplayTag::GameplayTagCallback(const FGameplayTag InTag, int32 NewCount)
+{
+	if (NewCount == TargetCount)
+	{
+		if (ShouldBroadcastDelegates())
+		{
+			BroadcastDelegate();
+			if (OnlyTriggerOnce)
+			{
+				EndAction();
+			}
+		}
+		else
+		{
+			EndAction();
+		}
+	}
+}
+
+void UAbilityAsync_WaitGameplayTag::BroadcastDelegate()
+{
+}
+
+void UAbilityAsync_WaitGameplayTag::EndAction()
+{
+	UAbilitySystemComponent* ASC = GetAbilitySystemComponent();
+	if (ASC && MyHandle.IsValid())
+	{
+		ASC->RegisterGameplayTagEvent(Tag).Remove(MyHandle);
+	}
+	Super::EndAction();
+}
+
+
+UAbilityAsync_WaitGameplayTagAdded* UAbilityAsync_WaitGameplayTagAdded::WaitGameplayTagAddToActor(AActor* TargetActor, FGameplayTag Tag, bool OnlyTriggerOnce/*=false*/)
+{
+	UAbilityAsync_WaitGameplayTagAdded* MyObj = NewObject<UAbilityAsync_WaitGameplayTagAdded>();
+	MyObj->SetAbilityActor(TargetActor);
+	MyObj->Tag = Tag;
+	MyObj->OnlyTriggerOnce = OnlyTriggerOnce;
+	MyObj->TargetCount = 1;
+
+	return MyObj;
+}
+
+void UAbilityAsync_WaitGameplayTagAdded::BroadcastDelegate()
+{
+	Added.Broadcast();
+}
+
+
+UAbilityAsync_WaitGameplayTagRemoved* UAbilityAsync_WaitGameplayTagRemoved::WaitGameplayTagRemoveFromActor(AActor* TargetActor, FGameplayTag Tag, bool OnlyTriggerOnce/*=false*/)
+{
+	UAbilityAsync_WaitGameplayTagRemoved* MyObj = NewObject<UAbilityAsync_WaitGameplayTagRemoved>();
+	MyObj->SetAbilityActor(TargetActor);
+	MyObj->Tag = Tag;
+	MyObj->OnlyTriggerOnce = OnlyTriggerOnce;
+	MyObj->TargetCount = 0;
+
+	return MyObj;
+}
+
+void UAbilityAsync_WaitGameplayTagRemoved::BroadcastDelegate()
+{
+	Removed.Broadcast();
+}
diff --git a/Engine/Plugins/Runtime/GameplayAbilities/Source/GameplayAbilities/Public/Abilities/Async/AbilityAsync.h b/Engine/Plugins/Runtime/GameplayAbilities/Source/GameplayAbilities/Public/Abilities/Async/AbilityAsync.h
new file mode 100644
--- /dev/null
+++ b/Engine/Plugins/Runtime/GameplayAbilities/Source/GameplayAbilities/Public/Abilities/Async/AbilityAsync.h
@@ -0,0 +1,42 @@
+// Copyright Epic Games, Inc. All Rights Reserved.
+#pragma once
+
+#include "CoreMinimal.h"
+#include "Kismet/BlueprintAsyncActionBase.h"
+#include "GameplayEffectTypes.h"
+#include "Abilities/GameplayAbilityTypes.h"
+#include "AbilityAsync.generated.h"
+
+class UAbilitySystemComponent;
+
+/**
+ *	AbilityAsync is a base class for ability-specific BlueprintAsyncActions.
+ *  These are similar to ability tasks, but they can be executed from any blueprint like an actor and are not tied to a specific ability lifespan.
+ */
+
+UCLASS(Abstract, meta = (ExposedAsyncProxy = AsyncAction))
+class GAMEPLAYABILITIES_API UAbilityAsync : public UBlueprintAsyncActionBase
+{
+	GENERATED_BODY()
+public:
+
+	/** Explicitly end the action, will disable any callbacks and allow action to be deleted */
+	UFUNCTION(BlueprintCallable, Category = "Ability|Async")
+	virtual void EndAction();
+
+	/** This should be called prior to broadcasting delegates back into the event graph, this ensures action is still valid */
+	virtual bool ShouldBroadcastDelegates() const;
+
+	/** Returns the ability system component this action is bound to */
+	virtual UAbilitySystemComponent* GetAbilitySystemComponent() const;
+
+	/** Sets the bound component, returns true on success */
+	virtual void SetAbilitySystemComponent(UAbilitySystemComponent* InAbilitySystemComponent);
+
+	/** Sets the bound component by searching actor, returns true on success */
+	virtual void SetAbilityActor(AActor* InActor);
+
+private:
+	TWeakObjectPtr<UAbilitySystemComponent> AbilitySystemComponent;
+
+};
diff --git a/Engine/Plugins/Runtime/GameplayAbilities/Source/GameplayAbilities/Public/Abilities/Async/AbilityAsync_WaitGameplayEffectApplied.h b/Engine/Plugins/Runtime/GameplayAbilities/Source/GameplayAbilities/Public/Abilities/Async/AbilityAsync_WaitGameplayEffectApplied.h
new file mode 100644
--- /dev/null
+++ b/Engine/Plugins/Runtime/GameplayAbilities/Source/GameplayAbilities/Public/Abilities/Async/AbilityAsync_WaitGameplayEffectApplied.h
@@ -0,0 +1,52 @@
+// Copyright Epic Games, Inc. All Rights Reserved.
+#pragma once
+
+#include "CoreMinimal.h"
+#include "AbilityAsync.h"
+#include "Abilities/GameplayAbilityTargetDataFilter.h"
+#include "GameplayEffectTypes.h"
+#include "Delegates/IDelegateInstance.h"
+#include "AbilityAsync_WaitGameplayEffectApplied.generated.h"
+
+class UAbilitySystemComponent;
+
+/**
+ * This action listens for specific gameplay effect applications based off specified tags. 
+ * Effects themselves are not replicated; rather the tags they grant, the attributes they 
+ * change, and the gameplay cues they emit are replicated.
+ * This will only listen for local server or predicted client effects.
+ */
+UCLASS()
+class GAMEPLAYABILITIES_API UAbilityAsync_WaitGameplayEffectApplied : public UAbilityAsync
+{
+	GENERATED_BODY()
+
+public:
+	/**
+	 * Wait until a GameplayEffect is applied to a target actor
+	 * If TriggerOnce is true, this action will only activate one time. Otherwise it will return every time a GE is applied that meets the requirements over the life of the ability
+	 */
+	UFUNCTION(BlueprintCallable, Category = "Ability|Async", meta = (DefaultToSelf = "TargetActor", BlueprintInternalUseOnly = "TRUE"))
+	static UAbilityAsync_WaitGameplayEffectApplied* WaitGameplayEffectAppliedToActor(AActor* TargetActor, const FGameplayTargetDataFilterHandle SourceFilter, FGameplayTagRequirements SourceTagRequirements, FGameplayTagRequirements TargetTagRequirements, bool TriggerOnce = false, bool ListenForPeriodicEffect = false);
+
+	DECLARE_DYNAMIC_MULTICAST_DELEGATE_ThreeParams(FOnAppliedDelegate, AActor*, Source, FGameplayEffectSpecHandle, SpecHandle, FActiveGameplayEffectHandle, ActiveHandle);
+	UPROPERTY(BlueprintAssignable)
+	FOnAppliedDelegate OnApplied;
+
+protected:
+	virtual void Activate() override;
+	virtual void EndAction() override;
+
+	void OnApplyGameplayEffectCallback(UAbilitySystemComponent* Target, const FGameplayEffectSpec& SpecApplied, FActiveGameplayEffectHandle ActiveHandle);
+
+	FGameplayTargetDataFilterHandle Filter;
+	FGameplayTagRequirements SourceTagRequirements;
+	FGameplayTagRequirements TargetTagRequirements;
+
+	bool TriggerOnce = false;
+	bool ListenForPeriodicEffects = false;
+
+	FDelegateHandle OnApplyGameplayEffectCallbackDelegateHandle;
+	FDelegateHandle OnPeriodicGameplayEffectExecuteCallbackDelegateHandle;
+	bool bLocked = false;
+};
diff --git a/Engine/Plugins/Runtime/GameplayAbilities/Source/GameplayAbilities/Public/Abilities/Async/AbilityAsync_WaitGameplayEvent.h b/Engine/Plugins/Runtime/GameplayAbilities/Source/GameplayAbilities/Public/Abilities/Async/AbilityAsync_WaitGameplayEvent.h
new file mode 100644
--- /dev/null
+++ b/Engine/Plugins/Runtime/GameplayAbilities/Source/GameplayAbilities/Public/Abilities/Async/AbilityAsync_WaitGameplayEvent.h
@@ -0,0 +1,43 @@
+// Copyright Epic Games, Inc. All Rights Reserved.
+#pragma once
+
+#include "CoreMinimal.h"
+#include "AbilityAsync.h"
+#include "GameplayEffectTypes.h"
+#include "Abilities/GameplayAbilityTypes.h"
+#include "AbilityAsync_WaitGameplayEvent.generated.h"
+
+class UAbilitySystemComponent;
+
+UCLASS()
+class GAMEPLAYABILITIES_API UAbilityAsync_WaitGameplayEvent : public UAbilityAsync
+{
+	GENERATED_BODY()
+
+public:
+	/**
+	 * Wait until the specified gameplay tag event is triggered on a target ability system component
+	 * It will keep listening as long as OnlyTriggerOnce = false
+	 * If OnlyMatchExact = false it will trigger for nested tags
+	 */
+	UFUNCTION(BlueprintCallable, Category = "Ability|Async", meta = (DefaultToSelf = "TargetActor", BlueprintInternalUseOnly = "TRUE"))
+	static UAbilityAsync_WaitGameplayEvent* WaitGameplayEventToActor(AActor* TargetActor, FGameplayTag EventTag, bool OnlyTriggerOnce = false, bool OnlyMatchExact = true);
+
+	DECLARE_DYNAMIC_MULTICAST_DELEGATE_OneParam(FEventReceivedDelegate, FGameplayEventData, Payload);
+
+	UPROPERTY(BlueprintAssignable)
+	FEventReceivedDelegate EventReceived;
+
+protected:
+	virtual void Activate() override;
+	virtual void EndAction() override;
+
+	virtual void GameplayEventCallback(const FGameplayEventData* Payload);
+	virtual void GameplayEventContainerCallback(FGameplayTag MatchingTag, const FGameplayEventData* Payload);
+
+	FGameplayTag Tag;
+	bool OnlyTriggerOnce = false;
+	bool OnlyMatchExact = false;
+
+	FDelegateHandle MyHandle;
+};
diff --git a/Engine/Plugins/Runtime/GameplayAbilities/Source/GameplayAbilities/Public/Abilities/Async/AbilityAsync_WaitGameplayTag.h b/Engine/Plugins/Runtime/GameplayAbilities/Source/GameplayAbilities/Public/Abilities/Async/AbilityAsync_WaitGameplayTag.h
new file mode 100644
--- /dev/null
+++ b/Engine/Plugins/Runtime/GameplayAbilities/Source/GameplayAbilities/Public/Abilities/Async/AbilityAsync_WaitGameplayTag.h
@@ -0,0 +1,69 @@
+// Copyright Epic Games, Inc. All Rights Reserved.
+#pragma once
+
+#include "CoreMinimal.h"
+#include "AbilityAsync.h"
+#include "GameplayEffectTypes.h"
+#include "Abilities/GameplayAbilityTypes.h"
+#include "AbilityAsync_WaitGameplayTag.generated.h"
+
+class UAbilitySystemComponent;
+
+UCLASS(Abstract)
+class GAMEPLAYABILITIES_API UAbilityAsync_WaitGameplayTag : public UAbilityAsync
+{
+	GENERATED_BODY()
+protected:
+
+	DECLARE_DYNAMIC_MULTICAST_DELEGATE(FAsyncWaitGameplayTagDelegate);
+
+	virtual void Activate() override;
+	virtual void EndAction() override;
+
+	virtual void GameplayTagCallback(const FGameplayTag Tag, int32 NewCount);
+	virtual void BroadcastDelegate();
+
+	int32 TargetCount = -1;
+	FGameplayTag Tag;
+	bool OnlyTriggerOnce = false;
+
+	FDelegateHandle MyHandle;
+};
+
+UCLASS()
+class GAMEPLAYABILITIES_API UAbilityAsync_WaitGameplayTagAdded : public UAbilityAsync_WaitGameplayTag
+{
+	GENERATED_BODY()
+public:
+
+	UPROPERTY(BlueprintAssignable)
+	FAsyncWaitGameplayTagDelegate Added;
+
+	/**
+	 * Wait until the specified gameplay tag is Added to Target Actor's ability component
+	 * If the tag is already present when this task is started, it will immediately broadcast the Added event. It will keep listening as long as OnlyTriggerOnce = false.
+	 */
+	UFUNCTION(BlueprintCallable, Category = "Ability|Tasks", meta = (DefaultToSelf = "TargetActor", BlueprintInternalUseOnly = "TRUE"))
+	static UAbilityAsync_WaitGameplayTagAdded* WaitGameplayTagAddToActor(AActor* TargetActor, FGameplayTag Tag, bool OnlyTriggerOnce=false);
+
+	virtual void BroadcastDelegate() override;
+};
+
+UCLASS()
+class GAMEPLAYABILITIES_API UAbilityAsync_WaitGameplayTagRemoved : public UAbilityAsync_WaitGameplayTag
+{
+	GENERATED_BODY()
+public:
+
+	UPROPERTY(BlueprintAssignable)
+	FAsyncWaitGameplayTagDelegate Removed;
+
+	/**
+	 * Wait until the specified gameplay tag is Removed from Target Actor's ability component
+	 * If the tag is not present when this task is started, it will immediately broadcast the Removed event. It will keep listening as long as OnlyTriggerOnce = false.
+	 */
+	UFUNCTION(BlueprintCallable, Category = "Ability|Tasks", meta = (DefaultToSelf = "TargetActor", BlueprintInternalUseOnly = "TRUE"))
+	static UAbilityAsync_WaitGameplayTagRemoved* WaitGameplayTagRemoveFromActor(AActor* TargetActor, FGameplayTag Tag, bool OnlyTriggerOnce=false);
+
+	virtual void BroadcastDelegate() override;
+};
\ No newline at end of file
diff --git a/Engine/Source/Runtime/Engine/Classes/Kismet/BlueprintAsyncActionBase.h b/Engine/Source/Runtime/Engine/Classes/Kismet/BlueprintAsyncActionBase.h
--- a/Engine/Source/Runtime/Engine/Classes/Kismet/BlueprintAsyncActionBase.h
+++ b/Engine/Source/Runtime/Engine/Classes/Kismet/BlueprintAsyncActionBase.h
@@ -18,7 +18,10 @@ class UGameInstance;
 UCLASS()
 class ENGINE_API UBlueprintAsyncActionBase : public UObject
 {
-	GENERATED_UCLASS_BODY()
+	GENERATED_BODY()
+public:
+	/** Default UObject constructor */
+	UBlueprintAsyncActionBase(const FObjectInitializer& ObjectInitializer = FObjectInitializer::Get());
 
 	/** Called to trigger the action once the delegates have been bound */
 	UFUNCTION(BlueprintCallable, meta=(BlueprintInternalUseOnly="true"))
