Fix ensure in abilities system if an ASC is registered and unregistered without being initialized, which can happen during level streaming

diff --git a/Engine/Plugins/Runtime/GameplayAbilities/Source/GameplayAbilities/Private/AbilitySystemComponent_Abilities.cpp b/Engine/Plugins/Runtime/GameplayAbilities/Source/GameplayAbilities/Private/AbilitySystemComponent_Abilities.cpp
--- a/Engine/Plugins/Runtime/GameplayAbilities/Source/GameplayAbilities/Private/AbilitySystemComponent_Abilities.cpp
+++ b/Engine/Plugins/Runtime/GameplayAbilities/Source/GameplayAbilities/Private/AbilitySystemComponent_Abilities.cpp
@@ -891,6 +891,12 @@ void UAbilitySystemComponent::DestroyActiveState()
 		// Avoid re-entrancy (ie if during CancelAbilities() an EndAbility callback destroys the Actor owning this ability system)
 		bDestroyActiveStateInitiated = true;
 
+		if (!HasBeenInitialized())
+		{
+			// If we're destroyed before abilities were added, don't try to remove them
+			return;
+		}
+
 		// Cancel all abilities before we are destroyed.
 		FGameplayAbilityActorInfo* ActorInfo = AbilityActorInfo.Get();
 		
