Add safety checks for FGameplayTagBlueprintPropertyMap to avoid crashing on trashed memory if it was unable to unregister or was copied incorrectly
Add ApplyCurrentTags to apply any changes that happened before delegates were bound

diff --git a/Engine/Plugins/Runtime/GameplayAbilities/Source/GameplayAbilities/Private/GameplayEffectTypes.cpp b/Engine/Plugins/Runtime/GameplayAbilities/Source/GameplayAbilities/Private/GameplayEffectTypes.cpp
--- a/Engine/Plugins/Runtime/GameplayAbilities/Source/GameplayAbilities/Private/GameplayEffectTypes.cpp
+++ b/Engine/Plugins/Runtime/GameplayAbilities/Source/GameplayAbilities/Private/GameplayEffectTypes.cpp
@@ -536,6 +536,12 @@ FGameplayTagBlueprintPropertyMap::FGameplayTagBlueprintPropertyMap()
 {
 }
 
+FGameplayTagBlueprintPropertyMap::FGameplayTagBlueprintPropertyMap(const FGameplayTagBlueprintPropertyMap& Other)
+{
+	ensureMsgf(Other.CachedOwner.IsExplicitlyNull(), TEXT("FGameplayTagBlueprintPropertyMap cannot be used inside an array or other container that is copied after register!"));
+	PropertyMappings = Other.PropertyMappings;
+}
+
 FGameplayTagBlueprintPropertyMap::~FGameplayTagBlueprintPropertyMap()
 {
 	Unregister();
@@ -610,7 +616,7 @@ void FGameplayTagBlueprintPropertyMap::Initialize(UObject* Owner, UAbilitySystem
 	CachedOwner = Owner;
 	CachedASC = ASC;
 
-	FOnGameplayEffectTagCountChanged::FDelegate Delegate = FOnGameplayEffectTagCountChanged::FDelegate::CreateRaw(this, &FGameplayTagBlueprintPropertyMap::GameplayTagEventCallback);
+	FOnGameplayEffectTagCountChanged::FDelegate Delegate = FOnGameplayEffectTagCountChanged::FDelegate::CreateRaw(this, &FGameplayTagBlueprintPropertyMap::GameplayTagEventCallback, CachedOwner);
 
 	// Process array starting at the end so we can remove invalid entries.
 	for (int32 MappingIndex = (PropertyMappings.Num() - 1); MappingIndex >= 0; --MappingIndex)
@@ -656,8 +662,15 @@ void FGameplayTagBlueprintPropertyMap::Unregister()
 	CachedASC = nullptr;
 }
 
-void FGameplayTagBlueprintPropertyMap::GameplayTagEventCallback(const FGameplayTag Tag, int32 NewCount)
+void FGameplayTagBlueprintPropertyMap::GameplayTagEventCallback(const FGameplayTag Tag, int32 NewCount, TWeakObjectPtr<UObject> RegisteredOwner)
 {
+	// If the index and serial don't match with registered owner, the memory might be trashed so abort
+	if (!ensure(RegisteredOwner.HasSameIndexAndSerialNumber(CachedOwner)))
+	{
+		ABILITY_LOG(Error, TEXT("FGameplayTagBlueprintPropertyMap::GameplayTagEventCallback called with corrupted Owner!"));
+		return;
+	}
+
 	UObject* Owner = CachedOwner.Get();
 	if (!Owner)
 	{
@@ -687,6 +700,44 @@ void FGameplayTagBlueprintPropertyMap::GameplayTagEventCallback(const FGameplayT
 	}
 }
 
+void FGameplayTagBlueprintPropertyMap::ApplyCurrentTags()
+{
+	UObject* Owner = CachedOwner.Get();
+	if (!Owner)
+	{
+		ABILITY_LOG(Warning, TEXT("FGameplayTagBlueprintPropertyMap::ApplyCurrentTags called with an invalid Owner."));
+		return;
+	}
+
+	UAbilitySystemComponent* ASC = CachedASC.Get();
+	if (!ASC)
+	{
+		ABILITY_LOG(Warning, TEXT("FGameplayTagBlueprintPropertyMap::ApplyCurrentTags called with an invalid AbilitySystemComponent."));
+		return;
+	}
+
+	for (FGameplayTagBlueprintPropertyMapping& Mapping : PropertyMappings)
+	{
+		if (Mapping.PropertyToEdit.Get() && Mapping.TagToMap.IsValid())
+		{
+			int32 NewCount = ASC->GetTagCount(Mapping.TagToMap);
+			
+			if (const FBoolProperty* BoolProperty = CastField<const FBoolProperty>(Mapping.PropertyToEdit.Get()))
+			{
+				BoolProperty->SetPropertyValue_InContainer(Owner, NewCount > 0);
+			}
+			else if (const FIntProperty* IntProperty = CastField<const FIntProperty>(Mapping.PropertyToEdit.Get()))
+			{
+				IntProperty->SetPropertyValue_InContainer(Owner, NewCount);
+			}
+			else if (const FFloatProperty* FloatProperty = CastField<const FFloatProperty>(Mapping.PropertyToEdit.Get()))
+			{
+				FloatProperty->SetPropertyValue_InContainer(Owner, (float)NewCount);
+			}
+		}
+	}
+}
+
 bool FGameplayTagBlueprintPropertyMap::IsPropertyTypeValid(const FProperty* Property) const
 {
 	check(Property);
diff --git a/Engine/Plugins/Runtime/GameplayAbilities/Source/GameplayAbilities/Public/GameplayEffectTypes.h b/Engine/Plugins/Runtime/GameplayAbilities/Source/GameplayAbilities/Public/GameplayEffectTypes.h
--- a/Engine/Plugins/Runtime/GameplayAbilities/Source/GameplayAbilities/Public/GameplayEffectTypes.h
+++ b/Engine/Plugins/Runtime/GameplayAbilities/Source/GameplayAbilities/Public/GameplayEffectTypes.h
@@ -1388,6 +1388,16 @@ struct GAMEPLAYABILITIES_API FGameplayTagBlueprintPropertyMapping
 
 public:
 
+	FGameplayTagBlueprintPropertyMapping() {}
+	FGameplayTagBlueprintPropertyMapping(const FGameplayTagBlueprintPropertyMapping& Other)
+	{
+		// Don't copy handle
+		TagToMap = Other.TagToMap;
+		PropertyToEdit = Other.PropertyToEdit;
+		PropertyName = Other.PropertyName;
+		PropertyGuid = Other.PropertyGuid;
+	}
+
 	/** Gameplay tag being counted. */
 	UPROPERTY(EditAnywhere, Category = GameplayTagBlueprintProperty)
 	FGameplayTag TagToMap;
@@ -1412,7 +1422,7 @@ public:
 /**
  * Struct used to manage gameplay tag blueprint property mappings.
  * It registers the properties with delegates on an ability system component.
- * This struct should not be used in containers (such as TArray) since it uses a raw pointer
+ * This struct can not be used in containers (such as TArray) since it uses a raw pointer
  * to bind the delegate and it's address could change causing an invalid binding.
  */
 USTRUCT()
@@ -1423,11 +1433,15 @@ struct GAMEPLAYABILITIES_API FGameplayTagBlueprintPropertyMap
 public:
 
 	FGameplayTagBlueprintPropertyMap();
+	FGameplayTagBlueprintPropertyMap(const FGameplayTagBlueprintPropertyMap& Other);
 	~FGameplayTagBlueprintPropertyMap();
 
 	/** Call this to initialize and bind the properties with the ability system component. */
 	void Initialize(UObject* Owner, class UAbilitySystemComponent* ASC);
 
+	/** Call to manually apply the current tag state, can handle cases where callbacks were skipped */
+	void ApplyCurrentTags();
+
 #if WITH_EDITOR
 	/** This can optionally be called in the owner's IsDataValid() for data validation. */
 	EDataValidationResult IsDataValid(UObject* ContainingAsset, TArray<FText>& ValidationErrors);
@@ -1437,7 +1451,7 @@ protected:
 
 	void Unregister();
 
-	void GameplayTagEventCallback(const FGameplayTag Tag, int32 NewCount);
+	void GameplayTagEventCallback(const FGameplayTag Tag, int32 NewCount, TWeakObjectPtr<UObject> RegisteredOwner);
 
 	bool IsPropertyTypeValid(const FProperty* Property) const;
 
