[NO NON-PII COMMIT MESSAGE]

diff --git a/Engine/Plugins/Runtime/GameplayAbilities/Source/GameplayAbilities/Private/GameplayCueInterface.cpp b/Engine/Plugins/Runtime/GameplayAbilities/Source/GameplayAbilities/Private/GameplayCueInterface.cpp
--- a/Engine/Plugins/Runtime/GameplayAbilities/Source/GameplayAbilities/Private/GameplayCueInterface.cpp
+++ b/Engine/Plugins/Runtime/GameplayAbilities/Source/GameplayAbilities/Private/GameplayCueInterface.cpp
@@ -130,7 +130,7 @@ void IGameplayCueInterface::HandleGameplayCue(UObject* Self, FGameplayTag Gamepl
 	bool bShouldContinue = true;
 	for (int32 FunctionIndex = 0; bShouldContinue && (FunctionIndex < FunctionList->Num()); ++FunctionIndex)
 	{
-		GameplayCueInterfacePrivate::FCueNameAndUFunction& CueFunctionPair = FunctionList->GetData()[FunctionIndex];
+		const GameplayCueInterfacePrivate::FCueNameAndUFunction& CueFunctionPair = (*FunctionList)[FunctionIndex];
 		UFunction* Func = CueFunctionPair.Func;
 		Parameters.MatchedTagName = CueFunctionPair.Tag;
 
