PR #7991: Bugfix GameplayAbilityTargetActor_Trace.cpp - [The replicated location of AGameplayAbilityTargetActor_GroundTrace in the simulated proxy is incorrect] (Contributed by [REDACTED_PERSON])

diff --git a/Engine/Plugins/Runtime/GameplayAbilities/Source/GameplayAbilities/Private/Abilities/GameplayAbilityTargetActor_Trace.cpp b/Engine/Plugins/Runtime/GameplayAbilities/Source/GameplayAbilities/Private/Abilities/GameplayAbilityTargetActor_Trace.cpp
--- a/Engine/Plugins/Runtime/GameplayAbilities/Source/GameplayAbilities/Private/Abilities/GameplayAbilityTargetActor_Trace.cpp
+++ b/Engine/Plugins/Runtime/GameplayAbilities/Source/GameplayAbilities/Private/Abilities/GameplayAbilityTargetActor_Trace.cpp
@@ -176,7 +176,7 @@ void AGameplayAbilityTargetActor_Trace::StartTargeting(UGameplayAbility* InAbili
 void AGameplayAbilityTargetActor_Trace::Tick(float DeltaSeconds)
 {
 	// very temp - do a mostly hardcoded trace from the source actor
-	if (SourceActor)
+	if (SourceActor && SourceActor->GetLocalRole() != ENetRole::ROLE_SimulatedProxy)
 	{
 		FHitResult HitResult = PerformTrace(SourceActor);
 		FVector EndPoint = HitResult.Component.IsValid() ? HitResult.ImpactPoint : HitResult.TraceEnd;
