Guarded a potential crash in UGameplayAbility::GetAbilityLevel if the ASC pointed to by the actor info has been destroyed

diff --git a/Engine/Plugins/Runtime/GameplayAbilities/Source/GameplayAbilities/Private/Abilities/GameplayAbility.cpp b/Engine/Plugins/Runtime/GameplayAbilities/Source/GameplayAbilities/Private/Abilities/GameplayAbility.cpp
--- a/Engine/Plugins/Runtime/GameplayAbilities/Source/GameplayAbilities/Private/Abilities/GameplayAbility.cpp
+++ b/Engine/Plugins/Runtime/GameplayAbilities/Source/GameplayAbilities/Private/Abilities/GameplayAbility.cpp
@@ -1463,10 +1463,13 @@ int32 UGameplayAbility::GetAbilityLevel() const
 	return GetAbilityLevel(CurrentSpecHandle, CurrentActorInfo);
 }
 
-/** Returns current ability level for non instanced abilities. You m ust call this version in these contexts! */
+/** Returns current ability level for non instanced abilities. You must call this version in these contexts! */
 int32 UGameplayAbility::GetAbilityLevel(FGameplayAbilitySpecHandle Handle, const FGameplayAbilityActorInfo* ActorInfo) const
 {
-	FGameplayAbilitySpec* Spec = ActorInfo->AbilitySystemComponent->FindAbilitySpecFromHandle(Handle);
+	check(ActorInfo);
+	UAbilitySystemComponent* const AbilitySystemComponent = ActorInfo->AbilitySystemComponent.Get();
+
+	FGameplayAbilitySpec* Spec = AbilitySystemComponent ? AbilitySystemComponent->FindAbilitySpecFromHandle(Handle) : nullptr;
 	
 	if (Spec)
 	{
