Back out CL 15789933 as it was causing abilities to not get properly removed on ASC destruction
I will find a better way to deal with the level streaming case later

diff --git a/Engine/Plugins/Runtime/GameplayAbilities/Source/GameplayAbilities/Private/AbilitySystemComponent_Abilities.cpp b/Engine/Plugins/Runtime/GameplayAbilities/Source/GameplayAbilities/Private/AbilitySystemComponent_Abilities.cpp
--- a/Engine/Plugins/Runtime/GameplayAbilities/Source/GameplayAbilities/Private/AbilitySystemComponent_Abilities.cpp
+++ b/Engine/Plugins/Runtime/GameplayAbilities/Source/GameplayAbilities/Private/AbilitySystemComponent_Abilities.cpp
@@ -924,12 +924,6 @@ void UAbilitySystemComponent::DestroyActiveState()
 		// Avoid re-entrancy (ie if during CancelAbilities() an EndAbility callback destroys the Actor owning this ability system)
 		bDestroyActiveStateInitiated = true;
 
-		if (!HasBeenInitialized())
-		{
-			// If we're destroyed before abilities were added, don't try to remove them
-			return;
-		}
-
 		// Cancel all abilities before we are destroyed.
 		FGameplayAbilityActorInfo* ActorInfo = AbilityActorInfo.Get();
 		
