Fix a bug where ability tasks do not broadcast their end task delegates when the ability ends.

diff --git a/Engine/Plugins/Runtime/GameplayAbilities/Source/GameplayAbilities/Private/Abilities/GameplayAbility.cpp b/Engine/Plugins/Runtime/GameplayAbilities/Source/GameplayAbilities/Private/Abilities/GameplayAbility.cpp
--- a/Engine/Plugins/Runtime/GameplayAbilities/Source/GameplayAbilities/Private/Abilities/GameplayAbility.cpp
+++ b/Engine/Plugins/Runtime/GameplayAbilities/Source/GameplayAbilities/Private/Abilities/GameplayAbility.cpp
@@ -636,7 +636,6 @@ void UGameplayAbility::EndAbility(const FGameplayAbilitySpecHandle Handle, const
 
 		if (GetInstancingPolicy() != EGameplayAbilityInstancingPolicy::NonInstanced)
 		{
-			bIsActive = false;
 			bIsAbilityEnding = false;
 		}
 
@@ -649,7 +648,13 @@ void UGameplayAbility::EndAbility(const FGameplayAbilitySpecHandle Handle, const
 				Task->TaskOwnerEnded();
 			}
 		}
-		ActiveTasks.Reset();	// Empty the array but dont resize memory, since this object is probably going to be destroyed very soon anyways.
+		ActiveTasks.Reset();	// Empty the array but don't resize memory, since this object is probably going to be destroyed very soon anyways.
+
+		// We mark the ability as no longer active after ending the tasks so they can broadcast delegates
+		if (GetInstancingPolicy() != EGameplayAbilityInstancingPolicy::NonInstanced)
+		{
+			bIsActive = false;
+		}
 
 		if (UAbilitySystemComponent* const AbilitySystemComponent = ActorInfo->AbilitySystemComponent.Get())
 		{
