Ability system component cleanup: removed no-longer-used actor info field

#rnx

diff --git a/Engine/Plugins/Runtime/GameplayAbilities/Source/GameplayAbilities/Public/Abilities/GameplayAbilityTypes.h b/Engine/Plugins/Runtime/GameplayAbilities/Source/GameplayAbilities/Public/Abilities/GameplayAbilityTypes.h
--- a/Engine/Plugins/Runtime/GameplayAbilities/Source/GameplayAbilities/Public/Abilities/GameplayAbilityTypes.h
+++ b/Engine/Plugins/Runtime/GameplayAbilities/Source/GameplayAbilities/Public/Abilities/GameplayAbilityTypes.h
@@ -159,10 +159,6 @@ struct GAMEPLAYABILITIES_API FGameplayAbilityActorInfo
 	UPROPERTY(BlueprintReadOnly, Category = "ActorInfo")
 	TWeakObjectPtr<USkeletalMeshComponent>	SkeletalMeshComponent;
 
-	/** Anim instance of the avatar actor. Often null */
-	UPROPERTY(BlueprintReadOnly, Category = "ActorInfo")
-	TWeakObjectPtr<UAnimInstance>	AnimInstance;
-
 	/** Movement component of the avatar actor. Often null */
 	UPROPERTY(BlueprintReadOnly, Category = "ActorInfo")
 	TWeakObjectPtr<UMovementComponent>	MovementComponent;
