Sequencer: Missing transactionals, which could cause null sections on undo. These haven't caused any specific problems but there was one in Media Track that was missing recently that caused a crash for 4.27

diff --git a/Engine/Plugins/Compositing/Composure/Source/Composure/Private/MovieScene/MovieSceneComposureExportTrack.cpp b/Engine/Plugins/Compositing/Composure/Source/Composure/Private/MovieScene/MovieSceneComposureExportTrack.cpp
--- a/Engine/Plugins/Compositing/Composure/Source/Composure/Private/MovieScene/MovieSceneComposureExportTrack.cpp
+++ b/Engine/Plugins/Compositing/Composure/Source/Composure/Private/MovieScene/MovieSceneComposureExportTrack.cpp
@@ -49,7 +49,7 @@ FText UMovieSceneComposureExportTrack::GetDisplayName() const
 
 UMovieSceneSection* UMovieSceneComposureExportTrack::CreateNewSection()
 {
-	return NewObject<UMovieSceneComposureExportSection>(this);
+	return NewObject<UMovieSceneComposureExportSection>(this, NAME_None, RF_Transactional);
 }
 
 UMovieSceneComposureExportSection::UMovieSceneComposureExportSection(const FObjectInitializer& ObjectInitializer)
diff --git a/Engine/Plugins/MovieScene/CustomizableSequencerTracks/Source/CustomizableSequencerTracks/Private/SequencerTrackBP.cpp b/Engine/Plugins/MovieScene/CustomizableSequencerTracks/Source/CustomizableSequencerTracks/Private/SequencerTrackBP.cpp
--- a/Engine/Plugins/MovieScene/CustomizableSequencerTracks/Source/CustomizableSequencerTracks/Private/SequencerTrackBP.cpp
+++ b/Engine/Plugins/MovieScene/CustomizableSequencerTracks/Source/CustomizableSequencerTracks/Private/SequencerTrackBP.cpp
@@ -11,19 +11,19 @@ UMovieSceneSection* USequencerTrackBP::CreateNewSection()
 {
 	if (UClass* Class = DefaultSectionType.Get())
 	{
-		return NewObject<USequencerSectionBP>(this, Class);
+		return NewObject<USequencerSectionBP>(this, Class, NAME_None, RF_Transactional);
 	}
 
 	for (TSubclassOf<USequencerSectionBP> SupportedSection : SupportedSections)
 	{
 		if (UClass* Class = SupportedSection.Get())
 		{
-			return NewObject<USequencerSectionBP>(this, Class);
+			return NewObject<USequencerSectionBP>(this, Class, NAME_None, RF_Transactional);
 		}
 	}
 
 	ensureMsgf(false, TEXT("Track does not have any supported section types. Returning a base class to avoid crashing."));
-	return NewObject<UMovieSceneSection>(this);
+	return NewObject<UMovieSceneSection>(this, NAME_None, RF_Transactional);
 }
 
 void USequencerTrackBP::GetAssetRegistryTags(TArray<FAssetRegistryTag>& OutTags) const
diff --git a/Engine/Plugins/Runtime/GameplayAbilities/Source/GameplayAbilities/Private/Sequencer/MovieSceneGameplayCueTrack.cpp b/Engine/Plugins/Runtime/GameplayAbilities/Source/GameplayAbilities/Private/Sequencer/MovieSceneGameplayCueTrack.cpp
--- a/Engine/Plugins/Runtime/GameplayAbilities/Source/GameplayAbilities/Private/Sequencer/MovieSceneGameplayCueTrack.cpp
+++ b/Engine/Plugins/Runtime/GameplayAbilities/Source/GameplayAbilities/Private/Sequencer/MovieSceneGameplayCueTrack.cpp
@@ -23,7 +23,7 @@ bool UMovieSceneGameplayCueTrack::SupportsType(TSubclassOf<UMovieSceneSection> S
 
 UMovieSceneSection* UMovieSceneGameplayCueTrack::CreateNewSection()
 {
-	return NewObject<UMovieSceneGameplayCueSection>(this);
+	return NewObject<UMovieSceneGameplayCueSection>(this, NAME_None, RF_Transactional);
 }
 
 const TArray<UMovieSceneSection*>& UMovieSceneGameplayCueTrack::GetAllSections() const
diff --git a/Engine/Plugins/Runtime/GameplayAbilities/Source/GameplayAbilitiesEditor/Private/Sequencer/GameplayCueTrackEditor.cpp b/Engine/Plugins/Runtime/GameplayAbilities/Source/GameplayAbilitiesEditor/Private/Sequencer/GameplayCueTrackEditor.cpp
--- a/Engine/Plugins/Runtime/GameplayAbilities/Source/GameplayAbilitiesEditor/Private/Sequencer/GameplayCueTrackEditor.cpp
+++ b/Engine/Plugins/Runtime/GameplayAbilities/Source/GameplayAbilitiesEditor/Private/Sequencer/GameplayCueTrackEditor.cpp
@@ -205,7 +205,7 @@ void FGameplayCueTrackEditor::HandleAddSectionToTrack(UMovieSceneTrack* Track, U
 	Track->Modify();
 
 
-	UMovieSceneSection* NewSection = NewObject<UMovieSceneSection>(Track, SectionClass);
+	UMovieSceneSection* NewSection = NewObject<UMovieSceneSection>(Track, SectionClass, NAME_None, RF_Transactional);
 
 	if (SectionClass == UMovieSceneGameplayCueSection::StaticClass())
 	{
@@ -233,7 +233,7 @@ void FGameplayCueTrackEditor::HandleAddSectionToTrack(UMovieSceneTrack* Track, U
 
 void FGameplayCueTrackEditor::AddSectionToTrack(UMovieSceneTrack* Track, const TRange<FFrameNumber>& SectionTickRange, UClass* SectionClass)
 {
-	UMovieSceneSection* NewSection = NewObject<UMovieSceneSection>(Track, SectionClass);
+	UMovieSceneSection* NewSection = NewObject<UMovieSceneSection>(Track, SectionClass, NAME_None, RF_Transactional);
 	NewSection->SetRange(SectionTickRange);
 	Track->AddSection(*NewSection);
 }
diff --git a/Engine/Source/Editor/MovieSceneTools/Private/TrackEditors/EventTrackEditor.cpp b/Engine/Source/Editor/MovieSceneTools/Private/TrackEditors/EventTrackEditor.cpp
--- a/Engine/Source/Editor/MovieSceneTools/Private/TrackEditors/EventTrackEditor.cpp
+++ b/Engine/Source/Editor/MovieSceneTools/Private/TrackEditors/EventTrackEditor.cpp
@@ -317,7 +317,7 @@ void FEventTrackEditor::CreateNewSection(UMovieSceneTrack* Track, int32 RowIndex
 
 		FScopedTransaction Transaction(LOCTEXT("CreateNewSectionTransactionText", "Add Section"));
 
-		UMovieSceneSection* NewSection = NewObject<UMovieSceneSection>(Track, SectionType);
+		UMovieSceneSection* NewSection = NewObject<UMovieSceneSection>(Track, SectionType, NAME_None, RF_Transactional);
 		check(NewSection);
 
 		int32 OverlapPriority = 0;
