Duplicating CL#17316718

Fixed issue where global ability task count tracking was double-counting, resulting in errors when we were only halfway to the limit

diff --git a/Engine/Plugins/Runtime/GameplayAbilities/Source/GameplayAbilities/Private/Abilities/Tasks/AbilityTask.cpp b/Engine/Plugins/Runtime/GameplayAbilities/Source/GameplayAbilities/Private/Abilities/Tasks/AbilityTask.cpp
--- a/Engine/Plugins/Runtime/GameplayAbilities/Source/GameplayAbilities/Private/Abilities/Tasks/AbilityTask.cpp
+++ b/Engine/Plugins/Runtime/GameplayAbilities/Source/GameplayAbilities/Private/Abilities/Tasks/AbilityTask.cpp
@@ -9,9 +9,11 @@ int32 UAbilityTask::GlobalAbilityTaskCount = 0;
 UAbilityTask::UAbilityTask(const FObjectInitializer& ObjectInitializer)
 	: Super(ObjectInitializer)
 {
-	WaitStateBitMask = (uint8)EAbilityTaskWaitState::WaitingOnGame;
+	WaitStateBitMask = static_cast<uint8>(EAbilityTaskWaitState::WaitingOnGame);
 	bWasSuccessfullyDestroyed = false;
-	SET_DWORD_STAT(STAT_AbilitySystem_TaskCount, ++GlobalAbilityTaskCount);
+
+	++GlobalAbilityTaskCount;
+	SET_DWORD_STAT(STAT_AbilitySystem_TaskCount, GlobalAbilityTaskCount);
 	if (!ensure(GlobalAbilityTaskCount < 1000))
 	{
 		ABILITY_LOG(Warning, TEXT("Way too many AbilityTasks are currently active! %d. %s"), GlobalAbilityTaskCount, *GetClass()->GetName());
@@ -20,7 +22,10 @@ UAbilityTask::UAbilityTask(const FObjectInitializer& ObjectInitializer)
 
 void UAbilityTask::OnDestroy(bool bInOwnerFinished)
 {
-	SET_DWORD_STAT(STAT_AbilitySystem_TaskCount, --GlobalAbilityTaskCount);
+	checkf(GlobalAbilityTaskCount > 0, TEXT("Mismatched AbilityTask counting"));
+    --GlobalAbilityTaskCount;
+	SET_DWORD_STAT(STAT_AbilitySystem_TaskCount, GlobalAbilityTaskCount);
+
 	bWasSuccessfullyDestroyed = true;
 
 	Super::OnDestroy(bInOwnerFinished);
@@ -33,8 +38,9 @@ void UAbilityTask::BeginDestroy()
 	if (!bWasSuccessfullyDestroyed)
 	{
 		// this shouldn't happen, it means that ability was destroyed while being active, but we need to keep GlobalAbilityTaskCount in sync anyway
-
-		SET_DWORD_STAT(STAT_AbilitySystem_TaskCount, --GlobalAbilityTaskCount);
+		checkf(GlobalAbilityTaskCount > 0, TEXT("Mismatched AbilityTask counting"));
+		--GlobalAbilityTaskCount;
+		SET_DWORD_STAT(STAT_AbilitySystem_TaskCount, GlobalAbilityTaskCount);
 		bWasSuccessfullyDestroyed = true;
 	}
 }
