Ability System
Fix GE GCs playing when bRequireModifierSuccessToTriggerCues=true but all execution modifiers fail to apply. Previously this flag was only honored for GE modifiers and not modifiers that are applied during GE executions. This will make the flag consistent between modifiers and executions.

#rnx

diff --git a/Engine/Plugins/Runtime/GameplayAbilities/Source/GameplayAbilities/Private/GameplayEffect.cpp b/Engine/Plugins/Runtime/GameplayAbilities/Source/GameplayAbilities/Private/GameplayEffect.cpp
--- a/Engine/Plugins/Runtime/GameplayAbilities/Source/GameplayAbilities/Private/GameplayEffect.cpp
+++ b/Engine/Plugins/Runtime/GameplayAbilities/Source/GameplayAbilities/Private/GameplayEffect.cpp
@@ -2047,11 +2047,15 @@ void FActiveGameplayEffectsContainer::ExecuteActiveEffectsFrom(FGameplayEffectSp
 	//	Invoke GameplayCue events
 	// ------------------------------------------------------
 	
-	// If there are no modifiers or we don't require modifier success to trigger, we apply the GameplayCue. 
-	bool InvokeGameplayCueExecute = (SpecToUse.Modifiers.Num() == 0) || !Spec.Def->bRequireModifierSuccessToTriggerCues;
+	// If there are no modifiers or we don't require modifier success to trigger, we apply the GameplayCue.
+	const bool bHasModifiers = SpecToUse.Modifiers.Num() > 0;
+	const bool bHasExecutions = SpecToUse.Def->Executions.Num() > 0;
+	const bool bHasModifiersOrExecutions = bHasModifiers || bHasExecutions;
 
-	// If there are modifiers, we only want to invoke the GameplayCue if one of them went through (could be blocked by immunity or % chance roll)
-	if (SpecToUse.Modifiers.Num() > 0 && ModifierSuccessfullyExecuted)
+	// If there are no modifiers or we don't require modifier success to trigger, we apply the GameplayCue.
+	bool InvokeGameplayCueExecute = (!bHasModifiersOrExecutions) || !Spec.Def->bRequireModifierSuccessToTriggerCues;
+
+	if (bHasModifiersOrExecutions && ModifierSuccessfullyExecuted)
 	{
 		InvokeGameplayCueExecute = true;
 	}
