Ability System
-Remove unnecessary call to AddActors in FGameplayEffectContext::Duplicate which was causing duplicate entries of the actors to be created (heh)

#rnx

diff --git a/Engine/Plugins/Runtime/GameplayAbilities/Source/GameplayAbilities/Public/GameplayEffectTypes.h b/Engine/Plugins/Runtime/GameplayAbilities/Source/GameplayAbilities/Public/GameplayEffectTypes.h
--- a/Engine/Plugins/Runtime/GameplayAbilities/Source/GameplayAbilities/Public/GameplayEffectTypes.h
+++ b/Engine/Plugins/Runtime/GameplayAbilities/Source/GameplayAbilities/Public/GameplayEffectTypes.h
@@ -502,7 +502,6 @@ struct GAMEPLAYABILITIES_API FGameplayEffectContext
 	{
 		FGameplayEffectContext* NewContext = new FGameplayEffectContext();
 		*NewContext = *this;
-		NewContext->AddActors(Actors);
 		if (GetHitResult())
 		{
 			// Does a deep copy of the hit result
