When removing an ability that was granted by a GE don't unbind the input. This prevents a bug that can block abilities from ending and being removed.
When an ability is waiting to be removed do not allow inputs to activate it.

diff --git a/Engine/Plugins/Runtime/GameplayAbilities/Source/GameplayAbilities/Private/AbilitySystemComponent_Abilities.cpp b/Engine/Plugins/Runtime/GameplayAbilities/Source/GameplayAbilities/Private/AbilitySystemComponent_Abilities.cpp
--- a/Engine/Plugins/Runtime/GameplayAbilities/Source/GameplayAbilities/Private/AbilitySystemComponent_Abilities.cpp
+++ b/Engine/Plugins/Runtime/GameplayAbilities/Source/GameplayAbilities/Private/AbilitySystemComponent_Abilities.cpp
@@ -320,7 +320,6 @@ void UAbilitySystemComponent::SetRemoveAbilityOnEnd(FGameplayAbilitySpecHandle A
 		if (FoundSpec->IsActive())
 		{
 			FoundSpec->RemoveAfterActivation = true;
-			FoundSpec->InputID = INDEX_NONE;
 		}
 		else
 		{
@@ -1126,6 +1125,12 @@ bool UAbilitySystemComponent::TryActivateAbility(FGameplayAbilitySpecHandle Abil
 		return false;
 	}
 
+	// don't activate abilities that are waiting to be removed
+	if (Spec->PendingRemove || Spec->RemoveAfterActivation)
+	{
+		return false;
+	}
+
 	UGameplayAbility* Ability = Spec->Ability;
 
 	if (!Ability)
