Updating GameplayCueNotifies (SoundInfo) to work with MetaSounds
- Added struct for fleshing out implementing Parameter Interface parameters (currently just a stop trigger)
- Added an evaluation for USoundWaveProcedurals to ignore loop validation

diff --git a/Engine/Plugins/Runtime/GameplayAbilities/Source/GameplayAbilities/Private/GameplayCueNotifyTypes.cpp b/Engine/Plugins/Runtime/GameplayAbilities/Source/GameplayAbilities/Private/GameplayCueNotifyTypes.cpp
--- a/Engine/Plugins/Runtime/GameplayAbilities/Source/GameplayAbilities/Private/GameplayCueNotifyTypes.cpp
+++ b/Engine/Plugins/Runtime/GameplayAbilities/Source/GameplayAbilities/Private/GameplayCueNotifyTypes.cpp
@@ -18,6 +18,7 @@
 #include "GameFramework/PlayerController.h"
 #include "GameFramework/PlayerState.h"
 #include "GameFramework/Character.h"
+#include "Sound/SoundWaveProcedural.h"
 
 
 DEFINE_LOG_CATEGORY(LogGameplayCueNotify);
@@ -455,6 +456,12 @@ void FGameplayCueNotify_ParticleInfo::ValidateBurstAssets(UObject* ContainingAss
 //////////////////////////////////////////////////////////////////////////
 // FGameplayCueNotify_SoundInfo
 //////////////////////////////////////////////////////////////////////////
+
+FGameplayCueNotify_SoundParameterInterfaceInfo::FGameplayCueNotify_SoundParameterInterfaceInfo()
+{
+	StopTriggerName = TEXT("OnStop");
+}
+
 FGameplayCueNotify_SoundInfo::FGameplayCueNotify_SoundInfo()
 {
 	SoundCue = nullptr;
@@ -462,6 +469,7 @@ FGameplayCueNotify_SoundInfo::FGameplayCueNotify_SoundInfo()
 	LoopingFadeVolumeLevel = 0.0f;
 	bOverrideSpawnCondition = false;
 	bOverridePlacementInfo = false;
+	bUseSoundParameterInterface = false;
 }
 
 bool FGameplayCueNotify_SoundInfo::PlaySound(const FGameplayCueNotify_SpawnContext& SpawnContext, FGameplayCueNotify_SpawnResult& OutSpawnResult) const
@@ -526,7 +534,7 @@ void FGameplayCueNotify_SoundInfo::ValidateBurstAssets(UObject* ContainingAsset,
 #if WITH_EDITORONLY_DATA
 	if (SoundCue != nullptr)
 	{
-		if (SoundCue->IsLooping())
+		if (SoundCue->IsLooping() && (!SoundCue->IsA<USoundWaveProcedural>()))
 		{
 			ValidationErrors.Add(FText::Format(
 				LOCTEXT("SoundCue_ShouldNotLoop", "Sound [{0}] used in slot [{1}] for asset [{2}] is set to looping, but the slot is a one-shot (the instance will leak)."),
@@ -1093,10 +1101,16 @@ void FGameplayCueNotify_LoopingEffects::StopEffects(FGameplayCueNotify_SpawnResu
 		UAudioComponent* AudioComponent = SpawnResult.AudioComponents[SoundIndex];
 		if (AudioComponent)
 		{
-			if (LoopingSounds.IsValidIndex(SoundIndex))
+			const FGameplayCueNotify_SoundInfo* SoundInfo = &LoopingSounds[SoundIndex];
+
+			if(SoundInfo->bUseSoundParameterInterface)
 			{
-				const FGameplayCueNotify_SoundInfo* SoundInfo = &LoopingSounds[SoundIndex];
+				// Call the Stop Trigger by Name
+				AudioComponent->SetTriggerParameter(SoundInfo->SoundParameterInterfaceInfo.StopTriggerName);
+			}
 
+			if (LoopingSounds.IsValidIndex(SoundIndex))
+			{
 				if (SoundInfo->LoopingFadeOutDuration > 0.0f)
 				{
 					AudioComponent->FadeOut(SoundInfo->LoopingFadeOutDuration, SoundInfo->LoopingFadeVolumeLevel);
diff --git a/Engine/Plugins/Runtime/GameplayAbilities/Source/GameplayAbilities/Public/GameplayCueNotifyTypes.h b/Engine/Plugins/Runtime/GameplayAbilities/Source/GameplayAbilities/Public/GameplayCueNotifyTypes.h
--- a/Engine/Plugins/Runtime/GameplayAbilities/Source/GameplayAbilities/Public/GameplayCueNotifyTypes.h
+++ b/Engine/Plugins/Runtime/GameplayAbilities/Source/GameplayAbilities/Public/GameplayCueNotifyTypes.h
@@ -352,6 +352,24 @@ public:
 	uint32 bCastShadow : 1;
 };
 
+/**
+ * FGameplayCueNotify_SoundParameterInterfaceInfo
+ *
+ *	Properties that specify how to interface with the ISoundGeneratorParameterInterface
+ */
+USTRUCT(BlueprintType)
+struct FGameplayCueNotify_SoundParameterInterfaceInfo
+{
+	GENERATED_BODY()
+
+public:
+
+	FGameplayCueNotify_SoundParameterInterfaceInfo();
+
+	// The name of the stop trigger set via the parameter interface
+	UPROPERTY(EditDefaultsOnly, BlueprintReadOnly, Category = GameplayCueNotify)
+	FName StopTriggerName = TEXT("OnStop");
+};
 
 /**
  * FGameplayCueNotify_SoundInfo
@@ -381,7 +399,7 @@ public:
 	UPROPERTY(EditDefaultsOnly, BlueprintReadOnly, Category = GameplayCueNotify, Meta = (EditCondition = "bOverridePlacementInfo"))
 	FGameplayCueNotify_PlacementInfo PlacementInfoOverride;
 
-	// Sound cue to play.
+	// Sound to play.
 	UPROPERTY(EditDefaultsOnly, BlueprintReadOnly, Category = GameplayCueNotify)
 	USoundBase* SoundCue;
 
@@ -395,6 +413,10 @@ public:
 	UPROPERTY(EditDefaultsOnly, BlueprintReadOnly, Category = GameplayCueNotify)
 	float LoopingFadeVolumeLevel;
 
+	// Defines how to interface with the sound.
+	UPROPERTY(EditDefaultsOnly, BlueprintReadOnly, Category = GameplayCueNotify, Meta = (EditCondition = "bUseSoundParameterInterface"))
+	FGameplayCueNotify_SoundParameterInterfaceInfo SoundParameterInterfaceInfo;
+
 	// If enabled, use the spawn condition override and not the default one.
 	UPROPERTY(EditDefaultsOnly, BlueprintReadOnly, Category = GameplayCueNotify, Meta = (InlineEditConditionToggle))
 	uint32 bOverrideSpawnCondition : 1;
@@ -402,6 +424,11 @@ public:
 	// If enabled, use the placement info override and not the default one.
 	UPROPERTY(EditDefaultsOnly, BlueprintReadOnly, Category = GameplayCueNotify, Meta = (InlineEditConditionToggle))
 	uint32 bOverridePlacementInfo : 1;
+
+	// If enabled, use the placement info override and not the default one.
+	UPROPERTY(EditDefaultsOnly, BlueprintReadOnly, Category = GameplayCueNotify, Meta = (InlineEditConditionToggle))
+	uint32 bUseSoundParameterInterface : 1;
+
 };
 
 
