Ability System
- Make UAbilitySystemComponent::GetAggregatedStackCount const

diff --git a/Engine/Plugins/Runtime/GameplayAbilities/Source/GameplayAbilities/Private/AbilitySystemComponent.cpp b/Engine/Plugins/Runtime/GameplayAbilities/Source/GameplayAbilities/Private/AbilitySystemComponent.cpp
--- a/Engine/Plugins/Runtime/GameplayAbilities/Source/GameplayAbilities/Private/AbilitySystemComponent.cpp
+++ b/Engine/Plugins/Runtime/GameplayAbilities/Source/GameplayAbilities/Private/AbilitySystemComponent.cpp
@@ -371,7 +371,7 @@ int32 UAbilitySystemComponent::GetGameplayEffectCount(TSubclassOf<UGameplayEffec
 	return Count;
 }
 
-int32 UAbilitySystemComponent::GetAggregatedStackCount(const FGameplayEffectQuery& Query)
+int32 UAbilitySystemComponent::GetAggregatedStackCount(const FGameplayEffectQuery& Query) const
 {
 	return ActiveGameplayEffects.GetActiveEffectCount(Query);
 }
diff --git a/Engine/Plugins/Runtime/GameplayAbilities/Source/GameplayAbilities/Public/AbilitySystemComponent.h b/Engine/Plugins/Runtime/GameplayAbilities/Source/GameplayAbilities/Public/AbilitySystemComponent.h
--- a/Engine/Plugins/Runtime/GameplayAbilities/Source/GameplayAbilities/Public/AbilitySystemComponent.h
+++ b/Engine/Plugins/Runtime/GameplayAbilities/Source/GameplayAbilities/Public/AbilitySystemComponent.h
@@ -328,7 +328,7 @@ class GAMEPLAYABILITIES_API UAbilitySystemComponent : public UGameplayTasksCompo
 	int32 GetGameplayEffectCount(TSubclassOf<UGameplayEffect> SourceGameplayEffect, UAbilitySystemComponent* OptionalInstigatorFilterComponent, bool bEnforceOnGoingCheck = true);
 
 	/** Returns the sum of StackCount of all gameplay effects that pass query */
-	int32 GetAggregatedStackCount(const FGameplayEffectQuery& Query);
+	int32 GetAggregatedStackCount(const FGameplayEffectQuery& Query) const;
 
 	/** This only exists so it can be hooked up to a multicast delegate */
 	void RemoveActiveGameplayEffect_NoReturn(FActiveGameplayEffectHandle Handle, int32 StacksToRemove=-1)
