Gameplay Abilities - AbilityAsync now derives from the new cancelable BP async class.
Gameplay Abilities - Adding a new Async action "Wait For Attribute Changed", can now be used to listen for an attribute changing, super handy for UI monitors.

diff --git a/Engine/Plugins/Runtime/GameplayAbilities/Source/GameplayAbilities/Private/Abilities/Async/AbilityAsync.cpp b/Engine/Plugins/Runtime/GameplayAbilities/Source/GameplayAbilities/Private/Abilities/Async/AbilityAsync.cpp
--- a/Engine/Plugins/Runtime/GameplayAbilities/Source/GameplayAbilities/Private/Abilities/Async/AbilityAsync.cpp
+++ b/Engine/Plugins/Runtime/GameplayAbilities/Source/GameplayAbilities/Private/Abilities/Async/AbilityAsync.cpp
@@ -4,6 +4,12 @@
 #include "AbilitySystemGlobals.h"
 #include "AbilitySystemComponent.h"
 
+void UAbilityAsync::Cancel()
+{
+	EndAction();
+	Super::Cancel();
+}
+
 void UAbilityAsync::EndAction()
 {
 	// Child classes should override this if they need to explicitly unbind delegates that aren't just using weak pointers
diff --git a/Engine/Plugins/Runtime/GameplayAbilities/Source/GameplayAbilities/Private/Abilities/Async/AbilityAsync_WaitAttributeChanged.cpp b/Engine/Plugins/Runtime/GameplayAbilities/Source/GameplayAbilities/Private/Abilities/Async/AbilityAsync_WaitAttributeChanged.cpp
new file mode 100644
--- /dev/null
+++ b/Engine/Plugins/Runtime/GameplayAbilities/Source/GameplayAbilities/Private/Abilities/Async/AbilityAsync_WaitAttributeChanged.cpp
@@ -0,0 +1,54 @@
+// Copyright Epic Games, Inc. All Rights Reserved.
+
+#include "Abilities/Async/AbilityAsync_WaitAttributeChanged.h"
+#include "AbilitySystemGlobals.h"
+#include "AbilitySystemComponent.h"
+
+UAbilityAsync_WaitAttributeChanged* UAbilityAsync_WaitAttributeChanged::WaitForAttributeChanged(AActor* TargetActor, FGameplayAttribute Attribute, bool OnlyTriggerOnce)
+{
+	UAbilityAsync_WaitAttributeChanged* MyObj = NewObject<UAbilityAsync_WaitAttributeChanged>();
+	MyObj->SetAbilityActor(TargetActor);
+	MyObj->Attribute = Attribute;
+	MyObj->OnlyTriggerOnce = OnlyTriggerOnce;
+	return MyObj;
+}
+
+void UAbilityAsync_WaitAttributeChanged::Activate()
+{
+	Super::Activate();
+
+	if (UAbilitySystemComponent* ASC = GetAbilitySystemComponent())
+	{
+		MyHandle = ASC->GetGameplayAttributeValueChangeDelegate(Attribute).AddUObject(this, &ThisClass::OnAttributeChanged);
+	}
+	else
+	{
+		EndAction();
+	}
+}
+
+void UAbilityAsync_WaitAttributeChanged::OnAttributeChanged(const FOnAttributeChangeData& ChangeData)
+{
+	if (ShouldBroadcastDelegates())
+	{
+		Changed.Broadcast(ChangeData.Attribute, ChangeData.NewValue, ChangeData.OldValue);
+
+		if (OnlyTriggerOnce)
+		{
+			EndAction();
+		}
+	}
+	else
+	{
+		EndAction();
+	}
+}
+
+void UAbilityAsync_WaitAttributeChanged::EndAction()
+{
+	if (UAbilitySystemComponent* ASC = GetAbilitySystemComponent())
+	{
+		ASC->GetGameplayAttributeValueChangeDelegate(Attribute).Remove(MyHandle);
+	}
+	Super::EndAction();
+}
diff --git a/Engine/Plugins/Runtime/GameplayAbilities/Source/GameplayAbilities/Public/Abilities/Async/AbilityAsync.h b/Engine/Plugins/Runtime/GameplayAbilities/Source/GameplayAbilities/Public/Abilities/Async/AbilityAsync.h
--- a/Engine/Plugins/Runtime/GameplayAbilities/Source/GameplayAbilities/Public/Abilities/Async/AbilityAsync.h
+++ b/Engine/Plugins/Runtime/GameplayAbilities/Source/GameplayAbilities/Public/Abilities/Async/AbilityAsync.h
@@ -5,6 +5,7 @@
 #include "Kismet/BlueprintAsyncActionBase.h"
 #include "GameplayEffectTypes.h"
 #include "Abilities/GameplayAbilityTypes.h"
+#include "Engine/CancellableAsyncAction.h"
 #include "AbilityAsync.generated.h"
 
 class UAbilitySystemComponent;
@@ -16,11 +17,13 @@ class UAbilitySystemComponent;
  * EndAction should be called when a one-time action has succeeded or failed, but for longer-lived actions with multiple triggers it can be called from blueprints.
  */
 UCLASS(Abstract, meta = (ExposedAsyncProxy = AsyncAction))
-class GAMEPLAYABILITIES_API UAbilityAsync : public UBlueprintAsyncActionBase
+class GAMEPLAYABILITIES_API UAbilityAsync : public UCancellableAsyncAction
 {
 	GENERATED_BODY()
 public:
 
+	virtual void Cancel() override;
+
 	/** Explicitly end the action, will disable any callbacks and allow action to be destroyed */
 	UFUNCTION(BlueprintCallable, Category = "Ability|Async")
 	virtual void EndAction();
diff --git a/Engine/Plugins/Runtime/GameplayAbilities/Source/GameplayAbilities/Public/Abilities/Async/AbilityAsync_WaitAttributeChanged.h b/Engine/Plugins/Runtime/GameplayAbilities/Source/GameplayAbilities/Public/Abilities/Async/AbilityAsync_WaitAttributeChanged.h
new file mode 100644
--- /dev/null
+++ b/Engine/Plugins/Runtime/GameplayAbilities/Source/GameplayAbilities/Public/Abilities/Async/AbilityAsync_WaitAttributeChanged.h
@@ -0,0 +1,37 @@
+// Copyright Epic Games, Inc. All Rights Reserved.
+#pragma once
+
+#include "CoreMinimal.h"
+#include "AbilityAsync.h"
+#include "GameplayEffectTypes.h"
+#include "Abilities/GameplayAbilityTypes.h"
+#include "AbilityAsync_WaitAttributeChanged.generated.h"
+
+class UAbilitySystemComponent;
+
+UCLASS()
+class GAMEPLAYABILITIES_API UAbilityAsync_WaitAttributeChanged : public UAbilityAsync
+{
+	GENERATED_BODY()
+
+public:
+
+	UFUNCTION(BlueprintCallable, Category = "Ability|Tasks", meta = (DefaultToSelf = "TargetActor", BlueprintInternalUseOnly = "TRUE"))
+	static UAbilityAsync_WaitAttributeChanged* WaitForAttributeChanged(AActor* TargetActor, FGameplayAttribute Attribute, bool OnlyTriggerOnce = false);
+
+	DECLARE_DYNAMIC_MULTICAST_DELEGATE_ThreeParams(FAsyncWaitAttributeChangedDelegate, FGameplayAttribute, Attribute, float, NewValue, float, OldValue);
+	UPROPERTY(BlueprintAssignable)
+	FAsyncWaitAttributeChangedDelegate Changed;
+
+protected:
+
+	virtual void Activate() override;
+	virtual void EndAction() override;
+
+	void OnAttributeChanged(const FOnAttributeChangeData& ChangeData);
+
+	FGameplayAttribute Attribute;
+	bool OnlyTriggerOnce = false;
+
+	FDelegateHandle MyHandle;
+};
