Gameplay Abilities: Exposed a helper method to get the UIData from a GE class

diff --git a/Engine/Plugins/Runtime/GameplayAbilities/Source/GameplayAbilities/Private/AbilitySystemBlueprintLibrary.cpp b/Engine/Plugins/Runtime/GameplayAbilities/Source/GameplayAbilities/Private/AbilitySystemBlueprintLibrary.cpp
--- a/Engine/Plugins/Runtime/GameplayAbilities/Source/GameplayAbilities/Private/AbilitySystemBlueprintLibrary.cpp
+++ b/Engine/Plugins/Runtime/GameplayAbilities/Source/GameplayAbilities/Private/AbilitySystemBlueprintLibrary.cpp
@@ -5,6 +5,7 @@
 #include "GameplayEffect.h"
 #include "AbilitySystemGlobals.h"
 #include "AbilitySystemComponent.h"
+#include "GameplayEffectUIData.h"
 
 UAbilitySystemBlueprintLibrary::UAbilitySystemBlueprintLibrary(const FObjectInitializer& ObjectInitializer)
 : Super(ObjectInitializer)
@@ -1053,6 +1054,19 @@ bool UAbilitySystemBlueprintLibrary::AddLooseGameplayTags(AActor* Actor, const F
 	return false;
 }
 
+const UGameplayEffectUIData* UAbilitySystemBlueprintLibrary::GetGameplayEffectUIData(TSubclassOf<UGameplayEffect> EffectClass, TSubclassOf<UGameplayEffectUIData> DataType)
+{
+	if (UClass* ActualPtr = EffectClass.Get())
+	{
+		const UGameplayEffectUIData* UIData = GetDefault<UGameplayEffect>(ActualPtr)->UIData;
+		if ((UIData != nullptr) && (DataType != nullptr) && UIData->IsA(DataType))
+		{
+			return UIData;
+		}
+	}
+	return nullptr;
+}
+
 bool UAbilitySystemBlueprintLibrary::RemoveLooseGameplayTags(AActor* Actor, const FGameplayTagContainer& GameplayTags)
 {
 	if (UAbilitySystemComponent* AbilitySysComp = GetAbilitySystemComponent(Actor))
diff --git a/Engine/Plugins/Runtime/GameplayAbilities/Source/GameplayAbilities/Public/AbilitySystemBlueprintLibrary.h b/Engine/Plugins/Runtime/GameplayAbilities/Source/GameplayAbilities/Public/AbilitySystemBlueprintLibrary.h
--- a/Engine/Plugins/Runtime/GameplayAbilities/Source/GameplayAbilities/Public/AbilitySystemBlueprintLibrary.h
+++ b/Engine/Plugins/Runtime/GameplayAbilities/Source/GameplayAbilities/Public/AbilitySystemBlueprintLibrary.h
@@ -18,6 +18,7 @@
 
 class UAbilitySystemComponent;
 class UGameplayEffect;
+class UGameplayEffectUIData;
 
 /** Blueprint library for ability system. Many of these functions are useful to call from native as well */
 UCLASS(meta=(ScriptName="AbilitySystemLibrary"))
@@ -384,4 +385,7 @@ class GAMEPLAYABILITIES_API UAbilitySystemBlueprintLibrary : public UBlueprintFu
 	UFUNCTION(BlueprintPure, Category = "Ability|GameplayEffect", Meta = (DisplayName = "Get Active GameplayEffect Debug String "))
 	static FString GetActiveGameplayEffectDebugString(FActiveGameplayEffectHandle ActiveHandle);
 
+	/** Returns the UI data for a gameplay effect class (if any) */
+	UFUNCTION(BlueprintCallable, Category = "Ability|GameplayEffect", Meta = (DisplayName = "Get GameplayEffect UI Data", DeterminesOutputType="DataType"))
+	static const UGameplayEffectUIData* GetGameplayEffectUIData(TSubclassOf<UGameplayEffect> EffectClass, TSubclassOf<UGameplayEffectUIData> DataType);
 };
