Disallowed function calls on non-instanced abilities no longer halt the program. Converted checks to ensures where reasonable and improved messaging. Pulled macro from sample game into Gameplay Ability plugin to support this.

diff --git a/Engine/Plugins/Runtime/GameplayAbilities/Source/GameplayAbilities/Private/Abilities/GameplayAbility.cpp b/Engine/Plugins/Runtime/GameplayAbilities/Source/GameplayAbilities/Private/Abilities/GameplayAbility.cpp
--- a/Engine/Plugins/Runtime/GameplayAbilities/Source/GameplayAbilities/Private/Abilities/GameplayAbility.cpp
+++ b/Engine/Plugins/Runtime/GameplayAbilities/Source/GameplayAbilities/Private/Abilities/GameplayAbility.cpp
@@ -974,11 +974,7 @@ FGameplayAbilityActorInfo UGameplayAbility::GetActorInfo() const
 
 AActor* UGameplayAbility::GetOwningActorFromActorInfo() const
 {
-	if (!ensureMsgf(IsInstantiated(), TEXT("%s: GetOwningActorFromActorInfo can not be called on a non-instanced ability"), *GetName()))
-	{
-		ABILITY_LOG(Warning, TEXT("%s: GetOwningActorFromActorInfo can not be called on a non-instanced ability"), *GetName());
-		return nullptr;
-	}
+	ENSURE_ABILITY_IS_INSTANTIATED_OR_RETURN(GetOwningActorFromActorInfo, nullptr);
 
 	if (!ensure(CurrentActorInfo))
 	{
@@ -1183,7 +1179,7 @@ void UGameplayAbility::MontageStop(float OverrideBlendOutTime)
 
 void UGameplayAbility::SetCurrentMontage(class UAnimMontage* InCurrentMontage)
 {
-	ensure(IsInstantiated());
+	ENSURE_ABILITY_IS_INSTANTIATED_OR_RETURN(SetCurrentMontage, );
 	CurrentMontage = InCurrentMontage;
 }
 
@@ -1493,7 +1489,7 @@ int32 UGameplayAbility::GetAbilityLevel_BP(FGameplayAbilitySpecHandle Handle, co
 
 FGameplayAbilitySpec* UGameplayAbility::GetCurrentAbilitySpec() const
 {
-	check(IsInstantiated()); // You should not call this on non instanced abilities.
+	ENSURE_ABILITY_IS_INSTANTIATED_OR_RETURN(GetCurrentAbilitySpec, nullptr);
 	check(CurrentActorInfo);
 
 	UAbilitySystemComponent* const AbilitySystemComponent = GetAbilitySystemComponentFromActorInfo_Checked();
@@ -1502,7 +1498,7 @@ FGameplayAbilitySpec* UGameplayAbility::GetCurrentAbilitySpec() const
 
 FGameplayEffectContextHandle UGameplayAbility::GetGrantedByEffectContext() const
 {
-	check(IsInstantiated()); // You should not call this on non instanced abilities.
+	ENSURE_ABILITY_IS_INSTANTIATED_OR_RETURN(GetGrantedByEffectContext, FGameplayEffectContextHandle());
 	check(CurrentActorInfo);
 	if (CurrentActorInfo)
 	{
@@ -1519,7 +1515,7 @@ FGameplayEffectContextHandle UGameplayAbility::GetGrantedByEffectContext() const
 
 void UGameplayAbility::RemoveGrantedByEffect()
 {
-	check(IsInstantiated()); // You should not call this on non instanced abilities.
+	ENSURE_ABILITY_IS_INSTANTIATED_OR_RETURN(RemoveGrantedByEffect, );
 	check(CurrentActorInfo);
 	if (CurrentActorInfo)
 	{
diff --git a/Engine/Plugins/Runtime/GameplayAbilities/Source/GameplayAbilities/Public/Abilities/GameplayAbility.h b/Engine/Plugins/Runtime/GameplayAbilities/Source/GameplayAbilities/Public/Abilities/GameplayAbility.h
--- a/Engine/Plugins/Runtime/GameplayAbilities/Source/GameplayAbilities/Public/Abilities/GameplayAbility.h
+++ b/Engine/Plugins/Runtime/GameplayAbilities/Source/GameplayAbilities/Public/Abilities/GameplayAbility.h
@@ -70,6 +70,16 @@ DECLARE_MULTICAST_DELEGATE_OneParam(FOnGameplayAbilityStateEnded, FName);
 /** Used to delay execution until we leave a critical section */
 DECLARE_DELEGATE(FPostLockDelegate);
 
+
+#define ENSURE_ABILITY_IS_INSTANTIATED_OR_RETURN(FunctionName, ReturnValue)																				\
+{																																						\
+	if (!ensure(IsInstantiated()))																														\
+	{																																					\
+		ABILITY_LOG(Error, TEXT("%s: " #FunctionName " cannot be called on a non-instanced ability. Check the instancing policy."), *GetPathName());	\
+		return ReturnValue;																																\
+	}																																					\
+}
+
 /** Structure that defines how an ability will be triggered by external events */
 USTRUCT()
 struct FAbilityTriggerData
@@ -175,28 +185,28 @@ public:
 	/** Gets the current actor info bound to this ability - can only be called on instanced abilities. */
 	const FGameplayAbilityActorInfo* GetCurrentActorInfo() const
 	{
-		check(IsInstantiated());
+		ENSURE_ABILITY_IS_INSTANTIATED_OR_RETURN(GetCurrentActorInfo, nullptr);
 		return CurrentActorInfo;
 	}
 
 	/** Gets the current activation info bound to this ability - can only be called on instanced abilities. */
 	FGameplayAbilityActivationInfo GetCurrentActivationInfo() const
 	{
-		check(IsInstantiated());
+		ENSURE_ABILITY_IS_INSTANTIATED_OR_RETURN(GetCurrentActivationInfo, FGameplayAbilityActivationInfo());
 		return CurrentActivationInfo;
 	}
 
 	/** Gets the current activation info bound to this ability - can only be called on instanced abilities. */
 	FGameplayAbilityActivationInfo& GetCurrentActivationInfoRef()
 	{
-		check(IsInstantiated());
+		checkf(IsInstantiated(), TEXT("%s: GetCurrentActivationInfoRef cannot be called on a non-instanced ability. Check the instancing policy."), *GetPathName());
 		return CurrentActivationInfo;
 	}
 
 	/** Gets the current AbilitySpecHandle- can only be called on instanced abilities. */
 	FGameplayAbilitySpecHandle GetCurrentAbilitySpecHandle() const
 	{
-		check(IsInstantiated());
+		ENSURE_ABILITY_IS_INSTANTIATED_OR_RETURN(GetCurrentAbilitySpecHandle, FGameplayAbilitySpecHandle());
 		return CurrentSpecHandle;
 	}
 
