FScalableFloatDetails: Dynamic preview levels based on the time range of the curve
#rnx

diff --git a/Engine/Plugins/Runtime/GameplayAbilities/Source/GameplayAbilitiesEditor/Private/AttributeDetails.cpp b/Engine/Plugins/Runtime/GameplayAbilities/Source/GameplayAbilitiesEditor/Private/AttributeDetails.cpp
--- a/Engine/Plugins/Runtime/GameplayAbilities/Source/GameplayAbilitiesEditor/Private/AttributeDetails.cpp
+++ b/Engine/Plugins/Runtime/GameplayAbilities/Source/GameplayAbilitiesEditor/Private/AttributeDetails.cpp
@@ -13,13 +13,7 @@
 #include "DetailWidgetRow.h"
 #include "DetailCategoryBuilder.h"
 #include "AttributeSet.h"
-#include "GameplayAbilitiesModule.h"
-#include "UObject/UObjectHash.h"
-#include "UObject/UObjectIterator.h"
 #include "AbilitySystemGlobals.h"
-#include "Widgets/Input/SSearchBox.h"
-#include "Widgets/Views/SListView.h"
-#include "SListViewSelectorDropdownMenu.h"
 #include "Widgets/Input/STextComboBox.h"
 #include "AbilitySystemComponent.h"
 #include "SGameplayAttributeWidget.h"
@@ -98,12 +92,10 @@ TSharedPtr<FString> FAttributePropertyDetails::GetPropertyType() const
 	if (MyProperty.IsValid())
 	{
 		FProperty *PropertyValue = nullptr;
-		FProperty *ObjPtr = nullptr;
-		MyProperty->GetValue(ObjPtr);
-		PropertyValue = ObjPtr;
+		MyProperty->GetValue(PropertyValue);
 		if (PropertyValue)
 		{
-			FString FullString = PropertyValue->GetOwnerVariant().GetName() + TEXT(".") + PropertyValue->GetName();
+			const FString FullString = PropertyValue->GetOwnerVariant().GetName() + TEXT(".") + PropertyValue->GetName();
 			for (int32 i=0; i < PropertyOptions.Num(); ++i)
 			{
 				if (PropertyOptions[i].IsValid() && PropertyOptions[i].Get()->Equals(FullString))
@@ -142,8 +134,6 @@ void FAttributePropertyDetails::OnChangeProperty(TSharedPtr<FString> ItemSelecte
 		UObject* nullObj = nullptr;
 		MyProperty->SetValue(nullObj);
 	}
-
-	
 }
 
 void FAttributePropertyDetails::OnAttributeChanged(FProperty* SelectedAttribute)
@@ -225,9 +215,7 @@ TSharedPtr<FString> FAttributeDetails::GetPropertyType() const
 		return PropertyOptions[0];
 
 	FProperty *PropertyValue = nullptr;
-	FProperty *ObjPtr = nullptr;
-	MyProperty->GetValue(ObjPtr);
-	PropertyValue = ObjPtr;
+	MyProperty->GetValue(PropertyValue);
 
 	if (PropertyValue != nullptr)
 	{
@@ -289,8 +277,11 @@ void FScalableFloatDetails::CustomizeHeader( TSharedRef<class IPropertyHandle> S
 		RowNameProperty = CurveTableHandleProperty->GetChildHandle(GET_MEMBER_NAME_CHECKED(FCurveTableRowHandle, RowName));
 		CurveTableProperty = CurveTableHandleProperty->GetChildHandle(GET_MEMBER_NAME_CHECKED(FCurveTableRowHandle, CurveTable));
 
+		UpdatePreviewLevels();
+
 		CurveTableProperty->SetOnPropertyValueChanged(FSimpleDelegate::CreateSP(this, &FScalableFloatDetails::OnCurveSourceChanged));
 		RegistryTypeProperty->SetOnPropertyValueChanged(FSimpleDelegate::CreateSP(this, &FScalableFloatDetails::OnCurveSourceChanged));
+		RowNameProperty->SetOnPropertyValueChanged(FSimpleDelegate::CreateSP(this, &FScalableFloatDetails::OnRowNameChanged));
 
 		HeaderRow
 			.NameContent()
@@ -622,6 +613,11 @@ void FScalableFloatDetails::OnCurveSourceChanged()
 	}
 }
 
+void FScalableFloatDetails::OnRowNameChanged()
+{
+	UpdatePreviewLevels();
+}
+
 void FScalableFloatDetails::RefreshSourceData()
 {
 	// Set the default value to 1.0 when using a curve source, so the value in the table is used directly. Only do this if the value is currently 0 (default)
@@ -719,12 +715,12 @@ EVisibility FScalableFloatDetails::GetPreviewVisibility() const
 
 float FScalableFloatDetails::GetPreviewLevel() const
 {
-	return (MaxPreviewLevel != 0) ? PreviewLevel / MaxPreviewLevel : 0;
+	return (MaxPreviewLevel != MinPreviewLevel) ? (PreviewLevel-MinPreviewLevel) / (MaxPreviewLevel-MinPreviewLevel) : 0;
 }
 
 void FScalableFloatDetails::SetPreviewLevel(float NewLevel)
 {
-	PreviewLevel = FMath::FloorToInt(NewLevel * MaxPreviewLevel);
+	PreviewLevel = FMath::FloorToInt(NewLevel * (MaxPreviewLevel - MinPreviewLevel) + MinPreviewLevel);
 }
 
 FText FScalableFloatDetails::GetRowNameComboBoxContentText() const
@@ -923,6 +919,21 @@ bool FScalableFloatDetails::IsEditable() const
 	return true;
 }
 
+void FScalableFloatDetails::UpdatePreviewLevels()
+{
+	if (const FRealCurve* FoundCurve = GetRealCurve())
+	{
+		FoundCurve->GetTimeRange(MinPreviewLevel, MaxPreviewLevel);
+	}
+	else
+	{
+		MinPreviewLevel = DefaultMinPreviewLevel;
+		MaxPreviewLevel = DefaultMaxPreviewLevel;
+	}
+
+	PreviewLevel = FMath::Clamp(PreviewLevel, MinPreviewLevel, MaxPreviewLevel);
+}
+
 void FScalableFloatDetails::CustomizeChildren( TSharedRef<class IPropertyHandle> StructPropertyHandle, class IDetailChildrenBuilder& StructBuilder, IPropertyTypeCustomizationUtils& StructCustomizationUtils )
 {
 	
diff --git a/Engine/Plugins/Runtime/GameplayAbilities/Source/GameplayAbilitiesEditor/Public/AttributeDetails.h b/Engine/Plugins/Runtime/GameplayAbilities/Source/GameplayAbilitiesEditor/Public/AttributeDetails.h
--- a/Engine/Plugins/Runtime/GameplayAbilities/Source/GameplayAbilitiesEditor/Public/AttributeDetails.h
+++ b/Engine/Plugins/Runtime/GameplayAbilities/Source/GameplayAbilitiesEditor/Public/AttributeDetails.h
@@ -70,10 +70,13 @@ class GAMEPLAYABILITIESEDITOR_API FScalableFloatDetails : public IPropertyTypeCu
 {
 public:
 	static TSharedRef<IPropertyTypeCustomization> MakeInstance();
+	static constexpr float DefaultMinPreviewLevel = 0.f;
+	static constexpr float DefaultMaxPreviewLevel = 30.f;
 
 	FScalableFloatDetails()
 		: PreviewLevel(0.f)
-		, MaxPreviewLevel(30.f) // This should perhaps be configurable per scalable float somehow
+		, MinPreviewLevel(DefaultMinPreviewLevel)
+		, MaxPreviewLevel(DefaultMaxPreviewLevel) 
 	{
 	}
 
@@ -82,7 +85,8 @@ protected:
 	virtual void CustomizeHeader( TSharedRef<class IPropertyHandle> StructPropertyHandle, class FDetailWidgetRow& HeaderRow, IPropertyTypeCustomizationUtils& StructCustomizationUtils ) override;
 	virtual void CustomizeChildren( TSharedRef<class IPropertyHandle> StructPropertyHandle, class IDetailChildrenBuilder& StructBuilder, IPropertyTypeCustomizationUtils& StructCustomizationUtils ) override;              
 
-	bool IsEditable( ) const;
+	bool IsEditable() const;
+	void UpdatePreviewLevels();
 
 	// Curve Table selector
 	TSharedRef<SWidget> CreateCurveTableWidget();
@@ -114,6 +118,7 @@ protected:
 	EVisibility GetRowNameVisibility() const;
 	FText GetRowNameComboBoxContentText() const;
 	FText GetRowNameComboBoxContentTooltip() const;
+	void OnRowNameChanged();
 
 	// Preview widgets
 	EVisibility GetPreviewVisibility() const;
@@ -138,6 +143,7 @@ protected:
 	TWeakPtr<IPropertyUtilities> PropertyUtilities;
 
 	float PreviewLevel;
+	float MinPreviewLevel;
 	float MaxPreviewLevel;
 	bool bSourceRefreshQueued;
 };
