Add End Ability Locally function to the generic gameplay ability, this is needed for predicted abilities with animations where you don't want the server ability to end early and miss a callback
Also clean up the comments for end ability

diff --git a/Engine/Plugins/Runtime/GameplayAbilities/Source/GameplayAbilities/Private/Abilities/GameplayAbility.cpp b/Engine/Plugins/Runtime/GameplayAbilities/Source/GameplayAbilities/Private/Abilities/GameplayAbility.cpp
--- a/Engine/Plugins/Runtime/GameplayAbilities/Source/GameplayAbilities/Private/Abilities/GameplayAbility.cpp
+++ b/Engine/Plugins/Runtime/GameplayAbilities/Source/GameplayAbilities/Private/Abilities/GameplayAbility.cpp
@@ -1140,6 +1140,15 @@ void UGameplayAbility::K2_EndAbility()
 	EndAbility(CurrentSpecHandle, CurrentActorInfo, CurrentActivationInfo, bReplicateEndAbility, bWasCancelled);
 }
 
+void UGameplayAbility::K2_EndAbilityLocally()
+{
+	check(CurrentActorInfo);
+
+	bool bReplicateEndAbility = false;
+	bool bWasCancelled = false;
+	EndAbility(CurrentSpecHandle, CurrentActorInfo, CurrentActivationInfo, bReplicateEndAbility, bWasCancelled);
+}
+
 void UGameplayAbility::MontageJumpToSection(FName SectionName)
 {
 	check(CurrentActorInfo);
diff --git a/Engine/Plugins/Runtime/GameplayAbilities/Source/GameplayAbilities/Public/Abilities/GameplayAbility.h b/Engine/Plugins/Runtime/GameplayAbilities/Source/GameplayAbilities/Public/Abilities/GameplayAbility.h
--- a/Engine/Plugins/Runtime/GameplayAbilities/Source/GameplayAbilities/Public/Abilities/GameplayAbility.h
+++ b/Engine/Plugins/Runtime/GameplayAbilities/Source/GameplayAbilities/Public/Abilities/GameplayAbility.h
@@ -585,12 +585,15 @@ protected:
 	// -------------------------------------
 	//	EndAbility
 	// -------------------------------------
-
-	/** Call from kismet to end the ability naturally */
-	UFUNCTION(BlueprintCallable, Category = Ability, DisplayName="EndAbility", meta=(ScriptName = "EndAbility"))
+	/** Call from blueprints to forcibly end the ability without canceling it. This will replicate the end ability to the client or server which can interrupt tasks */
+	UFUNCTION(BlueprintCallable, Category = Ability, DisplayName="End Ability", meta=(ScriptName = "EndAbility"))
 	virtual void K2_EndAbility();
 
-	/** Kismet event, will be called if an ability ends normally or abnormally */
+	/** Call from blueprints to end the ability naturally. This will only end predicted abilities locally, allowing it end naturally on the client or server */
+	UFUNCTION(BlueprintCallable, Category = Ability, DisplayName = "End Ability Locally", meta = (ScriptName = "EndAbilityLocally"))
+	virtual void K2_EndAbilityLocally();
+
+	/** Blueprint event, will be called if an ability ends normally or abnormally */
 	UFUNCTION(BlueprintImplementableEvent, Category = Ability, DisplayName = "OnEndAbility", meta=(ScriptName = "OnEndAbility"))
 	void K2_OnEndAbility(bool bWasCancelled);
 
