Added support for niagara camera lens effects to gameplay cue notifies

diff --git a/Engine/Plugins/FX/Niagara/Source/Niagara/Classes/NiagaraLensEffectBase.h b/Engine/Plugins/FX/Niagara/Source/Niagara/Classes/NiagaraLensEffectBase.h
--- a/Engine/Plugins/FX/Niagara/Source/Niagara/Classes/NiagaraLensEffectBase.h
+++ b/Engine/Plugins/FX/Niagara/Source/Niagara/Classes/NiagaraLensEffectBase.h
@@ -75,6 +75,7 @@ protected:
 	virtual void NotifyRetriggered() override;
 	virtual void ActivateLensEffect() override;
 	virtual void DeactivateLensEffect() override;
+	virtual bool IsLooping() const override;
 	//~ End ICameraLensEffectInterface interface
 
 	virtual void EndPlay(const EEndPlayReason::Type EndPlayReason) override;
diff --git a/Engine/Plugins/FX/Niagara/Source/Niagara/Private/NiagaraLensEffectBase.cpp b/Engine/Plugins/FX/Niagara/Source/Niagara/Private/NiagaraLensEffectBase.cpp
--- a/Engine/Plugins/FX/Niagara/Source/Niagara/Private/NiagaraLensEffectBase.cpp
+++ b/Engine/Plugins/FX/Niagara/Source/Niagara/Private/NiagaraLensEffectBase.cpp
@@ -3,6 +3,7 @@
 #include "NiagaraLensEffectBase.h"
 
 #include "NiagaraComponent.h"
+#include "NiagaraSystem.h"
 #include "Camera/PlayerCameraManager.h"
 
 
@@ -110,3 +111,13 @@ void ANiagaraLensEffectBase::EndPlay(const EEndPlayReason::Type EndPlayReason)
 		OwningCameraManager->RemoveGenericCameraLensEffect(this);
 	}
 }
+
+bool ANiagaraLensEffectBase::IsLooping() const
+{
+	if (UNiagaraSystem* System = GetNiagaraComponent()->GetAsset())
+	{
+		return System->IsLooping();
+	}
+
+	return false;
+}
diff --git a/Engine/Plugins/Runtime/GameplayAbilities/Source/GameplayAbilities/Private/GameplayCueNotifyTypes.cpp b/Engine/Plugins/Runtime/GameplayAbilities/Source/GameplayAbilities/Private/GameplayCueNotifyTypes.cpp
--- a/Engine/Plugins/Runtime/GameplayAbilities/Source/GameplayAbilities/Private/GameplayCueNotifyTypes.cpp
+++ b/Engine/Plugins/Runtime/GameplayAbilities/Source/GameplayAbilities/Private/GameplayCueNotifyTypes.cpp
@@ -2,7 +2,7 @@
 
 #include "GameplayCueNotifyTypes.h"
 #include "Components/SkeletalMeshComponent.h"
-#include "Particles/EmitterCameraLensEffectBase.h"
+#include "Camera/CameraLensEffectInterface.h"
 #include "NiagaraSystem.h"
 #include "NiagaraComponent.h"
 #include "NiagaraFunctionLibrary.h"
@@ -785,7 +785,7 @@ bool FGameplayCueNotify_CameraLensEffectInfo::PlayCameraLensEffect(const FGamepl
 				continue;
 			}
 
-			AEmitterCameraLensEffectBase* CameraLensEffectInstance = PC->PlayerCameraManager->AddCameraLensEffect(CameraLensEffect);
+			TScriptInterface<ICameraLensEffectInterface> CameraLensEffectInstance = PC->PlayerCameraManager->AddGenericCameraLensEffect(CameraLensEffect);
 			OutSpawnResult.CameraLensEffects.Add(CameraLensEffectInstance);
 		}
 	}
@@ -796,7 +796,7 @@ bool FGameplayCueNotify_CameraLensEffectInfo::PlayCameraLensEffect(const FGamepl
 
 		if (TargetPC && TargetPC->PlayerCameraManager)
 		{
-			AEmitterCameraLensEffectBase* CameraLensEffectInstance = TargetPC->PlayerCameraManager->AddCameraLensEffect(CameraLensEffect);
+			TScriptInterface<ICameraLensEffectInterface> CameraLensEffectInstance = TargetPC->PlayerCameraManager->AddGenericCameraLensEffect(CameraLensEffect);
 			OutSpawnResult.CameraLensEffects.Add(CameraLensEffectInstance);
 		}
 	}
@@ -809,14 +809,16 @@ void FGameplayCueNotify_CameraLensEffectInfo::ValidateBurstAssets(UObject* Conta
 #if WITH_EDITORONLY_DATA
 	if (CameraLensEffect != nullptr)
 	{
-		AEmitterCameraLensEffectBase* CameraLensEffectCDO = CameraLensEffect->GetDefaultObject<AEmitterCameraLensEffectBase>();
-		if (CameraLensEffectCDO->IsLooping())
+		if (ICameraLensEffectInterface* CameraLensEffectCDO = Cast<ICameraLensEffectInterface>(CameraLensEffect->GetDefaultObject<AActor>()))
 		{
-			ValidationErrors.Add(FText::Format(
-				LOCTEXT("CameraLensEffect_ShouldNotLoop", "Camera lens effect [{0}] used in slot [{1}] for asset [{2}] is set to looping, but the slot is a one-shot (the instance will leak)."),
-				FText::AsCultureInvariant(CameraLensEffect->GetPathName()),
-				FText::AsCultureInvariant(Context),
-				FText::AsCultureInvariant(ContainingAsset->GetPathName())));
+			if (CameraLensEffectCDO->IsLooping())
+			{
+				ValidationErrors.Add(FText::Format(
+					LOCTEXT("CameraLensEffect_ShouldNotLoop", "Camera lens effect [{0}] used in slot [{1}] for asset [{2}] is set to looping, but the slot is a one-shot (the instance will leak)."),
+					FText::AsCultureInvariant(CameraLensEffect->GetPathName()),
+					FText::AsCultureInvariant(Context),
+					FText::AsCultureInvariant(ContainingAsset->GetPathName())));
+			}
 		}
 	}
 #endif // #if WITH_EDITORONLY_DATA
@@ -1140,7 +1142,7 @@ void FGameplayCueNotify_LoopingEffects::StopEffects(FGameplayCueNotify_SpawnResu
 	}
 
 	// Stop the camera lens effect.
-	for (AEmitterCameraLensEffectBase* CameraLensEffect : SpawnResult.CameraLensEffects)
+	for (TScriptInterface<ICameraLensEffectInterface> CameraLensEffect : SpawnResult.CameraLensEffects)
 	{
 		if (CameraLensEffect)
 		{
diff --git a/Engine/Plugins/Runtime/GameplayAbilities/Source/GameplayAbilities/Public/GameplayCueNotifyTypes.h b/Engine/Plugins/Runtime/GameplayAbilities/Source/GameplayAbilities/Public/GameplayCueNotifyTypes.h
--- a/Engine/Plugins/Runtime/GameplayAbilities/Source/GameplayAbilities/Public/GameplayCueNotifyTypes.h
+++ b/Engine/Plugins/Runtime/GameplayAbilities/Source/GameplayAbilities/Public/GameplayCueNotifyTypes.h
@@ -17,7 +17,7 @@ class USoundBase;
 class UCameraShakeBase;
 class UCameraAnim;
 class UCameraAnimInst;
-class AEmitterCameraLensEffectBase;
+class ICameraLensEffectInterface;
 class UForceFeedbackEffect;
 class UForceFeedbackAttenuation;
 class UForceFeedbackComponent;
@@ -291,7 +291,7 @@ public:
 
 	// List of camera len effects spawned.  There will be one camera lens effect per local player controller if the effect is played in world.
 	UPROPERTY(BlueprintReadOnly, Transient, Category = GameplayCueNotify)
-	TArray<AEmitterCameraLensEffectBase*> CameraLensEffects;
+	TArray<TScriptInterface<ICameraLensEffectInterface>> CameraLensEffects;
 
 	// Force feedback component that was spawned.  This is only valid when force feedback is set to play in world.
 	UPROPERTY(BlueprintReadOnly, Transient, Category = GameplayCueNotify)
@@ -589,8 +589,8 @@ public:
 	FGameplayCueNotify_PlacementInfo PlacementInfoOverride;
 
 	// Camera lens effect to play.
-	UPROPERTY(EditDefaultsOnly, BlueprintReadOnly, Category = GameplayCueNotify)
-	TSubclassOf<AEmitterCameraLensEffectBase> CameraLensEffect;
+	UPROPERTY(EditDefaultsOnly, BlueprintReadOnly, Category = GameplayCueNotify, Meta = (MustImplement = "CameraLensEffectInterface"))
+	TSubclassOf<AActor> CameraLensEffect;
 
 	// If enabled, use the spawn condition override and not the default one.
 	UPROPERTY(EditDefaultsOnly, BlueprintReadOnly, Category = GameplayCueNotify, Meta = (InlineEditConditionToggle))
diff --git a/Engine/Source/Runtime/Engine/Classes/Camera/CameraLensEffectInterface.h b/Engine/Source/Runtime/Engine/Classes/Camera/CameraLensEffectInterface.h
--- a/Engine/Source/Runtime/Engine/Classes/Camera/CameraLensEffectInterface.h
+++ b/Engine/Source/Runtime/Engine/Classes/Camera/CameraLensEffectInterface.h
@@ -59,6 +59,9 @@ public:
 	virtual void AdjustBaseFOV(float NewFOV) = 0;
 
 	static FTransform GetAttachedEmitterTransform(const AActor* Emitter, const FVector& CamLoc, const FRotator& CamRot, float CamFOVDeg);
+
+	/** Returns true if any associated particle system is set to looping */
+	virtual bool IsLooping() const = 0;
 };
 
 //~ Engineering Note: You can validate assets to see if they have any TSubclassOf<AEmitterCameraLensEffectBase> [or other lens effect bases] and issue warnings where they do! Suggest this type instead!
diff --git a/Engine/Source/Runtime/Engine/Classes/Particles/EmitterCameraLensEffectBase.h b/Engine/Source/Runtime/Engine/Classes/Particles/EmitterCameraLensEffectBase.h
--- a/Engine/Source/Runtime/Engine/Classes/Particles/EmitterCameraLensEffectBase.h
+++ b/Engine/Source/Runtime/Engine/Classes/Particles/EmitterCameraLensEffectBase.h
@@ -73,6 +73,7 @@ public:
 	virtual bool ShouldAllowMultipleInstances() const override;
 	virtual bool ResetWhenTriggered() const override;
 	virtual bool ShouldTreatEmitterAsSame(TSubclassOf<AActor> OtherEmitter) const override;
+	virtual bool IsLooping() const override;
 
 	virtual void AdjustBaseFOV(float NewFOV);
 
@@ -97,8 +98,6 @@ public:
 	UE_DEPRECATED(5.0, "ICameraLensEffectInterface::GetAttachedEmitterTransform is favored now")
 	static FTransform GetAttachedEmitterTransform(AEmitterCameraLensEffectBase const* Emitter, const FVector& CamLoc, const FRotator& CamRot, float CamFOVDeg);
 
-	/** Returns true if either particle system would loop forever when played */
-	bool IsLooping() const;
 private:
 	/** UE_DEPRECATED(4.11) */
 	UPROPERTY()
