Keeping track of AbilitySystemComponent's references to other actors and releasing them to prevent levels from leaking

diff --git a/Engine/Plugins/Runtime/GameplayAbilities/Source/GameplayAbilities/Private/AbilitySystemComponent.cpp b/Engine/Plugins/Runtime/GameplayAbilities/Source/GameplayAbilities/Private/AbilitySystemComponent.cpp
--- a/Engine/Plugins/Runtime/GameplayAbilities/Source/GameplayAbilities/Private/AbilitySystemComponent.cpp
+++ b/Engine/Plugins/Runtime/GameplayAbilities/Source/GameplayAbilities/Private/AbilitySystemComponent.cpp
@@ -2686,7 +2686,15 @@ PRAGMA_DISABLE_DEPRECATION_WARNINGS
 void UAbilitySystemComponent::SetOwnerActor(AActor* NewOwnerActor)
 {
 	MARK_PROPERTY_DIRTY_FROM_NAME(UAbilitySystemComponent, OwnerActor, this);
+	if (OwnerActor)
+	{
+		OwnerActor->OnDestroyed.RemoveDynamic(this, &UAbilitySystemComponent::OnOwnerActorDestroyed);
+	}
 	OwnerActor = NewOwnerActor;
+	if (OwnerActor)
+	{
+		OwnerActor->OnDestroyed.AddDynamic(this, &UAbilitySystemComponent::OnOwnerActorDestroyed);
+	}
 }
 
 AActor* UAbilitySystemComponent::GetOwnerActor() const
@@ -2697,7 +2705,15 @@ AActor* UAbilitySystemComponent::GetOwnerActor() const
 void UAbilitySystemComponent::SetAvatarActor_Direct(AActor* NewAvatarActor)
 {
 	MARK_PROPERTY_DIRTY_FROM_NAME(UAbilitySystemComponent, AvatarActor, this);
+	if (AvatarActor)
+	{
+		AvatarActor->OnDestroyed.RemoveDynamic(this, &UAbilitySystemComponent::OnAvatarActorDestroyed);
+	}
 	AvatarActor = NewAvatarActor;
+	if (AvatarActor)
+	{
+		AvatarActor->OnDestroyed.AddDynamic(this, &UAbilitySystemComponent::OnAvatarActorDestroyed);
+	}
 }
 
 AActor* UAbilitySystemComponent::GetAvatarActor_Direct() const
@@ -2705,9 +2721,44 @@ AActor* UAbilitySystemComponent::GetAvatarActor_Direct() const
 	return AvatarActor;
 }
 
+void UAbilitySystemComponent::OnAvatarActorDestroyed(AActor* InActor)
+{
+	if (InActor == AvatarActor)
+	{
+		AvatarActor = nullptr;
+	}
+}
+
+void UAbilitySystemComponent::OnOwnerActorDestroyed(AActor* InActor)
+{
+	if (InActor == OwnerActor)
+	{
+		OwnerActor = nullptr;
+	}
+}
+
 void UAbilitySystemComponent::SetSpawnedAttributes(const TArray<UAttributeSet*>& NewSpawnedAttributes)
 {
+	TArray<UAttributeSet*>& LocalSpawnedAttributes = GetSpawnedAttributes_Mutable();
+	for (int32 Index = LocalSpawnedAttributes.Num() - 1; Index >= 0; --Index)
+	{
+		UAttributeSet* AttributeSet = LocalSpawnedAttributes[Index];
+		AActor* ActorOwner = AttributeSet->GetTypedOuter<AActor>();
+		if (ActorOwner)
+		{
+			ActorOwner->OnEndPlay.RemoveDynamic(this, &UAbilitySystemComponent::OnSpawnedAttributesEndPlayed);
+		}
+	}
 	GetSpawnedAttributes_Mutable() = NewSpawnedAttributes;
+	for (int32 Index = NewSpawnedAttributes.Num() - 1; Index >= 0; --Index)
+	{
+		UAttributeSet* AttributeSet = NewSpawnedAttributes[Index];
+		AActor* ActorOwner = AttributeSet->GetTypedOuter<AActor>();
+		if (ActorOwner)
+		{
+			ActorOwner->OnEndPlay.AddUniqueDynamic(this, &UAbilitySystemComponent::OnSpawnedAttributesEndPlayed);
+		}
+	}
 }
 
 TArray<UAttributeSet*>& UAbilitySystemComponent::GetSpawnedAttributes_Mutable()
@@ -2721,6 +2772,19 @@ const TArray<UAttributeSet*>& UAbilitySystemComponent::GetSpawnedAttributes() co
 	return SpawnedAttributes;
 }
 
+void UAbilitySystemComponent::OnSpawnedAttributesEndPlayed(AActor* InActor, EEndPlayReason::Type EndPlayReason)
+{
+	TArray<UAttributeSet*>& LocalSpawnedAttributes = GetSpawnedAttributes_Mutable();
+	for (int32 Index = LocalSpawnedAttributes.Num() - 1; Index >= 0; --Index)
+	{
+		UAttributeSet* AttributeSet = LocalSpawnedAttributes[Index];
+		if (AttributeSet->GetTypedOuter<AActor>() == InActor)
+		{
+			LocalSpawnedAttributes[Index] = nullptr;
+		}
+	}
+}
+
 void UAbilitySystemComponent::SetClientDebugStrings(TArray<FString>&& NewClientDebugStrings)
 {
 	GetClientDebugStrings_Mutable() = NewClientDebugStrings;
diff --git a/Engine/Plugins/Runtime/GameplayAbilities/Source/GameplayAbilities/Public/AbilitySystemComponent.h b/Engine/Plugins/Runtime/GameplayAbilities/Source/GameplayAbilities/Public/AbilitySystemComponent.h
--- a/Engine/Plugins/Runtime/GameplayAbilities/Source/GameplayAbilities/Public/AbilitySystemComponent.h
+++ b/Engine/Plugins/Runtime/GameplayAbilities/Source/GameplayAbilities/Public/AbilitySystemComponent.h
@@ -1239,6 +1239,15 @@ class GAMEPLAYABILITIES_API UAbilitySystemComponent : public UGameplayTasksCompo
 	UFUNCTION()
 	void OnRep_OwningActor();
 
+	UFUNCTION()
+	void OnAvatarActorDestroyed(AActor* InActor);
+
+	UFUNCTION()
+	void OnOwnerActorDestroyed(AActor* InActor);
+
+	UFUNCTION()
+	void OnSpawnedAttributesEndPlayed(AActor* InActor, EEndPlayReason::Type EndPlayReason);
+
 	/** Cached off data about the owning actor that abilities will need to frequently access (movement component, mesh component, anim instance, etc) */
 	TSharedPtr<FGameplayAbilityActorInfo>	AbilityActorInfo;
 
