Fixing crash when showing 'AbilitySystem' debug info from the command line when a not-fully-initialized ASC is in play

diff --git a/Engine/Plugins/Runtime/GameplayAbilities/Source/GameplayAbilities/Private/AbilitySystemComponent.cpp b/Engine/Plugins/Runtime/GameplayAbilities/Source/GameplayAbilities/Private/AbilitySystemComponent.cpp
--- a/Engine/Plugins/Runtime/GameplayAbilities/Source/GameplayAbilities/Private/AbilitySystemComponent.cpp
+++ b/Engine/Plugins/Runtime/GameplayAbilities/Source/GameplayAbilities/Private/AbilitySystemComponent.cpp
@@ -2018,7 +2018,7 @@ UAbilitySystemComponent* GetDebugTarget(FASCDebugTargetInfo* Info)
 			if (ASC->GetWorld() == Info->TargetWorld.Get() && MakeWeakObjectPtr(ASC).Get())
 			{
 				Info->LastDebugTarget = ASC;
-				if (ASC->AbilityActorInfo->IsLocallyControlledPlayer())
+				if (ASC->AbilityActorInfo != nullptr && ASC->AbilityActorInfo->IsLocallyControlledPlayer())
 				{
 					// Default to local player first
 					break;
