Fix a bug where effects were being removed based on the source data of a new effect being applied instead of the target data.

diff --git a/Engine/Plugins/Runtime/GameplayAbilities/Source/GameplayAbilities/Private/GameplayEffect.cpp b/Engine/Plugins/Runtime/GameplayAbilities/Source/GameplayAbilities/Private/GameplayEffect.cpp
--- a/Engine/Plugins/Runtime/GameplayAbilities/Source/GameplayAbilities/Private/GameplayEffect.cpp
+++ b/Engine/Plugins/Runtime/GameplayAbilities/Source/GameplayAbilities/Private/GameplayEffect.cpp
@@ -4306,7 +4306,7 @@ void FActiveGameplayEffectsContainer::AttemptRemoveActiveEffectsOnEffectApplicat
 			}
 		}
 
-		const FGameplayTagContainer* AggregatedSourceTags = InSpec.CapturedSourceTags.GetAggregatedTags();
+		const FGameplayTagContainer* AggregatedTargetTags = InSpec.CapturedTargetTags.GetAggregatedTags();
 		for (int32 idx = GetNumGameplayEffects() - 1; idx >= 0; --idx)
 		{
 			const FActiveGameplayEffect& ActiveEffect = *GetActiveGameplayEffect(idx);
@@ -4319,7 +4319,7 @@ void FActiveGameplayEffectsContainer::AttemptRemoveActiveEffectsOnEffectApplicat
 					InternalRemoveActiveGameplayEffect(idx, -1, true);
 				}
 				// check if the effect has removal tag requirements
-				else if (Effect && AggregatedSourceTags && !Effect->RemovalTagRequirements.IsEmpty() && Effect->RemovalTagRequirements.RequirementsMet(*AggregatedSourceTags))
+				else if (Effect && AggregatedTargetTags && !Effect->RemovalTagRequirements.IsEmpty() && Effect->RemovalTagRequirements.RequirementsMet(*AggregatedTargetTags))
 				{
 					InternalRemoveActiveGameplayEffect(idx, -1, true);
 				}
