Expanded GAS Blueprint access layer and utility. Includes:
- Giving, querying and clearing abilities.
- Querying gameplay tags, attributes and attribute values
- Ability and Active Effect Handle comparisons

diff --git a/Engine/Plugins/Runtime/GameplayAbilities/Source/GameplayAbilities/Private/AbilitySystemBlueprintLibrary.cpp b/Engine/Plugins/Runtime/GameplayAbilities/Source/GameplayAbilities/Private/AbilitySystemBlueprintLibrary.cpp
--- a/Engine/Plugins/Runtime/GameplayAbilities/Source/GameplayAbilities/Private/AbilitySystemBlueprintLibrary.cpp
+++ b/Engine/Plugins/Runtime/GameplayAbilities/Source/GameplayAbilities/Private/AbilitySystemBlueprintLibrary.cpp
@@ -2,10 +2,11 @@
 
 #include "AbilitySystemBlueprintLibrary.h"
 #include "GameplayEffectAggregator.h"
-#include "GameplayEffect.h"
 #include "AbilitySystemGlobals.h"
 #include "AbilitySystemComponent.h"
+#include "GameplayEffect.h"
 #include "GameplayEffectUIData.h"
+#include "GameplayAbilitySpec.h"
 
 UAbilitySystemBlueprintLibrary::UAbilitySystemBlueprintLibrary(const FObjectInitializer& ObjectInitializer)
 : Super(ObjectInitializer)
@@ -1077,3 +1078,72 @@ bool UAbilitySystemBlueprintLibrary::RemoveLooseGameplayTags(AActor* Actor, cons
 
 	return false;
 }
+
+bool UAbilitySystemBlueprintLibrary::EqualEqual_ActiveGameplayEffectHandle(const FActiveGameplayEffectHandle& A, const FActiveGameplayEffectHandle& B)
+{
+	return A == B;
+}
+
+bool UAbilitySystemBlueprintLibrary::NotEqual_ActiveGameplayEffectHandle(const FActiveGameplayEffectHandle& A, const FActiveGameplayEffectHandle& B)
+{
+	return A != B;
+}
+
+const UGameplayEffect* UAbilitySystemBlueprintLibrary::GetGameplayEffectFromActiveEffectHandle(const FActiveGameplayEffectHandle& ActiveHandle)
+{
+	const UAbilitySystemComponent* ASC = ActiveHandle.GetOwningAbilitySystemComponent();
+
+	if (ASC)
+	{
+		return ASC->GetGameplayEffectCDO(ActiveHandle);
+	}
+
+	ABILITY_LOG(Error, TEXT("GetGameplayAbilityFromSpecHandle() called with an invalid Active Gameplay Effect Handle"));
+
+	return nullptr;
+}
+
+const UGameplayAbility* UAbilitySystemBlueprintLibrary::GetGameplayAbilityFromSpecHandle(UAbilitySystemComponent* AbilitySystem, const FGameplayAbilitySpecHandle& AbilitySpecHandle, bool& bIsInstance)
+{
+	// validate the ASC
+	if (!AbilitySystem)
+	{
+		ABILITY_LOG(Error, TEXT("GetGameplayAbilityFromSpecHandle() called with an invalid Ability System Component"));
+
+		bIsInstance = false;
+		return nullptr;
+	}
+
+	// get and validate the ability spec
+	FGameplayAbilitySpec* AbilitySpec = AbilitySystem->FindAbilitySpecFromHandle(AbilitySpecHandle);
+	if (!AbilitySpec)
+	{
+		ABILITY_LOG(Error, TEXT("GetGameplayAbilityFromSpecHandle() Ability Spec not found on passed Ability System Component"));
+
+		bIsInstance = false;
+		return nullptr;
+	}
+
+	// try to get the ability instance
+	UGameplayAbility* AbilityInstance = AbilitySpec->GetPrimaryInstance();
+	bIsInstance = true;
+
+	// default to the CDO if we can't
+	if (!AbilityInstance)
+	{
+		AbilityInstance = AbilitySpec->Ability;
+		bIsInstance = false;
+	}
+
+	return AbilityInstance;
+}
+
+bool UAbilitySystemBlueprintLibrary::EqualEqual_GameplayAbilitySpecHandle(const FGameplayAbilitySpecHandle& A, const FGameplayAbilitySpecHandle& B)
+{
+	return A == B;
+}
+
+bool UAbilitySystemBlueprintLibrary::NotEqual_GameplayAbilitySpecHandle(const FGameplayAbilitySpecHandle& A, const FGameplayAbilitySpecHandle& B)
+{
+	return A != B;
+}
\ No newline at end of file
diff --git a/Engine/Plugins/Runtime/GameplayAbilities/Source/GameplayAbilities/Private/AbilitySystemComponent.cpp b/Engine/Plugins/Runtime/GameplayAbilities/Source/GameplayAbilities/Private/AbilitySystemComponent.cpp
--- a/Engine/Plugins/Runtime/GameplayAbilities/Source/GameplayAbilities/Private/AbilitySystemComponent.cpp
+++ b/Engine/Plugins/Runtime/GameplayAbilities/Source/GameplayAbilities/Private/AbilitySystemComponent.cpp
@@ -128,7 +128,7 @@ bool UAbilitySystemComponent::HasAttributeSetForAttribute(FGameplayAttribute Att
 	return (Attribute.IsValid() && (Attribute.IsSystemAttribute() || GetAttributeSubobject(Attribute.GetAttributeSetClass()) != nullptr));
 }
 
-void UAbilitySystemComponent::GetAllAttributes(OUT TArray<FGameplayAttribute>& Attributes)
+void UAbilitySystemComponent::GetAllAttributes(TArray<FGameplayAttribute>& OutAttributes)
 {
 	for (const UAttributeSet* Set : GetSpawnedAttributes())
 	{
@@ -141,16 +141,46 @@ void UAbilitySystemComponent::GetAllAttributes(OUT TArray<FGameplayAttribute>& A
 		{
 			if (FFloatProperty* FloatProperty = CastField<FFloatProperty>(*It))
 			{
-				Attributes.Push(FGameplayAttribute(FloatProperty));
+				OutAttributes.Push(FGameplayAttribute(FloatProperty));
 			}
 			else if (FGameplayAttribute::IsGameplayAttributeDataProperty(*It))
 			{
-				Attributes.Push(FGameplayAttribute(*It));
+				OutAttributes.Push(FGameplayAttribute(*It));
 			}
 		}
 	}
 }
 
+const UAttributeSet* UAbilitySystemComponent::GetAttributeSet(TSubclassOf<UAttributeSet> AttributeSetClass) const
+{
+	// get the attribute set
+	const UAttributeSet* AttributeSet = GetAttributeSubobject(AttributeSetClass);
+
+	// return the pointer
+	return AttributeSet;
+}
+
+float UAbilitySystemComponent::GetGameplayAttributeValue(FGameplayAttribute Attribute, OUT bool& bFound) const
+{
+	// validate the attribute
+	if (Attribute.IsValid())
+	{
+		// get the associated AttributeSet
+		const UAttributeSet* InternalAttributeSet = GetAttributeSubobject(Attribute.GetAttributeSetClass());
+
+		if (InternalAttributeSet)
+		{
+			// NOTE: this is currently not taking predicted gameplay effect modifiers into consideration, so the value may not be accurate on client
+			bFound = true;
+			return Attribute.GetNumericValue(InternalAttributeSet);
+		}
+	}
+
+	// the attribute was not found
+	bFound = false;
+	return 0.0f;
+}
+
 void UAbilitySystemComponent::OnRegister()
 {
 	Super::OnRegister();
@@ -214,6 +244,21 @@ const FActiveGameplayEffect* UAbilitySystemComponent::GetActiveGameplayEffect(co
 	return ActiveGameplayEffects.GetActiveGameplayEffect(Handle);
 }
 
+const UGameplayEffect* UAbilitySystemComponent::GetGameplayEffectCDO(const FActiveGameplayEffectHandle Handle) const
+{
+	// get the active gameplay effect struct
+	const FActiveGameplayEffect* ActiveGE = GetActiveGameplayEffect(Handle);
+
+	if (ActiveGE)
+	{
+		// return the spec's CDO
+		return ActiveGE->Spec.Def;
+	}
+
+	// active effect not found
+	return nullptr;
+}
+
 const FGameplayTagContainer* UAbilitySystemComponent::GetGameplayEffectSourceTagsFromHandle(FActiveGameplayEffectHandle Handle) const
 {
 	return ActiveGameplayEffects.GetGameplayEffectSourceTagsFromHandle(Handle);
@@ -531,6 +576,13 @@ void UAbilitySystemComponent::UpdateTagMap_Internal(const FGameplayTagContainer&
 	}
 }
 
+int32 UAbilitySystemComponent::GetGameplayTagCount(FGameplayTag GameplayTag) const
+{
+	// NOTE -- no authority check to allow read-only access on all clients.
+	// This may lead to incorrect values due to network lag. 
+	return GetTagCount(GameplayTag);
+}
+
 FOnGameplayEffectTagCountChanged& UAbilitySystemComponent::RegisterGameplayTagEvent(FGameplayTag Tag, EGameplayTagEventType::Type EventType)
 {
 	return GameplayTagCountContainer.RegisterGameplayTagEvent(Tag, EventType);
diff --git a/Engine/Plugins/Runtime/GameplayAbilities/Source/GameplayAbilities/Private/AbilitySystemComponent_Abilities.cpp b/Engine/Plugins/Runtime/GameplayAbilities/Source/GameplayAbilities/Private/AbilitySystemComponent_Abilities.cpp
--- a/Engine/Plugins/Runtime/GameplayAbilities/Source/GameplayAbilities/Private/AbilitySystemComponent_Abilities.cpp
+++ b/Engine/Plugins/Runtime/GameplayAbilities/Source/GameplayAbilities/Private/AbilitySystemComponent_Abilities.cpp
@@ -246,9 +246,20 @@ void UAbilitySystemComponent::RefreshAbilityActorInfo()
 }
 
 FGameplayAbilitySpecHandle UAbilitySystemComponent::GiveAbility(const FGameplayAbilitySpec& Spec)
-{	
-	check(Spec.Ability);
-	check(IsOwnerActorAuthoritative());	// Should be called on authority
+{
+	if (!IsValid(Spec.Ability))
+	{
+		ABILITY_LOG(Error, TEXT("GiveAbility called with an invalid Ability Class."));
+
+		return FGameplayAbilitySpecHandle();
+	}
+
+	if (!IsOwnerActorAuthoritative())
+	{
+		ABILITY_LOG(Error, TEXT("GiveAbility called on ability %s on the client, not allowed!"), *Spec.Ability->GetName());
+
+		return FGameplayAbilitySpecHandle();
+	}
 
 	// If locked, add to pending list. The Spec.Handle is not regenerated when we receive, so returning this is ok.
 	if (AbilityScopeLockCount > 0)
@@ -274,7 +285,12 @@ FGameplayAbilitySpecHandle UAbilitySystemComponent::GiveAbility(const FGameplayA
 
 FGameplayAbilitySpecHandle UAbilitySystemComponent::GiveAbilityAndActivateOnce(FGameplayAbilitySpec& Spec, const FGameplayEventData* GameplayEventData)
 {
-	check(Spec.Ability);
+	if (!IsValid(Spec.Ability))
+	{
+		ABILITY_LOG(Error, TEXT("GiveAbilityAndActivateOnce called with an invalid Ability Class."));
+
+		return FGameplayAbilitySpecHandle();
+	}
 
 	if (Spec.Ability->GetInstancingPolicy() == EGameplayAbilityInstancingPolicy::NonInstanced || Spec.Ability->GetNetExecutionPolicy() == EGameplayAbilityNetExecutionPolicy::LocalOnly)
 	{
@@ -312,6 +328,39 @@ FGameplayAbilitySpecHandle UAbilitySystemComponent::GiveAbilityAndActivateOnce(F
 	return AddedAbilityHandle;
 }
 
+FGameplayAbilitySpecHandle UAbilitySystemComponent::K2_GiveAbility(TSubclassOf<UGameplayAbility> AbilityClass, int32 Level /*= 0*/, int32 InputID /*= -1*/)
+{
+	// build and validate the ability spec
+	FGameplayAbilitySpec AbilitySpec = BuildAbilitySpecFromClass(AbilityClass, Level, InputID);
+
+	// validate the class
+	if (!IsValid(AbilitySpec.Ability))
+	{
+		ABILITY_LOG(Error, TEXT("K2_GiveAbility() called with an invalid Ability Class."));
+
+		return FGameplayAbilitySpecHandle();
+	}
+
+	// grant the ability and return the handle. This will run validation and authority checks
+	return GiveAbility(AbilitySpec);
+}
+
+FGameplayAbilitySpecHandle UAbilitySystemComponent::K2_GiveAbilityAndActivateOnce(TSubclassOf<UGameplayAbility> AbilityClass, int32 Level /* = 0*/, int32 InputID /* = -1*/)
+{
+	// build and validate the ability spec
+	FGameplayAbilitySpec AbilitySpec = BuildAbilitySpecFromClass(AbilityClass, Level, InputID);
+
+	// validate the class
+	if (!IsValid(AbilitySpec.Ability))
+	{
+		ABILITY_LOG(Error, TEXT("K2_GiveAbilityAndActivateOnce() called with an invalid Ability Class."));
+
+		return FGameplayAbilitySpecHandle();
+	}
+
+	return GiveAbilityAndActivateOnce(AbilitySpec);
+}
+
 void UAbilitySystemComponent::SetRemoveAbilityOnEnd(FGameplayAbilitySpecHandle AbilitySpecHandle)
 {
 	FGameplayAbilitySpec* FoundSpec = FindAbilitySpecFromHandle(AbilitySpecHandle);
@@ -330,9 +379,15 @@ void UAbilitySystemComponent::SetRemoveAbilityOnEnd(FGameplayAbilitySpecHandle A
 
 void UAbilitySystemComponent::ClearAllAbilities()
 {
-	check(IsOwnerActorAuthoritative());	// Should be called on authority
 	check(AbilityScopeLockCount == 0);	// We should never be calling this from a scoped lock situation.
 
+	if (!IsOwnerActorAuthoritative())
+	{
+		ABILITY_LOG(Error, TEXT("Attempted to call ClearAllAbilities() without authority."));
+
+		return;
+	}
+
 	// Note we aren't marking any old abilities pending kill. This shouldn't matter since they will be garbage collected.
 	ABILITYLIST_SCOPE_LOCK();
 	for (FGameplayAbilitySpec& Spec : ActivatableAbilities.Items)
@@ -347,9 +402,28 @@ void UAbilitySystemComponent::ClearAllAbilities()
 	CheckForClearedAbilities();
 }
 
+void UAbilitySystemComponent::ClearAllAbilitiesWithInputID(int32 InputID /*= 0*/)
+{
+	// find all matching abilities
+	TArray<const FGameplayAbilitySpec*> OutSpecs;
+	FindAllAbilitySpecsFromInputID(InputID, OutSpecs);
+
+	// iterate through the bound abilities
+	for (const FGameplayAbilitySpec* CurrentSpec : OutSpecs)
+	{
+		// clear the ability
+		ClearAbility(CurrentSpec->Handle);
+	}
+}
+
 void UAbilitySystemComponent::ClearAbility(const FGameplayAbilitySpecHandle& Handle)
 {
-	check(IsOwnerActorAuthoritative()); // Should be called on authority
+	if (!IsOwnerActorAuthoritative())
+	{
+		ABILITY_LOG(Error, TEXT("Attempted to call ClearAbility() on the client. This is not allowed!"));
+
+		return;
+	}
 
 	bIsNetDirty = true;
 	for (int Idx = 0; Idx < AbilityPendingAdds.Num(); ++Idx)
@@ -728,6 +802,147 @@ FGameplayAbilitySpec* UAbilitySystemComponent::FindAbilitySpecFromInputID(int32
 	return nullptr;
 }
 
+void UAbilitySystemComponent::FindAllAbilitySpecsFromInputID(int32 InputID, TArray<const FGameplayAbilitySpec*>& OutAbilitySpecs) const
+{
+	// ignore invalid inputs
+	if (InputID != INDEX_NONE)
+	{
+		// iterate through all abilities
+		for (const FGameplayAbilitySpec& Spec : ActivatableAbilities.Items)
+		{
+			// add maching abilities to the list
+			if (Spec.InputID == InputID)
+			{
+				OutAbilitySpecs.Add(&Spec);
+			}
+		}
+	}
+
+	UE_LOG(LogTemp, Log, TEXT("OutAbilitySpecs count: %i"), OutAbilitySpecs.Num());
+}
+
+FGameplayAbilitySpec UAbilitySystemComponent::BuildAbilitySpecFromClass(TSubclassOf<UGameplayAbility> AbilityClass, int32 Level /*= 0*/, int32 InputID /*= -1*/)
+{
+	// validate the class
+	if (!ensure(AbilityClass))
+	{
+		ABILITY_LOG(Error, TEXT("BuildAbilitySpecFromClass called with an invalid Ability Class."));
+
+		return FGameplayAbilitySpec();
+	}
+
+	// get the CDO. GiveAbility will validate so we don't need to
+	UGameplayAbility* AbilityCDO = AbilityClass->GetDefaultObject<UGameplayAbility>();
+
+	// build the ability spec
+	// we need to initialize this through the constructor,
+	// or the Handle won't be properly set and will cause errors further down the line
+	return FGameplayAbilitySpec(AbilityClass, Level, InputID);
+}
+
+void UAbilitySystemComponent::GetAllAbilities(OUT TArray<FGameplayAbilitySpecHandle>& OutAbilityHandles) const
+{
+	// ensure the output array is empty
+	OutAbilityHandles.Empty(ActivatableAbilities.Items.Num());
+
+	// iterate through all activatable abilities
+	// NOTE: currently this doesn't include abilities that are mid-activation
+	for (const FGameplayAbilitySpec& Spec : ActivatableAbilities.Items)
+	{
+		// add the spec handle to the list
+		OutAbilityHandles.Add(Spec.Handle);
+	}
+}
+
+void UAbilitySystemComponent::FindAllAbilitiesWithTags(OUT TArray<FGameplayAbilitySpecHandle>& OutAbilityHandles, FGameplayTagContainer Tags, bool bExactMatch /* = true */) const
+{
+	// ensure the output array is empty
+	OutAbilityHandles.Empty();
+
+	// iterate through all Ability Specs
+	for (const FGameplayAbilitySpec& CurrentSpec : ActivatableAbilities.Items)
+	{
+		// try to get the ability instance
+		UGameplayAbility* AbilityInstance = CurrentSpec.GetPrimaryInstance();
+
+		// default to the CDO if we can't
+		if (!AbilityInstance)
+		{
+			AbilityInstance = CurrentSpec.Ability;
+		}
+
+		// ensure the ability instance is valid
+		if (IsValid(AbilityInstance))
+		{
+			// do we want an exact match?
+			if (bExactMatch)
+			{
+				// check if we match all tags
+				if (AbilityInstance->AbilityTags.HasAll(Tags))
+				{
+					// add the matching handle
+					OutAbilityHandles.Add(CurrentSpec.Handle);
+				}
+			}
+			else
+			{
+
+				// check if we match any tags
+				if (AbilityInstance->AbilityTags.HasAny(Tags))
+				{
+					// add the matching handle
+					OutAbilityHandles.Add(CurrentSpec.Handle);
+				}
+			}
+		}
+	}
+}
+
+void UAbilitySystemComponent::FindAllAbilitiesMatchingQuery(OUT TArray<FGameplayAbilitySpecHandle>& OutAbilityHandles, FGameplayTagQuery Query) const
+{
+	// ensure the output array is empty
+	OutAbilityHandles.Empty();
+
+	// iterate through all Ability Specs
+	for (const FGameplayAbilitySpec& CurrentSpec : ActivatableAbilities.Items)
+	{
+		// try to get the ability instance
+		UGameplayAbility* AbilityInstance = CurrentSpec.GetPrimaryInstance();
+
+		// default to the CDO if we can't
+		if (!AbilityInstance)
+		{
+			AbilityInstance = CurrentSpec.Ability;
+		}
+
+		// ensure the ability instance is valid
+		if (IsValid(AbilityInstance))
+		{
+			if (AbilityInstance->AbilityTags.MatchesQuery(Query))
+			{
+				// add the matching handle
+				OutAbilityHandles.Add(CurrentSpec.Handle);
+			}
+		}
+	}
+}
+
+void UAbilitySystemComponent::FindAllAbilitiesWithInputID(OUT TArray<FGameplayAbilitySpecHandle>& OutAbilityHandles, int32 InputID /*= 0*/) const
+{
+	// ensure the output array is empty
+	OutAbilityHandles.Empty();
+
+	// find all ability specs matching the Input ID
+	TArray<const FGameplayAbilitySpec*> MatchingSpecs;
+	FindAllAbilitySpecsFromInputID(InputID, MatchingSpecs);
+
+	// add all matching specs to the out array
+	for (const FGameplayAbilitySpec* CurrentSpec : MatchingSpecs)
+	{
+		OutAbilityHandles.Add(CurrentSpec->Handle);
+	}
+}
+
 FGameplayEffectContextHandle UAbilitySystemComponent::GetEffectContextFromActiveGEHandle(FActiveGameplayEffectHandle Handle)
 {
 	FActiveGameplayEffect* ActiveGE = ActiveGameplayEffects.GetActiveGameplayEffect(Handle);
@@ -2295,6 +2510,16 @@ void UAbilitySystemComponent::AbilityLocalInputReleased(int32 InputID)
 	}
 }
 
+void UAbilitySystemComponent::PressInputID(int32 InputID)
+{
+	AbilityLocalInputPressed(InputID);
+}
+
+void UAbilitySystemComponent::ReleaseInputID(int32 InputID)
+{
+	AbilityLocalInputReleased(InputID);
+}
+
 void UAbilitySystemComponent::ServerSetInputPressed_Implementation(FGameplayAbilitySpecHandle AbilityHandle)
 {
 	FGameplayAbilitySpec* Spec = FindAbilitySpecFromHandle(AbilityHandle);
@@ -2380,6 +2605,16 @@ void UAbilitySystemComponent::LocalInputCancel()
 	Temp.Broadcast();
 }
 
+void UAbilitySystemComponent::InputConfirm()
+{
+	LocalInputConfirm();
+}
+
+void UAbilitySystemComponent::InputCancel()
+{
+	LocalInputCancel();
+}
+
 void UAbilitySystemComponent::TargetConfirm()
 {
 	// Callbacks may modify the spawned target actor array so iterate over a copy instead
diff --git a/Engine/Plugins/Runtime/GameplayAbilities/Source/GameplayAbilities/Private/GameplayModMagnitudeCalculation.cpp b/Engine/Plugins/Runtime/GameplayAbilities/Source/GameplayAbilities/Private/GameplayModMagnitudeCalculation.cpp
--- a/Engine/Plugins/Runtime/GameplayAbilities/Source/GameplayAbilities/Private/GameplayModMagnitudeCalculation.cpp
+++ b/Engine/Plugins/Runtime/GameplayAbilities/Source/GameplayAbilities/Private/GameplayModMagnitudeCalculation.cpp
@@ -40,3 +40,84 @@ bool UGameplayModMagnitudeCalculation::GetCapturedAttributeMagnitude(const FGame
 
 	return true;
 }
+
+float UGameplayModMagnitudeCalculation::K2_GetCapturedAttributeMagnitude(const FGameplayEffectSpec& EffectSpec, FGameplayAttribute Attribute, const FGameplayTagContainer& SourceTags, const FGameplayTagContainer& TargetTags)
+{
+	float Magnitude = 0.0f;
+
+	// look for the attribute in the capture list
+	for (const FGameplayEffectAttributeCaptureDefinition& CurrentCapture : RelevantAttributesToCapture)
+	{
+		if (CurrentCapture.AttributeToCapture == Attribute)
+		{
+			// configure the aggregator evaluation parameters
+			// TODO: apply filters?
+			FAggregatorEvaluateParameters EvaluationParameters;
+
+			EvaluationParameters.SourceTags = &SourceTags;
+			EvaluationParameters.TargetTags = &TargetTags;
+
+			// get the attribute magnitude
+			GetCapturedAttributeMagnitude(CurrentCapture, EffectSpec, EvaluationParameters, Magnitude);
+
+			// capture found. Stop the search
+			break;
+		}
+	}
+
+	return Magnitude;
+}
+
+float UGameplayModMagnitudeCalculation::GetSetByCallerMagnitudeByTag(const FGameplayEffectSpec& EffectSpec, const FGameplayTag& Tag)
+{
+	return EffectSpec.GetSetByCallerMagnitude(Tag, true, 0.0f);
+}
+
+float UGameplayModMagnitudeCalculation::GetSetByCallerMagnitudeByName(const FGameplayEffectSpec& EffectSpec, const FName& MagnitudeName)
+{
+	return EffectSpec.GetSetByCallerMagnitude(MagnitudeName, true, 0.0f);
+}
+
+FGameplayTagContainer UGameplayModMagnitudeCalculation::GetSourceAggregatedTags(const FGameplayEffectSpec& EffectSpec)
+{
+	const FGameplayTagContainer* Tags = EffectSpec.CapturedSourceTags.GetAggregatedTags();
+
+	if (Tags)
+	{
+		return *Tags;
+	}
+
+	return FGameplayTagContainer();
+}
+
+const FGameplayTagContainer& UGameplayModMagnitudeCalculation::GetSourceActorTags(const FGameplayEffectSpec& EffectSpec)
+{
+	return EffectSpec.CapturedSourceTags.GetActorTags();
+}
+
+const FGameplayTagContainer& UGameplayModMagnitudeCalculation::GetSourceSpecTags(const FGameplayEffectSpec& EffectSpec)
+{
+	return EffectSpec.CapturedSourceTags.GetSpecTags();
+}
+
+FGameplayTagContainer UGameplayModMagnitudeCalculation::GetTargetAggregatedTags(const FGameplayEffectSpec& EffectSpec)
+{
+	const FGameplayTagContainer* Tags = EffectSpec.CapturedTargetTags.GetAggregatedTags();
+
+	if (Tags)
+	{
+		return *Tags;
+	}
+
+	return FGameplayTagContainer();
+}
+
+const FGameplayTagContainer& UGameplayModMagnitudeCalculation::GetTargetActorTags(const FGameplayEffectSpec& EffectSpec)
+{
+	return EffectSpec.CapturedTargetTags.GetActorTags();
+}
+
+const FGameplayTagContainer& UGameplayModMagnitudeCalculation::GetTargetSpecTags(const FGameplayEffectSpec& EffectSpec)
+{
+	return EffectSpec.CapturedTargetTags.GetSpecTags();
+}
diff --git a/Engine/Plugins/Runtime/GameplayAbilities/Source/GameplayAbilities/Public/AbilitySystemBlueprintLibrary.h b/Engine/Plugins/Runtime/GameplayAbilities/Source/GameplayAbilities/Public/AbilitySystemBlueprintLibrary.h
--- a/Engine/Plugins/Runtime/GameplayAbilities/Source/GameplayAbilities/Public/AbilitySystemBlueprintLibrary.h
+++ b/Engine/Plugins/Runtime/GameplayAbilities/Source/GameplayAbilities/Public/AbilitySystemBlueprintLibrary.h
@@ -388,4 +388,44 @@ class GAMEPLAYABILITIES_API UAbilitySystemBlueprintLibrary : public UBlueprintFu
 	/** Returns the UI data for a gameplay effect class (if any) */
 	UFUNCTION(BlueprintCallable, Category = "Ability|GameplayEffect", Meta = (DisplayName = "Get GameplayEffect UI Data", DeterminesOutputType="DataType"))
 	static const UGameplayEffectUIData* GetGameplayEffectUIData(TSubclassOf<UGameplayEffect> EffectClass, TSubclassOf<UGameplayEffectUIData> DataType);
+
+	/** Equality operator for two Active Gameplay Effect Handles */
+	UFUNCTION(BlueprintPure, Category = "Ability|GameplayEffect", meta = (DisplayName = "Equal (Active Gameplay Effect Handle)", CompactNodeTitle = "==", ScriptOperator = "=="))
+	static bool EqualEqual_ActiveGameplayEffectHandle(const FActiveGameplayEffectHandle& A, const FActiveGameplayEffectHandle& B);
+
+	/** Inequality operator for two Active Gameplay Effect Handles */
+	UFUNCTION(BlueprintPure, Category = "Ability|GameplayEffect", meta = (DisplayName = "Not Equal (Active Gameplay Effect Handle)", CompactNodeTitle = "!=", ScriptOperator = "!="))
+	static bool NotEqual_ActiveGameplayEffectHandle(const FActiveGameplayEffectHandle& A, const FActiveGameplayEffectHandle& B);
+
+	/**
+	 * Returns the Gameplay Effect CDO from an active handle.
+	 * This reference should be considered read only,
+	 * but you can use it to read additional Gameplay Effect info, such as icon, description, etc. 
+	 */
+	UFUNCTION(BlueprintPure, Category = "Ability|GameplayEffect")
+	static const UGameplayEffect* GetGameplayEffectFromActiveEffectHandle(const FActiveGameplayEffectHandle& ActiveHandle);
+
+	// -------------------------------------------------------------------------------
+	//		GameplayAbility
+	// -------------------------------------------------------------------------------
+
+	/**
+	 * Provides the Gameplay Ability object associated with an Ability Spec Handle
+	 * This can be either an instanced ability, or in the case of shared abilities, the Class Default Object
+	 * 
+	 * @param AbilitySpec The Gameplay Ability Spec you want to get the object from
+	 * @param bIsInstance Set to true if this is an instanced ability instead of a shared CDO
+	 * 
+	 * @return Pointer to the Gameplay Ability object
+	 */
+	UFUNCTION(BlueprintPure, Category = "Ability|GameplayAbility")
+	static const UGameplayAbility* GetGameplayAbilityFromSpecHandle(UAbilitySystemComponent* AbilitySystem, const FGameplayAbilitySpecHandle& AbilitySpecHandle, bool& bIsInstance);
+
+	/** Equality operator for two Gameplay Ability Spec Handles */
+	UFUNCTION(BlueprintPure, Category = "Ability|GameplayEffect", meta = (DisplayName = "Equal (Gameplay Ability Spec Handle)", CompactNodeTitle = "==", ScriptOperator = "=="))
+	static bool EqualEqual_GameplayAbilitySpecHandle(const FGameplayAbilitySpecHandle& A, const FGameplayAbilitySpecHandle& B);
+
+	/** Inequality operator for two Gameplay Ability Spec Handles */
+	UFUNCTION(BlueprintPure, Category = "Ability|GameplayEffect", meta = (DisplayName = "Not Equal (Gameplay Ability Spec Handle)", CompactNodeTitle = "!=", ScriptOperator = "!="))
+	static bool NotEqual_GameplayAbilitySpecHandle(const FGameplayAbilitySpecHandle& A, const FGameplayAbilitySpecHandle& B);
 };
diff --git a/Engine/Plugins/Runtime/GameplayAbilities/Source/GameplayAbilities/Public/AbilitySystemComponent.h b/Engine/Plugins/Runtime/GameplayAbilities/Source/GameplayAbilities/Public/AbilitySystemComponent.h
--- a/Engine/Plugins/Runtime/GameplayAbilities/Source/GameplayAbilities/Public/AbilitySystemComponent.h
+++ b/Engine/Plugins/Runtime/GameplayAbilities/Source/GameplayAbilities/Public/AbilitySystemComponent.h
@@ -142,9 +142,9 @@ class GAMEPLAYABILITIES_API UAbilitySystemComponent : public UGameplayTasksCompo
 
 	/**
 	 * Does this ability system component have this attribute?
-	 *
+	 * 
 	 * @param Attribute	Handle of the gameplay effect to retrieve target tags from
-	 *
+	 * 
 	 * @return true if Attribute is valid and this ability system component contains an attribute set that contains Attribute. Returns false otherwise.
 	 */
 	bool HasAttributeSetForAttribute(FGameplayAttribute Attribute) const;
@@ -155,8 +155,28 @@ class GAMEPLAYABILITIES_API UAbilitySystemComponent : public UGameplayTasksCompo
 	UFUNCTION(BlueprintCallable, Category="Skills", meta=(DisplayName="InitStats", ScriptName="InitStats"))
 	void K2_InitStats(TSubclassOf<class UAttributeSet> Attributes, const UDataTable* DataTable);
 		
-	/** Returns a list of all attributes for this abiltiy system component */
-	void GetAllAttributes(OUT TArray<FGameplayAttribute>& Attributes);
+	/** Returns a list of all attributes for this abilty system component */
+	UFUNCTION(BlueprintPure, Category="Gameplay Attributes")
+	void GetAllAttributes(TArray<FGameplayAttribute>& OutAttributes);
+
+	/**
+	 * Returns a reference to the Attribute Set instance, if one exists in this component
+	 *
+	 * @param AttributeSetClass The type of attribute set to look for
+	 * @param bFound Set to true if an instance of the Attribute Set exists
+	 */
+	UFUNCTION(BlueprintCallable, BlueprintPure = false, Category = "Gameplay Attributes")
+	const UAttributeSet* GetAttributeSet(TSubclassOf<UAttributeSet> AttributeSetClass) const;
+
+	/**
+	 * Returns the current value of the given gameplay attribute, or zero if the attribute is not found.
+	 * NOTE: This doesn't take predicted gameplay effect modifiers into consideration, so the value may not be accurate on clients at all times.
+	 *
+	 * @param Attribute The gameplay attribute to query
+	 * @param bFound Set to true if the attribute exists in this component
+	 */
+	UFUNCTION(BlueprintPure, Category = "Gameplay Attributes")
+	float GetGameplayAttributeValue(FGameplayAttribute Attribute, bool& bFound) const;
 
 	UPROPERTY(EditAnywhere, Category="AttributeTest")
 	TArray<FAttributeDefaults>	DefaultStartingData;
@@ -185,10 +205,8 @@ class GAMEPLAYABILITIES_API UAbilitySystemComponent : public UGameplayTasksCompo
 	/**
 	 *	Applies an in-place mod to the given attribute. This correctly update the attribute's aggregator, updates the attribute set property,
 	 *	and invokes the OnDirty callbacks.
-	 *	
 	 *	This does not invoke Pre/PostGameplayEffectExecute calls on the attribute set. This does no tag checking, application requirements, immunity, etc.
 	 *	No GameplayEffectSpec is created or is applied!
-	 *
 	 *	This should only be used in cases where applying a real GameplayEffectSpec is too slow or not possible.
 	 */
 	void ApplyModToAttribute(const FGameplayAttribute &Attribute, TEnumAsByte<EGameplayModOp::Type> ModifierOp, float ModifierMagnitude);
@@ -299,7 +317,7 @@ class GAMEPLAYABILITIES_API UAbilitySystemComponent : public UGameplayTasksCompo
 
 	/** 
 	 * Remove active gameplay effects whose backing definition are the specified gameplay effect class
-	 *
+	 * 
 	 * @param GameplayEffect					Class of gameplay effect to remove; Does nothing if left null
 	 * @param InstigatorAbilitySystemComponent	If specified, will only remove gameplay effects applied from this instigator ability system component
 	 * @param StacksToRemove					Number of stacks to remove, -1 means remove all
@@ -375,9 +393,12 @@ class GAMEPLAYABILITIES_API UAbilitySystemComponent : public UGameplayTasksCompo
 	/** Gets the GE Handle of the GE that granted the passed in Ability */
 	FActiveGameplayEffectHandle FindActiveGameplayEffectHandle(FGameplayAbilitySpecHandle Handle) const;
 
-	/** Returns const pointer to the actual active gamepay effect structure */
+	/** Returns const pointer to the actual active gameplay effect structure */
 	const FActiveGameplayEffect* GetActiveGameplayEffect(const FActiveGameplayEffectHandle Handle) const;
 
+	/** Returns a const pointer to the gameplay effect CDO associated with an active handle. */
+	const UGameplayEffect* GetGameplayEffectCDO(const FActiveGameplayEffectHandle Handle) const;
+
 	/**
 	 * Get the source tags from the gameplay spec represented by the specified handle, if possible
 	 * 
@@ -544,9 +565,7 @@ class GAMEPLAYABILITIES_API UAbilitySystemComponent : public UGameplayTasksCompo
 
 	/** 	 
 	 *  Allows GameCode to add loose gameplaytags which are not backed by a GameplayEffect. 
-	 *
 	 *	Tags added this way are not replicated! 
-	 *	
 	 *	It is up to the calling GameCode to make sure these tags are added on clients/server where necessary
 	 */
 	FORCEINLINE void AddLooseGameplayTag(const FGameplayTag& GameplayTag, int32 Count=1)
@@ -574,6 +593,16 @@ class GAMEPLAYABILITIES_API UAbilitySystemComponent : public UGameplayTasksCompo
 		SetTagMapCount(GameplayTag, NewCount);
 	}
 
+	/**
+	 * Returns the current count of the given gameplay tag.
+	 * This includes both loose tags, and tags granted by gameplay effects and abilities.
+	 * This function can be called on the client, but it may not display the most current count on the server.
+	 *
+	 * @param GameplayTag The gameplay tag to query
+	 */
+	UFUNCTION(BlueprintPure, Category = "Gameplay Tags")
+	int32 GetGameplayTagCount(FGameplayTag GameplayTag) const;
+
 	/** 	 
 	 * Minimally replicated tags are replicated tags that come from GEs when in bMinimalReplication mode. 
 	 * (The GEs do not replicate, but the tags they grant do replicate via these functions)
@@ -828,16 +857,66 @@ class GAMEPLAYABILITIES_API UAbilitySystemComponent : public UGameplayTasksCompo
 	 *	
 	 */
 
-	/** Grants Ability. Returns handle that can be used in TryActivateAbility, etc. */
+	/*
+	 * Grants an Ability.
+	 * This will be ignored if the actor is not authoritative.
+	 * Returns handle that can be used in TryActivateAbility, etc.
+	 * 
+	 * @param AbilitySpec FGameplayAbilitySpec containing information about the ability class, level and input ID to bind it to.
+	 */
 	FGameplayAbilitySpecHandle GiveAbility(const FGameplayAbilitySpec& AbilitySpec);
 
-	/** Grants an ability and attempts to activate it exactly one time, which will cause it to be removed. Only valid on the server! */
+	/*
+	 * Grants an ability and attempts to activate it exactly one time, which will cause it to be removed.
+	 * Only valid on the server, and the ability's Net Execution Policy cannot be set to Local or Local Predicted
+	 * 
+	 * @param AbilitySpec FGameplayAbilitySpec containing information about the ability class, level and input ID to bind it to.
+	 * @param GameplayEventData Optional activation event data. If provided, Activate Ability From Event will be called instead of ActivateAbility, passing the Event Data
+	 */
 	FGameplayAbilitySpecHandle GiveAbilityAndActivateOnce(FGameplayAbilitySpec& AbilitySpec, const FGameplayEventData* GameplayEventData = nullptr);
 
-	/** Wipes all 'given' abilities. */
+	/**
+	 * Grants a Gameplay Ability and returns its handle.
+	 * This will be ignored if the actor is not authoritative.
+	 *
+	 * @param AbilityClass Type of ability to grant
+	 * @param Level Level to grant the ability at
+	 * @param InputID Input ID value to bind ability activation to.
+	 */
+	UFUNCTION(BlueprintCallable, BlueprintAuthorityOnly, Category = "Gameplay Abilities", meta = (DisplayName = "Give Ability", ScriptName = "GiveAbility"))
+	FGameplayAbilitySpecHandle K2_GiveAbility(TSubclassOf<UGameplayAbility> AbilityClass, int32 Level = 0, int32 InputID = -1);
+
+	/**
+	 * Grants a Gameplay Ability, activates it once, and removes it.
+	 * This will be ignored if the actor is not authoritative.
+	 *
+	 * @param AbilityClass Type of ability to grant
+	 * @param Level Level to grant the ability at
+	 * @param InputID Input ID value to bind ability activation to.
+	 */
+	UFUNCTION(BlueprintCallable, BlueprintAuthorityOnly, Category = "Gameplay Abilities", meta = (DisplayName = "Give Ability And Activate Once", ScriptName = "GiveAbilityAndActivateOnce"))
+	FGameplayAbilitySpecHandle K2_GiveAbilityAndActivateOnce(TSubclassOf<UGameplayAbility> AbilityClass, int32 Level = 0, int32 InputID = -1);
+
+	/** Wipes all 'given' abilities. This will be ignored if the actor is not authoritative. */
+	UFUNCTION(BlueprintCallable, BlueprintAuthorityOnly, Category="Gameplay Abilities")
 	void ClearAllAbilities();
 
-	/** Removes the specified ability */
+	/**
+	 * Clears all abilities bound to a given Input ID
+	 * This will be ignored if the actor is not authoritative
+	 *
+	 * @param InputID The numeric Input ID of the abilities to remove
+	 */
+	UFUNCTION(BlueprintCallable, BlueprintAuthorityOnly, Category = "Gameplay Abilities")
+	void ClearAllAbilitiesWithInputID(int32 InputID = 0);
+
+	/** 
+	 * Removes the specified ability.
+	 * This will be ignored if the actor is not authoritative.
+	 * 
+	 * @param Handle Ability Spec Handle of the ability we want to remove
+	 */
+	UFUNCTION(BlueprintCallable, BlueprintAuthorityOnly, Category = "Gameplay Abilities")
 	void ClearAbility(const FGameplayAbilitySpecHandle& Handle);
 	
 	/** Sets an ability spec to remove when its finished. If the spec is not currently active, it terminates it immediately. Also clears InputID of the Spec. */
@@ -875,6 +954,7 @@ class GAMEPLAYABILITIES_API UAbilitySystemComponent : public UGameplayTasksCompo
 	 * Returns true if it thinks it activated, but it may return false positives due to failure later in activation.
 	 * If bAllowRemoteActivation is true, it will remotely activate local/server abilities, if false it will only try to locally activate the ability
 	 */
+	UFUNCTION(BlueprintCallable, Category = "Abilities")
 	bool TryActivateAbility(FGameplayAbilitySpecHandle AbilityToActivate, bool bAllowRemoteActivation = true);
 
 	/** Triggers an ability from a gameplay event, will only trigger on local/server depending on execution flags */
@@ -901,6 +981,7 @@ class GAMEPLAYABILITIES_API UAbilitySystemComponent : public UGameplayTasksCompo
 
 	/** 
 	 * Called from ability activation or native code, will apply the correct ability blocking tags and cancel existing abilities. Subclasses can override the behavior 
+	 * 
 	 * @param AbilityTags The tags of the ability that has block and cancel flags
 	 * @param RequestingAbility The gameplay ability requesting the change, can be NULL for native events
 	 * @param bEnableBlockTags If true will enable the block tags, if false will disable the block tags
@@ -931,12 +1012,13 @@ class GAMEPLAYABILITIES_API UAbilitySystemComponent : public UGameplayTasksCompo
 	// Functions meant to be called from GameplayAbility and subclasses, but not meant for general use
 	// ----------------------------------------------------------------------------------------------------------------
 
-	/** Returns the list of all activatable abilities */
+	/** Returns the list of all activatable abilities. Read-only. */
 	const TArray<FGameplayAbilitySpec>& GetActivatableAbilities() const
 	{
 		return ActivatableAbilities.Items;
 	}
 
+	/** Returns the list of all activatable abilities. */
 	TArray<FGameplayAbilitySpec>& GetActivatableAbilities()
 	{
 		return ActivatableAbilities.Items;
@@ -957,6 +1039,56 @@ class GAMEPLAYABILITIES_API UAbilitySystemComponent : public UGameplayTasksCompo
 	/** Returns an ability spec from a handle. If modifying call MarkAbilitySpecDirty */
 	FGameplayAbilitySpec* FindAbilitySpecFromInputID(int32 InputID);
 
+	/**
+	 * Returns all abilities with the given InputID
+	 *
+	 * @param InputID The Input ID to match
+	 * @param OutAbilitySpecs Array of pointers to matching specs
+	 */
+	virtual void FindAllAbilitySpecsFromInputID(int32 InputID, TArray<const FGameplayAbilitySpec*>& OutAbilitySpecs) const;
+
+	/**
+	 * Build a simple FGameplayAbilitySpec from class, level and optional Input ID
+	 */
+	virtual FGameplayAbilitySpec BuildAbilitySpecFromClass(TSubclassOf<UGameplayAbility> AbilityClass, int32 Level = 0, int32 InputID = -1);
+
+	/**
+	 * Returns an array with all granted ability handles
+	 * NOTE: currently this doesn't include abilities that are mid-activation
+	 * 
+	 * @param OutAbilityHandles This array will be filled with the granted Ability Spec Handles
+	 */
+	UFUNCTION(BlueprintCallable, BlueprintPure = false, Category = "Gameplay Abilities")
+	void GetAllAbilities(TArray<FGameplayAbilitySpecHandle>& OutAbilityHandles) const;
+
+	/**
+	 * Returns an array with all abilities that match the provided tags
+	 *
+	 * @param OutAbilityHandles This array will be filled with matching Ability Spec Handles
+	 * @param Tags Gameplay Tags to match
+	 * @param bMatchAll If true, tags must be matched exactly. Otherwise, abilities matching any of the tags will be returned
+	 */
+	UFUNCTION(BlueprintCallable, BlueprintPure = false, Category = "Gameplay Abilities")
+	void FindAllAbilitiesWithTags(TArray<FGameplayAbilitySpecHandle>& OutAbilityHandles, FGameplayTagContainer Tags, bool bExactMatch = true) const;
+
+	/**
+	 * Returns an array with all abilities that match the provided Gameplay Tag Query
+	 *
+	 * @param OutAbilityHandles This array will be filled with matching Ability Spec Handles
+	 * @param Query Gameplay Tag Query to match
+	 */
+	UFUNCTION(BlueprintCallable, BlueprintPure = false, Category = "Gameplay Abilities")
+	void FindAllAbilitiesMatchingQuery(TArray<FGameplayAbilitySpecHandle>& OutAbilityHandles, FGameplayTagQuery Query) const;
+
+	/**
+	 * Returns an array with all abilities bound to an Input ID value
+	 *
+	 * @param OutAbilityHandles This array will be filled with matching Ability Spec Handles
+	 * @param InputID The Input ID to match
+	 */
+	UFUNCTION(BlueprintCallable, BlueprintPure = false, Category = "Gameplay Abilities")
+	void FindAllAbilitiesWithInputID(TArray<FGameplayAbilitySpecHandle>& OutAbilityHandles, int32 InputID = 0) const;
+
 	/** Retrieves the EffectContext of the GameplayEffect of the active GameplayEffect. */
 	FGameplayEffectContextHandle GetEffectContextFromActiveGEHandle(FActiveGameplayEffectHandle Handle);
 
@@ -966,7 +1098,7 @@ class GAMEPLAYABILITIES_API UAbilitySystemComponent : public UGameplayTasksCompo
 	/** Attempts to activate the given ability, will only work if called from the correct client/server context */
 	bool InternalTryActivateAbility(FGameplayAbilitySpecHandle AbilityToActivate, FPredictionKey InPredictionKey = FPredictionKey(), UGameplayAbility ** OutInstancedAbility = nullptr, FOnGameplayAbilityEnded::FDelegate* OnGameplayAbilityEndedDelegate = nullptr, const FGameplayEventData* TriggerEventData = nullptr);
 
-	// Failure tags used by InternalTryActivateAbility (E.g., this stores the  FailureTags of the last call to InternalTryActivateAbility
+	/** Failure tags used by InternalTryActivateAbility (E.g., this stores the  FailureTags of the last call to InternalTryActivateAbility */
 	FGameplayTagContainer InternalTryActivateAbilityFailureTags;
 
 	/** Called from the ability to let the component know it is ended */
@@ -1139,9 +1271,36 @@ class GAMEPLAYABILITIES_API UAbilitySystemComponent : public UGameplayTasksCompo
 	virtual void AbilityLocalInputPressed(int32 InputID);
 	virtual void AbilityLocalInputReleased(int32 InputID);
 
+	/*
+	 * Sends a local player Input Pressed event with the provided Input ID, notifying any bound abilities
+	 *
+	 * @param InputID The Input ID to match
+	 */
+	UFUNCTION(BlueprintCallable, Category = "Gameplay Abilities")
+	void PressInputID(int32 InputID);
+
+	/**
+	 * Sends a local player Input Released event with the provided Input ID, notifying any bound abilities
+	 * @param InputID The Input ID to match
+	 */
+	UFUNCTION(BlueprintCallable, Category = "Gameplay Abilities")
+	void ReleaseInputID(int32 InputID);
+
 	/** Handle confirm/cancel for target actors */
 	virtual void LocalInputConfirm();
 	virtual void LocalInputCancel();
+
+	/**
+	 * Sends a local player Input Confirm event, notifying abilities
+	 */
+	UFUNCTION(BlueprintCallable, Category = "Gameplay Abilities")
+	void InputConfirm();
+
+	/**
+	 * Sends a local player Input Cancel event, notifying abilities
+	 */
+	UFUNCTION(BlueprintCallable, Category = "Gameplay Abilities")
+	void InputCancel();
 	
 	/** InputID for binding GenericConfirm/Cancel events */
 	int32 GenericConfirmInputID;
@@ -1656,4 +1815,5 @@ protected:
 	};
 
 	TArray<FAbilityListLockActiveChange*> AbilityListLockActiveChanges;
+	
 };
\ No newline at end of file
diff --git a/Engine/Plugins/Runtime/GameplayAbilities/Source/GameplayAbilities/Public/GameplayModMagnitudeCalculation.h b/Engine/Plugins/Runtime/GameplayAbilities/Source/GameplayAbilities/Public/GameplayModMagnitudeCalculation.h
--- a/Engine/Plugins/Runtime/GameplayAbilities/Source/GameplayAbilities/Public/GameplayModMagnitudeCalculation.h
+++ b/Engine/Plugins/Runtime/GameplayAbilities/Source/GameplayAbilities/Public/GameplayModMagnitudeCalculation.h
@@ -63,4 +63,100 @@ protected:
 	 */
 	UPROPERTY(EditDefaultsOnly, Category=ExternalDependencies, AdvancedDisplay)
 	bool bAllowNonNetAuthorityDependencyRegistration;
+
+	/**
+	 * Gets the captured magnitude value for the given Attribute
+	 * For this to work correctly, the Attribute needs to be added to the Relevant Attributes to Capture array
+	 *
+	 * @param EffectSpec The Gameplay Effect Spec to get the info from
+	 * @param Attribute The attribute to query
+	 *
+	 * @return The magnitude value if found, zero otherwise
+	 */
+	UFUNCTION(BlueprintCallable, Category = "Ability|GameplayEffect", meta=(DisplayName="Get Captured Attribute Magnitude", ScriptName="GetCapturedAttributeMagnitude"))
+	float K2_GetCapturedAttributeMagnitude(const FGameplayEffectSpec& EffectSpec, FGameplayAttribute Attribute, const FGameplayTagContainer& SourceTags, const FGameplayTagContainer& TargetTags);
+
+	/**
+	 * Extracts the Set by Caller Magnitude from a Gameplay Effect Spec
+	 *
+	 * @param EffectSpec The Gameplay Effect Spec to get the info from
+	 * @param Tag The effect tag to query
+	 * 
+	 * @return The magnitude value if found, zero otherwise
+	 */
+	UFUNCTION(BlueprintCallable, Category = "Ability|GameplayEffect")
+	float GetSetByCallerMagnitudeByTag(const FGameplayEffectSpec& EffectSpec, const FGameplayTag& Tag);
+
+	/**
+	 * Extracts the Set by Caller Magnitude from a Gameplay Effect Spec
+	 *
+	 * @param EffectSpec The Gameplay Effect Spec to get the info from
+	 * @param MagnitudeName The effect name to query
+	 * 
+	 * @return The magnitude value if found, zero otherwise
+	 */
+	UFUNCTION(BlueprintCallable, Category = "Ability|GameplayEffect")
+	float GetSetByCallerMagnitudeByName(const FGameplayEffectSpec& EffectSpec, const FName& MagnitudeName);
+
+	/**
+	 * Copies and returns the source aggregated tags from a Gameplay Effect Spec
+	 *
+	 * @param EffectSpec The Gameplay Effect Spec to get the info from
+	 * 
+	 * @return Gameplay Tag Container with the copied tags. The container will be empty if no captured source tags exist.
+	 */
+	UFUNCTION(BlueprintCallable, Category = "Ability|GameplayEffect")
+	FGameplayTagContainer GetSourceAggregatedTags(const FGameplayEffectSpec& EffectSpec);
+
+	/**
+	 * Returns the source actor tags from a Gameplay Effect Spec
+	 *
+	 * @param EffectSpec The Gameplay Effect Spec to get the info from
+	 * 
+	 * @return Gameplay Tag Container with the copied tags. The container will be empty if no captured source tags exist.
+	 */
+	UFUNCTION(BlueprintCallable, Category = "Ability|GameplayEffect")
+	const FGameplayTagContainer& GetSourceActorTags(const FGameplayEffectSpec& EffectSpec);
+
+	/**
+	 * Returns the source spec tags from a Gameplay Effect Spec
+	 *
+	 * @param EffectSpec The Gameplay Effect Spec to get the info from
+	 * 
+	 * @return Gameplay Tag Container with the copied tags. The container will be empty if no captured source tags exist.
+	 */
+	UFUNCTION(BlueprintCallable, Category = "Ability|GameplayEffect")
+	const FGameplayTagContainer& GetSourceSpecTags(const FGameplayEffectSpec& EffectSpec);
+
+	/**
+	 * Copies and returns the target aggregated tags from a Gameplay Effect Spec
+	 *
+	 * @param EffectSpec The Gameplay Effect Spec to get the info from
+	 * 
+	 * @return Gameplay Tag Container with the copied tags. The container will be empty if no captured source tags exist.
+	 */
+	UFUNCTION(BlueprintCallable, Category = "Ability|GameplayEffect")
+	FGameplayTagContainer GetTargetAggregatedTags(const FGameplayEffectSpec& EffectSpec);
+
+	/**
+	 * Returns the target actor tags from a Gameplay Effect Spec
+	 * Useful for Modifier Magnitude Calculations
+	 *
+	 * @param EffectSpec The Gameplay Effect Spec to get the info from
+	 * 
+	 * @return Gameplay Tag Container with the copied tags. The container will be empty if no captured source tags exist.
+	 */
+	UFUNCTION(BlueprintCallable, Category = "Ability|GameplayEffect")
+	const FGameplayTagContainer& GetTargetActorTags(const FGameplayEffectSpec& EffectSpec);
+
+	/**
+	 * Returns the target spec tags from a Gameplay Effect Spec
+	 * Useful for Modifier Magnitude Calculations
+	 *
+	 * @param EffectSpec The Gameplay Effect Spec to get the info from
+	 * 
+	 * @return Gameplay Tag Container with the copied tags. The container will be empty if no captured source tags exist.
+	 */
+	UFUNCTION(BlueprintCallable, Category = "Ability|GameplayEffect")
+	const FGameplayTagContainer& GetTargetSpecTags(const FGameplayEffectSpec& EffectSpec);
 };
