Fixed issue where a looping gameplay cue notify is marked as having looping effects active after being constructed or Recycled, leading to unnecessary processing in FGameplayCueNotify_LoopingEffects, and hitting an ensure in some cases.

diff --git a/Engine/Plugins/Runtime/GameplayAbilities/Source/GameplayAbilities/Private/GameplayCueNotify_Looping.cpp b/Engine/Plugins/Runtime/GameplayAbilities/Source/GameplayAbilities/Private/GameplayCueNotify_Looping.cpp
--- a/Engine/Plugins/Runtime/GameplayAbilities/Source/GameplayAbilities/Private/GameplayCueNotify_Looping.cpp
+++ b/Engine/Plugins/Runtime/GameplayAbilities/Source/GameplayAbilities/Private/GameplayCueNotify_Looping.cpp
@@ -42,7 +42,7 @@ bool AGameplayCueNotify_Looping::Recycle()
 	RecurringSpawnResults.Reset();
 	RemovalSpawnResults.Reset();
 
-	bLoopingEffectsRemoved = false;
+	bLoopingEffectsRemoved = true;
 
 	return true;
 }
