Adding AbilitySystemComponent support for replicated loose gameplay tags (i.e. not backed up by a gameplay effect). Using this to allow ActivationOwnedTags from gameplay abilities to replicate. This is on by default. Also exposed the option for BP functions that add/remove loose tags.

diff --git a/Engine/Plugins/Runtime/GameplayAbilities/Source/GameplayAbilities/Private/Abilities/GameplayAbility.cpp b/Engine/Plugins/Runtime/GameplayAbilities/Source/GameplayAbilities/Private/Abilities/GameplayAbility.cpp
--- a/Engine/Plugins/Runtime/GameplayAbilities/Source/GameplayAbilities/Private/Abilities/GameplayAbility.cpp
+++ b/Engine/Plugins/Runtime/GameplayAbilities/Source/GameplayAbilities/Private/Abilities/GameplayAbility.cpp
@@ -666,6 +666,11 @@ void UGameplayAbility::EndAbility(const FGameplayAbilitySpecHandle Handle, const
 			// Remove tags we added to owner
 			AbilitySystemComponent->RemoveLooseGameplayTags(ActivationOwnedTags);
 
+			if (UAbilitySystemGlobals::Get().ShouldReplicateActivationOwnedTags())
+			{
+				AbilitySystemComponent->RemoveReplicatedLooseGameplayTags(ActivationOwnedTags);
+			}
+
 			// Remove tracked GameplayCues that we added
 			for (FGameplayTag& GameplayCueTag : TrackedGameplayCues)
 			{
@@ -765,8 +770,14 @@ void UGameplayAbility::PreActivate(const FGameplayAbilitySpecHandle Handle, cons
 	SetCurrentInfo(Handle, ActorInfo, ActivationInfo);
 
 	Comp->HandleChangeAbilityCanBeCanceled(AbilityTags, this, true);
+
 	Comp->AddLooseGameplayTags(ActivationOwnedTags);
 
+	if (UAbilitySystemGlobals::Get().ShouldReplicateActivationOwnedTags())
+	{
+		Comp->AddReplicatedLooseGameplayTags(ActivationOwnedTags);
+	}
+
 	if (OnGameplayAbilityEndedDelegate)
 	{
 		OnGameplayAbilityEnded.Add(*OnGameplayAbilityEndedDelegate);
diff --git a/Engine/Plugins/Runtime/GameplayAbilities/Source/GameplayAbilities/Private/AbilitySystemBlueprintLibrary.cpp b/Engine/Plugins/Runtime/GameplayAbilities/Source/GameplayAbilities/Private/AbilitySystemBlueprintLibrary.cpp
--- a/Engine/Plugins/Runtime/GameplayAbilities/Source/GameplayAbilities/Private/AbilitySystemBlueprintLibrary.cpp
+++ b/Engine/Plugins/Runtime/GameplayAbilities/Source/GameplayAbilities/Private/AbilitySystemBlueprintLibrary.cpp
@@ -1044,11 +1044,34 @@ FString UAbilitySystemBlueprintLibrary::GetActiveGameplayEffectDebugString(FActi
 	return Str;
 }
 
-bool UAbilitySystemBlueprintLibrary::AddLooseGameplayTags(AActor* Actor, const FGameplayTagContainer& GameplayTags)
+bool UAbilitySystemBlueprintLibrary::AddLooseGameplayTags(AActor* Actor, const FGameplayTagContainer& GameplayTags, bool bShouldReplicate)
 {
 	if (UAbilitySystemComponent* AbilitySysComp = GetAbilitySystemComponent(Actor))
 	{
 		AbilitySysComp->AddLooseGameplayTags(GameplayTags);
+
+		if (bShouldReplicate)
+		{
+			AbilitySysComp->AddReplicatedLooseGameplayTags(GameplayTags);
+		}
+
+		return true;
+	}
+
+	return false;
+}
+
+bool UAbilitySystemBlueprintLibrary::RemoveLooseGameplayTags(AActor* Actor, const FGameplayTagContainer& GameplayTags, bool bShouldReplicate)
+{
+	if (UAbilitySystemComponent* AbilitySysComp = GetAbilitySystemComponent(Actor))
+	{
+		AbilitySysComp->RemoveLooseGameplayTags(GameplayTags);
+
+		if (bShouldReplicate)
+		{
+			AbilitySysComp->RemoveReplicatedLooseGameplayTags(GameplayTags);
+		}
+
 		return true;
 	}
 
@@ -1068,17 +1091,6 @@ const UGameplayEffectUIData* UAbilitySystemBlueprintLibrary::GetGameplayEffectUI
 	return nullptr;
 }
 
-bool UAbilitySystemBlueprintLibrary::RemoveLooseGameplayTags(AActor* Actor, const FGameplayTagContainer& GameplayTags)
-{
-	if (UAbilitySystemComponent* AbilitySysComp = GetAbilitySystemComponent(Actor))
-	{
-		AbilitySysComp->RemoveLooseGameplayTags(GameplayTags);
-		return true;
-	}
-
-	return false;
-}
-
 bool UAbilitySystemBlueprintLibrary::EqualEqual_ActiveGameplayEffectHandle(const FActiveGameplayEffectHandle& A, const FActiveGameplayEffectHandle& B)
 {
 	return A == B;
diff --git a/Engine/Plugins/Runtime/GameplayAbilities/Source/GameplayAbilities/Private/AbilitySystemComponent.cpp b/Engine/Plugins/Runtime/GameplayAbilities/Source/GameplayAbilities/Private/AbilitySystemComponent.cpp
--- a/Engine/Plugins/Runtime/GameplayAbilities/Source/GameplayAbilities/Private/AbilitySystemComponent.cpp
+++ b/Engine/Plugins/Runtime/GameplayAbilities/Source/GameplayAbilities/Private/AbilitySystemComponent.cpp
@@ -210,6 +210,7 @@ PRAGMA_DISABLE_DEPRECATION_WARNINGS
 	// so we don't need to mark it dirty.
 	MinimalReplicationTags.Owner = this;
 PRAGMA_ENABLE_DEPRECATION_WARNINGS
+	ReplicatedLooseTags.Owner = this;
 
 	/** Allocate an AbilityActorInfo. Note: this goes through a global function and is a SharedPtr so projects can make their own AbilityActorInfo */
 	if(!AbilityActorInfo.IsValid())
@@ -1576,6 +1577,7 @@ void UAbilitySystemComponent::GetLifetimeReplicatedProps(TArray< FLifetimeProper
 	Params.Condition = COND_SkipOwner;
 	DOREPLIFETIME_WITH_PARAMS_FAST(UAbilitySystemComponent, MinimalReplicationGameplayCues, Params);
 	DOREPLIFETIME_WITH_PARAMS_FAST(UAbilitySystemComponent, MinimalReplicationTags, Params);
+	DOREPLIFETIME_WITH_PARAMS_FAST(UAbilitySystemComponent, ReplicatedLooseTags, Params);
 
 	Super::GetLifetimeReplicatedProps(OutLifetimeProps);
 }
@@ -2918,6 +2920,17 @@ const FMinimalReplicationTagCountMap& UAbilitySystemComponent::GetMinimalReplica
 	return MinimalReplicationTags;
 }
 
+FMinimalReplicationTagCountMap& UAbilitySystemComponent::GetReplicatedLooseTags_Mutable()
+{
+	MARK_PROPERTY_DIRTY_FROM_NAME(UAbilitySystemComponent, ReplicatedLooseTags, this);
+	return ReplicatedLooseTags;
+}
+
+const FMinimalReplicationTagCountMap& UAbilitySystemComponent::GetReplicatedLooseTags() const
+{
+	return ReplicatedLooseTags;
+}
+
 PRAGMA_ENABLE_DEPRECATION_WARNINGS
 
 #undef LOCTEXT_NAMESPACE
diff --git a/Engine/Plugins/Runtime/GameplayAbilities/Source/GameplayAbilities/Private/AbilitySystemGlobals.cpp b/Engine/Plugins/Runtime/GameplayAbilities/Source/GameplayAbilities/Private/AbilitySystemGlobals.cpp
--- a/Engine/Plugins/Runtime/GameplayAbilities/Source/GameplayAbilities/Private/AbilitySystemGlobals.cpp
+++ b/Engine/Plugins/Runtime/GameplayAbilities/Source/GameplayAbilities/Private/AbilitySystemGlobals.cpp
@@ -25,6 +25,8 @@ UAbilitySystemGlobals::UAbilitySystemGlobals(const FObjectInitializer& ObjectIni
 
 	PredictTargetGameplayEffects = true;
 
+	ReplicateActivationOwnedTags = true;
+
 	MinimalReplicationTagCountBits = 5;
 
 	bAllowGameplayModEvaluationChannels = false;
@@ -191,6 +193,11 @@ bool UAbilitySystemGlobals::ShouldPredictTargetGameplayEffects() const
 	return PredictTargetGameplayEffects;
 }
 
+bool UAbilitySystemGlobals::ShouldReplicateActivationOwnedTags() const
+{
+	return ReplicateActivationOwnedTags;
+}
+
 // --------------------------------------------------------------------
 
 UFunction* UAbilitySystemGlobals::GetGameplayCueFunction(const FGameplayTag& ChildTag, UClass* Class, FName &MatchedTag)
diff --git a/Engine/Plugins/Runtime/GameplayAbilities/Source/GameplayAbilities/Public/Abilities/GameplayAbility.h b/Engine/Plugins/Runtime/GameplayAbilities/Source/GameplayAbilities/Public/Abilities/GameplayAbility.h
--- a/Engine/Plugins/Runtime/GameplayAbilities/Source/GameplayAbilities/Public/Abilities/GameplayAbility.h
+++ b/Engine/Plugins/Runtime/GameplayAbilities/Source/GameplayAbilities/Public/Abilities/GameplayAbility.h
@@ -738,7 +738,7 @@ protected:
 	UPROPERTY(EditDefaultsOnly, Category = Tags, meta=(Categories="AbilityTagCategory"))
 	FGameplayTagContainer BlockAbilitiesWithTag;
 
-	/** Tags to apply to activating owner while this ability is active */
+	/** Tags to apply to activating owner while this ability is active. These are replicated if ReplicateActivationOwnedTags is enabled in AbilitySystemGlobals. */
 	UPROPERTY(EditDefaultsOnly, Category = Tags, meta=(Categories="OwnedTagsCategory"))
 	FGameplayTagContainer ActivationOwnedTags;
 
diff --git a/Engine/Plugins/Runtime/GameplayAbilities/Source/GameplayAbilities/Public/AbilitySystemBlueprintLibrary.h b/Engine/Plugins/Runtime/GameplayAbilities/Source/GameplayAbilities/Public/AbilitySystemBlueprintLibrary.h
--- a/Engine/Plugins/Runtime/GameplayAbilities/Source/GameplayAbilities/Public/AbilitySystemBlueprintLibrary.h
+++ b/Engine/Plugins/Runtime/GameplayAbilities/Source/GameplayAbilities/Public/AbilitySystemBlueprintLibrary.h
@@ -334,13 +334,13 @@ class GAMEPLAYABILITIES_API UAbilitySystemBlueprintLibrary : public UBlueprintFu
 	UFUNCTION(BlueprintPure, Category = "Ability|GameplayEffect")
 	static TArray<FGameplayEffectSpecHandle> GetAllLinkedGameplayEffectSpecHandles(FGameplayEffectSpecHandle SpecHandle);
 
-	/** manually adds a set of tags to given actor */
+	/** Manually adds a set of tags to a given actor, and optionally replicates them. */
 	UFUNCTION(BlueprintCallable, Category="Ability|GameplayEffect")
-	static bool AddLooseGameplayTags(AActor* Actor, const FGameplayTagContainer& GameplayTags);
+	static bool AddLooseGameplayTags(AActor* Actor, const FGameplayTagContainer& GameplayTags, bool bShouldReplicate=false);
 
-	/** manually removes a set of tags from a given actor */
+	/** Manually removes a set of tags from a given actor, with optional replication. */
 	UFUNCTION(BlueprintCallable, Category="Ability|GameplayEffect")
-	static bool RemoveLooseGameplayTags(AActor* Actor, const FGameplayTagContainer& GameplayTags);
+	static bool RemoveLooseGameplayTags(AActor* Actor, const FGameplayTagContainer& GameplayTags, bool bShouldReplicate=false);
 
 	// -------------------------------------------------------------------------------
 	//		GameplayEffectSpec
diff --git a/Engine/Plugins/Runtime/GameplayAbilities/Source/GameplayAbilities/Public/AbilitySystemComponent.h b/Engine/Plugins/Runtime/GameplayAbilities/Source/GameplayAbilities/Public/AbilitySystemComponent.h
--- a/Engine/Plugins/Runtime/GameplayAbilities/Source/GameplayAbilities/Public/AbilitySystemComponent.h
+++ b/Engine/Plugins/Runtime/GameplayAbilities/Source/GameplayAbilities/Public/AbilitySystemComponent.h
@@ -565,7 +565,7 @@ class GAMEPLAYABILITIES_API UAbilitySystemComponent : public UGameplayTasksCompo
 
 	/** 	 
 	 *  Allows GameCode to add loose gameplaytags which are not backed by a GameplayEffect. 
-	 *	Tags added this way are not replicated! 
+	 *	Tags added this way are not replicated! Use the 'Replicated' versions of these functions if replication is needed.
 	 *	It is up to the calling GameCode to make sure these tags are added on clients/server where necessary
 	 */
 	FORCEINLINE void AddLooseGameplayTag(const FGameplayTag& GameplayTag, int32 Count=1)
@@ -603,6 +603,41 @@ class GAMEPLAYABILITIES_API UAbilitySystemComponent : public UGameplayTasksCompo
 	UFUNCTION(BlueprintPure, Category = "Gameplay Tags")
 	int32 GetGameplayTagCount(FGameplayTag GameplayTag) const;
 
+	/**
+	 *  Allows GameCode to add loose gameplaytags which are not backed by a GameplayEffect. Tags added using 
+	 *  these functions will be replicated. Note that replicated loose tags will override any locally-set tag counts
+	 *  on simulated proxies.
+	 */
+	FORCEINLINE void AddReplicatedLooseGameplayTag(const FGameplayTag& GameplayTag)
+	{
+		GetReplicatedLooseTags_Mutable().AddTag(GameplayTag);
+		bIsNetDirty = true;
+	}
+
+	FORCEINLINE void AddReplicatedLooseGameplayTags(const FGameplayTagContainer& GameplayTags)
+	{
+		GetReplicatedLooseTags_Mutable().AddTags(GameplayTags);
+		bIsNetDirty = true;
+	}
+
+	FORCEINLINE void RemoveReplicatedLooseGameplayTag(const FGameplayTag& GameplayTag)
+	{
+		GetReplicatedLooseTags_Mutable().RemoveTag(GameplayTag);
+		bIsNetDirty = true;
+	}
+
+	FORCEINLINE void RemoveReplicatedLooseGameplayTags(const FGameplayTagContainer& GameplayTags)
+	{
+		GetReplicatedLooseTags_Mutable().RemoveTags(GameplayTags);
+		bIsNetDirty = true;
+	}
+
+	FORCEINLINE void SetReplicatedLooseGameplayTagCount(const FGameplayTag& GameplayTag, int32 NewCount)
+	{
+		GetReplicatedLooseTags_Mutable().SetTagCount(GameplayTag, NewCount);
+		bIsNetDirty = true;
+	}
+
 	/** 	 
 	 * Minimally replicated tags are replicated tags that come from GEs when in bMinimalReplication mode. 
 	 * (The GEs do not replicate, but the tags they grant do replicate via these functions)
@@ -1750,6 +1785,9 @@ protected:
 	FMinimalReplicationTagCountMap& GetMinimalReplicationTags_Mutable();
 	const FMinimalReplicationTagCountMap& GetMinimalReplicationTags() const;
 
+	FMinimalReplicationTagCountMap& GetReplicatedLooseTags_Mutable();
+	const FMinimalReplicationTagCountMap& GetReplicatedLooseTags() const;
+
 	void ResetTagMap();
 
 	void NotifyTagMap_StackCountChange(const FGameplayTagContainer& Container);
@@ -1779,6 +1817,9 @@ private:
 	FDelegateHandle MonitoredTagChangedDelegateHandle;
 	FTimerHandle    OnRep_ActivateAbilitiesTimerHandle;
 
+	/** Container used for replicating loose gameplay tags */
+	UPROPERTY(Replicated)
+	FMinimalReplicationTagCountMap ReplicatedLooseTags;
 
 	uint8 bDestroyActiveStateInitiated : 1;
 public:
diff --git a/Engine/Plugins/Runtime/GameplayAbilities/Source/GameplayAbilities/Public/AbilitySystemGlobals.h b/Engine/Plugins/Runtime/GameplayAbilities/Source/GameplayAbilities/Public/AbilitySystemGlobals.h
--- a/Engine/Plugins/Runtime/GameplayAbilities/Source/GameplayAbilities/Public/AbilitySystemGlobals.h
+++ b/Engine/Plugins/Runtime/GameplayAbilities/Source/GameplayAbilities/Public/AbilitySystemGlobals.h
@@ -88,6 +88,9 @@ class GAMEPLAYABILITIES_API UAbilitySystemGlobals : public UObject
 	/** Returns true if the ability system should try to predict gameplay effects applied to non local targets */
 	bool ShouldPredictTargetGameplayEffects() const;
 
+	/** Returns true if tags granted to owners from ability activations should be replicated */
+	bool ShouldReplicateActivationOwnedTags() const;
+
 	/** Searches the passed in class to look for a UFunction implementing the gameplay cue tag, sets MatchedTag to the exact tag found */
 	UFunction* GetGameplayCueFunction(const FGameplayTag &Tag, UClass* Class, FName &MatchedTag);
 
@@ -375,6 +378,10 @@ protected:
 	UPROPERTY(config)
 	bool PredictTargetGameplayEffects;
 
+	/** Set to true if you want tags granted to owners from ability activations to be replicated. If false, ActivationOwnedTags are only applied locally. */
+	UPROPERTY(config)
+	bool ReplicateActivationOwnedTags;
+
 	/** Manager for all gameplay cues */
 	UPROPERTY()
 	UGameplayCueManager* GlobalGameplayCueManager;
diff --git a/Engine/Plugins/Runtime/GameplayAbilities/Source/GameplayAbilities/Public/GameplayEffectTypes.h b/Engine/Plugins/Runtime/GameplayAbilities/Source/GameplayAbilities/Public/GameplayEffectTypes.h
--- a/Engine/Plugins/Runtime/GameplayAbilities/Source/GameplayAbilities/Public/GameplayEffectTypes.h
+++ b/Engine/Plugins/Runtime/GameplayAbilities/Source/GameplayAbilities/Public/GameplayEffectTypes.h
@@ -1671,6 +1671,19 @@ struct GAMEPLAYABILITIES_API FMinimalReplicationTagCountMap
 		}
 	}
 
+	void SetTagCount(const FGameplayTag& Tag, int32 NewCount)
+	{
+		MapID++;
+		int32& Count = TagMap.FindOrAdd(Tag);
+		Count = NewCount;
+		
+		if (Count == 0)
+		{
+			// Remove from map so that we do not replicate
+			TagMap.Remove(Tag);
+		}
+	}
+
 	bool NetSerialize(FArchive& Ar, class UPackageMap* Map, bool& bOutSuccess);
 
 	TMap<FGameplayTag, int32>	TagMap;
