Move safeguard back before cleaning up tasks to prevent EndAbility being called multiple times, causing AbilityTask counters to be off and causing server crashes later in runtime. Essentially revert cl/16854020, as [REDACTED_HANDLE] says should no longer be an issue.

#tests loaded stormshield, spawned coconut using "cheat giveweapon item_perk_h_goinbananas_fruit", ate coconut, verified that ability task count incremented and decremented properly

diff --git a/Engine/Plugins/Runtime/GameplayAbilities/Source/GameplayAbilities/Private/Abilities/GameplayAbility.cpp b/Engine/Plugins/Runtime/GameplayAbilities/Source/GameplayAbilities/Private/Abilities/GameplayAbility.cpp
--- a/Engine/Plugins/Runtime/GameplayAbilities/Source/GameplayAbilities/Private/Abilities/GameplayAbility.cpp
+++ b/Engine/Plugins/Runtime/GameplayAbilities/Source/GameplayAbilities/Private/Abilities/GameplayAbility.cpp
@@ -636,6 +636,7 @@ void UGameplayAbility::EndAbility(const FGameplayAbilitySpecHandle Handle, const
 
 		if (GetInstancingPolicy() != EGameplayAbilityInstancingPolicy::NonInstanced)
 		{
+			bIsActive = false;
 			bIsAbilityEnding = false;
 		}
 
@@ -650,12 +651,6 @@ void UGameplayAbility::EndAbility(const FGameplayAbilitySpecHandle Handle, const
 		}
 		ActiveTasks.Reset();	// Empty the array but don't resize memory, since this object is probably going to be destroyed very soon anyways.
 
-		// We mark the ability as no longer active after ending the tasks so they can broadcast delegates
-		if (GetInstancingPolicy() != EGameplayAbilityInstancingPolicy::NonInstanced)
-		{
-			bIsActive = false;
-		}
-
 		if (UAbilitySystemComponent* const AbilitySystemComponent = ActorInfo->AbilitySystemComponent.Get())
 		{
 			if (bReplicateEndAbility)
