Move the call to FlushNetDormancy earlier when removing active gameplay effects.
This prevents issues where changes to ActiveGameplayEffects aren't tracked properly on Dormant actors.

Fred.Kimberley Brian.Bekich

diff --git a/Engine/Plugins/Runtime/GameplayAbilities/Source/GameplayAbilities/Private/GameplayEffect.cpp b/Engine/Plugins/Runtime/GameplayAbilities/Source/GameplayAbilities/Private/GameplayEffect.cpp
--- a/Engine/Plugins/Runtime/GameplayAbilities/Source/GameplayAbilities/Private/GameplayEffect.cpp
+++ b/Engine/Plugins/Runtime/GameplayAbilities/Source/GameplayAbilities/Private/GameplayEffect.cpp
@@ -3378,9 +3378,18 @@ bool FActiveGameplayEffectsContainer::InternalRemoveActiveGameplayEffect(int32 I
 			return false;
 		}
 
+		// FlushNetDormancy needs to happen as early as possible.
+		// Any changes made to Fast Arrays (like the Gameplay Effect Containers) won't be
+		// tracked if made while the owner is still Dormant.
+		const bool bIsNetAuthority = IsNetAuthority();
+		AActor* OwnerActor = Owner->GetOwnerActor();
+		if (bIsNetAuthority && OwnerActor)
+		{
+			OwnerActor->FlushNetDormancy();
+		}
+		
 		// Invoke Remove GameplayCue event
 		bool ShouldInvokeGameplayCueEvent = true;
-		const bool bIsNetAuthority = IsNetAuthority();
 		if (!bIsNetAuthority && Effect.PredictionKey.IsLocalClientKey() && Effect.PredictionKey.WasReceived() == false)
 		{
 			// This was an effect that we predicted. Don't invoke GameplayCue event if we have another GameplayEffect that shares the same predictionkey and was received from the server
@@ -3405,12 +3414,6 @@ bool FActiveGameplayEffectsContainer::InternalRemoveActiveGameplayEffect(int32 I
 			Owner->GetWorld()->GetTimerManager().ClearTimer(Effect.PeriodHandle);
 		}
 
-		AActor* OwnerActor = Owner->GetOwnerActor();
-		if (bIsNetAuthority && OwnerActor)
-		{
-			OwnerActor->FlushNetDormancy();
-		}
-
 		// Remove this handle from the global map
 		Effect.Handle.RemoveFromGlobalMap();
 
