Ability System
- Add methods to dynamically alter set-by-caller magnitudes on already active gameplay effects (one for a single tag and one for multiple at once)
- Fix typo of "Get Activate" to "Get Active" on bp-exposed "Get Active Gameplay Effects for Query"

diff --git a/Engine/Plugins/Runtime/GameplayAbilities/Source/GameplayAbilities/Private/AbilitySystemComponent.cpp b/Engine/Plugins/Runtime/GameplayAbilities/Source/GameplayAbilities/Private/AbilitySystemComponent.cpp
--- a/Engine/Plugins/Runtime/GameplayAbilities/Source/GameplayAbilities/Private/AbilitySystemComponent.cpp
+++ b/Engine/Plugins/Runtime/GameplayAbilities/Source/GameplayAbilities/Private/AbilitySystemComponent.cpp
@@ -1057,6 +1057,16 @@ void UAbilitySystemComponent::GetGameplayEffectStartTimeAndDuration(FActiveGamep
 	return ActiveGameplayEffects.GetGameplayEffectStartTimeAndDuration(Handle, StartEffectTime, Duration);
 }
 
+void UAbilitySystemComponent::UpdateActiveGameplayEffectSetByCallerMagnitude(FActiveGameplayEffectHandle ActiveHandle, FGameplayTag SetByCallerTag, float NewValue)
+{
+	ActiveGameplayEffects.UpdateActiveGameplayEffectSetByCallerMagnitude(ActiveHandle, SetByCallerTag, NewValue);
+}
+
+void UAbilitySystemComponent::UpdateActiveGameplayEffectSetByCallerMagnitudes(FActiveGameplayEffectHandle ActiveHandle, const TMap<FGameplayTag, float>& NewSetByCallerValues)
+{
+	ActiveGameplayEffects.UpdateActiveGameplayEffectSetByCallerMagnitudes(ActiveHandle, NewSetByCallerValues);
+}
+
 float UAbilitySystemComponent::GetGameplayEffectMagnitude(FActiveGameplayEffectHandle Handle, FGameplayAttribute Attribute) const
 {
 	return ActiveGameplayEffects.GetGameplayEffectMagnitude(Handle, Attribute);
diff --git a/Engine/Plugins/Runtime/GameplayAbilities/Source/GameplayAbilities/Private/GameplayEffect.cpp b/Engine/Plugins/Runtime/GameplayAbilities/Source/GameplayAbilities/Private/GameplayEffect.cpp
--- a/Engine/Plugins/Runtime/GameplayAbilities/Source/GameplayAbilities/Private/GameplayEffect.cpp
+++ b/Engine/Plugins/Runtime/GameplayAbilities/Source/GameplayAbilities/Private/GameplayEffect.cpp
@@ -2545,6 +2545,33 @@ void FActiveGameplayEffectsContainer::SetActiveGameplayEffectLevel(FActiveGamepl
 	}
 }
 
+void FActiveGameplayEffectsContainer::UpdateActiveGameplayEffectSetByCallerMagnitude(FActiveGameplayEffectHandle ActiveHandle, const FGameplayTag& SetByCallerTag, float NewValue)
+{
+	if (FActiveGameplayEffect* Effect = GetActiveGameplayEffect(ActiveHandle))
+	{
+		Effect->Spec.SetSetByCallerMagnitude(SetByCallerTag, NewValue);
+		Effect->Spec.CalculateModifierMagnitudes();
+		MarkItemDirty(*Effect);
+
+		UpdateAllAggregatorModMagnitudes(*Effect);
+	}
+}
+
+void FActiveGameplayEffectsContainer::UpdateActiveGameplayEffectSetByCallerMagnitudes(FActiveGameplayEffectHandle ActiveHandle, const TMap<FGameplayTag, float>& NewSetByCallerValues)
+{
+	if (FActiveGameplayEffect* Effect = GetActiveGameplayEffect(ActiveHandle))
+	{
+		for (const TPair<FGameplayTag, float>& CurPair : NewSetByCallerValues)
+		{
+			Effect->Spec.SetSetByCallerMagnitude(CurPair.Key, CurPair.Value);
+		}
+		Effect->Spec.CalculateModifierMagnitudes();
+		MarkItemDirty(*Effect);
+
+		UpdateAllAggregatorModMagnitudes(*Effect);
+	}
+}
+
 const FGameplayTagContainer* FActiveGameplayEffectsContainer::GetGameplayEffectSourceTagsFromHandle(FActiveGameplayEffectHandle Handle) const
 {
 	// @todo: Need to consider this with tag changes
diff --git a/Engine/Plugins/Runtime/GameplayAbilities/Source/GameplayAbilities/Public/AbilitySystemComponent.h b/Engine/Plugins/Runtime/GameplayAbilities/Source/GameplayAbilities/Public/AbilitySystemComponent.h
--- a/Engine/Plugins/Runtime/GameplayAbilities/Source/GameplayAbilities/Public/AbilitySystemComponent.h
+++ b/Engine/Plugins/Runtime/GameplayAbilities/Source/GameplayAbilities/Public/AbilitySystemComponent.h
@@ -366,6 +366,14 @@ class GAMEPLAYABILITIES_API UAbilitySystemComponent : public UGameplayTasksCompo
 	/** Return start time and total duration of a gameplay effect */
 	void GetGameplayEffectStartTimeAndDuration(FActiveGameplayEffectHandle Handle, float& StartEffectTime, float& Duration) const;
 
+	/** Dynamically update the set-by-caller magnitude for an active gameplay effect */
+	UFUNCTION(BlueprintCallable, BlueprintAuthorityOnly, Category = GameplayEffects)
+	virtual void UpdateActiveGameplayEffectSetByCallerMagnitude(FActiveGameplayEffectHandle ActiveHandle, UPARAM(meta=(Categories = "SetByCaller"))FGameplayTag SetByCallerTag, float NewValue);
+
+	/** Dynamically update multiple set-by-caller magnitudes for an active gameplay effect */
+	UFUNCTION(BlueprintCallable, BlueprintAuthorityOnly, Category = GameplayEffects)
+	virtual void UpdateActiveGameplayEffectSetByCallerMagnitudes(FActiveGameplayEffectHandle ActiveHandle, const TMap<FGameplayTag, float>& NewSetByCallerValues);
+
 	/** Updates the level of an already applied gameplay effect. The intention is that this is 'seemless' and doesnt behave like removing/reapplying */
 	UFUNCTION(BlueprintCallable, BlueprintAuthorityOnly, Category = GameplayEffects)
 	virtual void SetActiveGameplayEffectLevel(FActiveGameplayEffectHandle ActiveHandle, int32 NewLevel);
@@ -762,7 +770,7 @@ class GAMEPLAYABILITIES_API UAbilitySystemComponent : public UGameplayTasksCompo
 	TArray<TPair<float,float>> GetActiveEffectsTimeRemainingAndDuration(const FGameplayEffectQuery& Query) const;
 
 	/** Returns list of active effects, for a query */
-	UFUNCTION(BlueprintCallable, BlueprintPure=false, Category = "GameplayEffects", meta=(DisplayName = "Get Activate Gameplay Effects for Query"))
+	UFUNCTION(BlueprintCallable, BlueprintPure=false, Category = "GameplayEffects", meta=(DisplayName = "Get Active Gameplay Effects for Query"))
 	TArray<FActiveGameplayEffectHandle> GetActiveEffects(const FGameplayEffectQuery& Query) const;
 
 	/** Returns list of active effects that have all of the passed in tags */
diff --git a/Engine/Plugins/Runtime/GameplayAbilities/Source/GameplayAbilities/Public/GameplayEffect.h b/Engine/Plugins/Runtime/GameplayAbilities/Source/GameplayAbilities/Public/GameplayEffect.h
--- a/Engine/Plugins/Runtime/GameplayAbilities/Source/GameplayAbilities/Public/GameplayEffect.h
+++ b/Engine/Plugins/Runtime/GameplayAbilities/Source/GameplayAbilities/Public/GameplayEffect.h
@@ -1587,6 +1587,23 @@ struct GAMEPLAYABILITIES_API FActiveGameplayEffectsContainer : public FFastArray
 
 	void SetActiveGameplayEffectLevel(FActiveGameplayEffectHandle ActiveHandle, int32 NewLevel);
 
+	/**
+	 * Update a set-by-caller magnitude for the active effect to match the new value, if possible
+	 * 
+	 * @param ActiveHandle		Handle of the active effect to update
+	 * @param SetByCallerTag	Set-by-caller tag to update
+	 * @param NewValue			New value of the set-by-caller magnitude
+	 */
+	void UpdateActiveGameplayEffectSetByCallerMagnitude(FActiveGameplayEffectHandle ActiveHandle, const FGameplayTag& SetByCallerTag, float NewValue);
+
+	/**
+	 * Update set-by-caller magnitudes for the active effect to match the new values, if possible; Replaces existing values
+	 *
+	 * @param ActiveHandle			Handle of the active effect to update
+	 * @param NewSetByCallerValues	Map of set-by-caller tag to new magnitude
+	 */
+	void UpdateActiveGameplayEffectSetByCallerMagnitudes(FActiveGameplayEffectHandle ActiveHandle, const TMap<FGameplayTag, float>& NewSetByCallerValues);
+
 	void SetAttributeBaseValue(FGameplayAttribute Attribute, float NewBaseValue);
 
 	float GetAttributeBaseValue(FGameplayAttribute Attribute) const;
