Added comment about the legacy game code support intent of an option. No functional changes.

diff --git a/Engine/Plugins/Runtime/GameplayAbilities/Source/GameplayAbilities/Public/AbilitySystemGlobals.h b/Engine/Plugins/Runtime/GameplayAbilities/Source/GameplayAbilities/Public/AbilitySystemGlobals.h
--- a/Engine/Plugins/Runtime/GameplayAbilities/Source/GameplayAbilities/Public/AbilitySystemGlobals.h
+++ b/Engine/Plugins/Runtime/GameplayAbilities/Source/GameplayAbilities/Public/AbilitySystemGlobals.h
@@ -378,7 +378,10 @@ protected:
 	UPROPERTY(config)
 	bool PredictTargetGameplayEffects;
 
-	/** Set to true if you want tags granted to owners from ability activations to be replicated. If false, ActivationOwnedTags are only applied locally. */
+	/** 
+	 * Set to true if you want tags granted to owners from ability activations to be replicated. If false, ActivationOwnedTags are only applied locally. 
+	 * This should only be disabled for legacy game code that depends on non-replication of ActivationOwnedTags.
+	 */
 	UPROPERTY(config)
 	bool ReplicateActivationOwnedTags;
 
