Audio - Initialize the Sound member of the GameplayCueNotify_SOundInfo

diff --git a/Engine/Plugins/Runtime/GameplayAbilities/Source/GameplayAbilities/Private/GameplayCueNotifyTypes.cpp b/Engine/Plugins/Runtime/GameplayAbilities/Source/GameplayAbilities/Private/GameplayCueNotifyTypes.cpp
--- a/Engine/Plugins/Runtime/GameplayAbilities/Source/GameplayAbilities/Private/GameplayCueNotifyTypes.cpp
+++ b/Engine/Plugins/Runtime/GameplayAbilities/Source/GameplayAbilities/Private/GameplayCueNotifyTypes.cpp
@@ -464,6 +464,7 @@ FGameplayCueNotify_SoundParameterInterfaceInfo::FGameplayCueNotify_SoundParamete
 
 FGameplayCueNotify_SoundInfo::FGameplayCueNotify_SoundInfo()
 {
+	Sound = nullptr;
 	SoundCue = nullptr;
 	LoopingFadeOutDuration = 0.5f;
 	LoopingFadeVolumeLevel = 0.0f;
