Checking in missed files in CL 18650552

diff --git a/Engine/Config/BaseEngine.ini b/Engine/Config/BaseEngine.ini
--- a/Engine/Config/BaseEngine.ini
+++ b/Engine/Config/BaseEngine.ini
@@ -978,6 +978,9 @@ NearClipPlane=10.0
 +FunctionRedirects=(OldName="SoundGeneratorParameterInterface.SetIntParameter", NewName="AudioParameterControllerInterface.SetIntParameter")
 +ClassRedirects=(OldName="/Script/AudioExtensions.AudioParameterInterface", NewName="/Script/AudioExtensions.AudioParameterControllerInterface")
 
+; Updating the gameplay  cue notify to have correct property name
++PropertyRedirects=(OldName="GameplayCueNotify_SoundInfo.SoundCue",NewName="GameplayCueNotify_SoundInfo.Sound")
+
 ; LiveLink data moving to runtime core module
 +ClassRedirects=(OldName="/Script/LiveLink.LiveLinkRetargetAsset", NewName="/Script/LiveLinkAnimationCore.LiveLinkRetargetAsset")
 +ClassRedirects=(OldName="/Script/LiveLink.LiveLinkRemapAsset", NewName="/Script/LiveLinkAnimationCore.LiveLinkRemapAsset")
diff --git a/Engine/Plugins/Runtime/GameplayAbilities/Source/GameplayAbilities/Private/GameplayCueNotifyTypes.cpp b/Engine/Plugins/Runtime/GameplayAbilities/Source/GameplayAbilities/Private/GameplayCueNotifyTypes.cpp
--- a/Engine/Plugins/Runtime/GameplayAbilities/Source/GameplayAbilities/Private/GameplayCueNotifyTypes.cpp
+++ b/Engine/Plugins/Runtime/GameplayAbilities/Source/GameplayAbilities/Private/GameplayCueNotifyTypes.cpp
@@ -478,7 +478,7 @@ bool FGameplayCueNotify_SoundInfo::PlaySound(const FGameplayCueNotify_SpawnConte
 	UAudioComponent* AudioComponent = nullptr;
 	bool bSoundPlayed = false;
 
-	if (SoundCue != nullptr)
+	if (Sound != nullptr)
 	{
 		const FGameplayCueNotify_SpawnCondition& SpawnCondition = SpawnContext.GetSpawnCondition(bOverrideSpawnCondition, SpawnConditionOverride);
 		const FGameplayCueNotify_PlacementInfo& PlacementInfo = SpawnContext.GetPlacementInfo(bOverridePlacementInfo, PlacementInfoOverride);
@@ -504,18 +504,18 @@ bool FGameplayCueNotify_SoundInfo::PlaySound(const FGameplayCueNotify_SpawnConte
 				{
 					const EAttachLocation::Type AttachLocationType = GetAttachLocationTypeFromRule(PlacementInfo.AttachmentRule);
 
-					AudioComponent = UGameplayStatics::SpawnSoundAttached(SoundCue, SpawnContext.TargetComponent, PlacementInfo.SocketName,
+					AudioComponent = UGameplayStatics::SpawnSoundAttached(Sound, SpawnContext.TargetComponent, PlacementInfo.SocketName,
 						SpawnLocation, SpawnRotation, AttachLocationType, bStopWhenAttachedToDestroyed,
 						VolumeMultiplier, PitchMultiplier, StartTime, AttenuationSettings, ConcurrencySettings, bAutoDestroy);
 				}
-				else if (SoundCue->IsLooping())
+				else if (Sound->IsLooping())
 				{
-					AudioComponent = UGameplayStatics::SpawnSoundAtLocation(SpawnContext.World, SoundCue, SpawnLocation, SpawnRotation,
+					AudioComponent = UGameplayStatics::SpawnSoundAtLocation(SpawnContext.World, Sound, SpawnLocation, SpawnRotation,
 						VolumeMultiplier, PitchMultiplier, StartTime, AttenuationSettings, ConcurrencySettings, bAutoDestroy);
 				}
 				else
 				{
-					UGameplayStatics::PlaySoundAtLocation(SpawnContext.World, SoundCue, SpawnLocation, SpawnRotation,
+					UGameplayStatics::PlaySoundAtLocation(SpawnContext.World, Sound, SpawnLocation, SpawnRotation,
 						VolumeMultiplier, PitchMultiplier, StartTime, AttenuationSettings, ConcurrencySettings, SpawnContext.TargetActor);
 				}
 
@@ -533,13 +533,13 @@ bool FGameplayCueNotify_SoundInfo::PlaySound(const FGameplayCueNotify_SpawnConte
 void FGameplayCueNotify_SoundInfo::ValidateBurstAssets(UObject* ContainingAsset, const FString& Context, TArray<FText>& ValidationErrors) const
 {
 #if WITH_EDITORONLY_DATA
-	if (SoundCue != nullptr)
+	if (Sound != nullptr)
 	{
-		if (SoundCue->IsLooping() && (!SoundCue->IsA<USoundWaveProcedural>()))
+		if (!Sound->IsOneShot())
 		{
 			ValidationErrors.Add(FText::Format(
-				LOCTEXT("SoundCue_ShouldNotLoop", "Sound [{0}] used in slot [{1}] for asset [{2}] is set to looping, but the slot is a one-shot (the instance will leak)."),
-				FText::AsCultureInvariant(SoundCue->GetPathName()),
+				LOCTEXT("SoundCue_ShouldNotLoop", "Sound [{0}] used in slot [{1}] for asset [{2}] is not a one-shot, but the slot is a one-shot (the instance will be orphaned)."),
+				FText::AsCultureInvariant(Sound->GetPathName()),
 				FText::AsCultureInvariant(Context),
 				FText::AsCultureInvariant(ContainingAsset->GetPathName())));
 		}
@@ -547,7 +547,6 @@ void FGameplayCueNotify_SoundInfo::ValidateBurstAssets(UObject* ContainingAsset,
 #endif //#if WITH_EDITORONLY_DATA
 }
 
-
 //////////////////////////////////////////////////////////////////////////
 // FGameplayCueNotify_CameraShakeInfo
 //////////////////////////////////////////////////////////////////////////
