Allow game feature plugins to modify the gameplay cue paths

diff --git a/Engine/Plugins/Runtime/GameplayAbilities/Source/GameplayAbilities/Private/AbilitySystemGlobals.cpp b/Engine/Plugins/Runtime/GameplayAbilities/Source/GameplayAbilities/Private/AbilitySystemGlobals.cpp
--- a/Engine/Plugins/Runtime/GameplayAbilities/Source/GameplayAbilities/Private/AbilitySystemGlobals.cpp
+++ b/Engine/Plugins/Runtime/GameplayAbilities/Source/GameplayAbilities/Private/AbilitySystemGlobals.cpp
@@ -139,6 +139,16 @@ const FName& UAbilitySystemGlobals::GetGameplayModEvaluationChannelAlias(int32 I
 	return GameplayModEvaluationChannelAliases[Index];
 }
 
+void UAbilitySystemGlobals::AddGameplayCueNotifyPath(const FString& InPath)
+{
+	GameplayCueNotifyPaths.AddUnique(InPath);
+}
+
+int32 UAbilitySystemGlobals::RemoveGameplayCueNotifyPath(const FString& InPath)
+{
+	return GameplayCueNotifyPaths.Remove(InPath);
+}
+
 #if WITH_EDITOR
 
 void UAbilitySystemGlobals::OnTableReimported(UObject* InObject)
diff --git a/Engine/Plugins/Runtime/GameplayAbilities/Source/GameplayAbilities/Private/GameplayCueManager.cpp b/Engine/Plugins/Runtime/GameplayAbilities/Source/GameplayAbilities/Private/GameplayCueManager.cpp
--- a/Engine/Plugins/Runtime/GameplayAbilities/Source/GameplayAbilities/Private/GameplayCueManager.cpp
+++ b/Engine/Plugins/Runtime/GameplayAbilities/Source/GameplayAbilities/Private/GameplayCueManager.cpp
@@ -698,6 +698,32 @@ TArray<FString> UGameplayCueManager::GetAlwaysLoadedGameplayCuePaths()
 	return UAbilitySystemGlobals::Get().GetGameplayCueNotifyPaths();
 }
 
+void UGameplayCueManager::AddGameplayCueNotifyPath(const FString& InPath, const bool bShouldRescanCueAssets /* = true */)
+{
+	UAbilitySystemGlobals::Get().AddGameplayCueNotifyPath(InPath);
+	const int32 NumAdded = RuntimeGameplayCueObjectLibrary.Paths.AddUnique(InPath);
+	
+	if(bShouldRescanCueAssets && NumAdded != INDEX_NONE)
+	{
+		InitializeRuntimeObjectLibrary();		
+	}
+}
+
+int32 UGameplayCueManager::RemoveGameplayCueNotifyPath(const FString& InPath, const bool bShouldRescanCueAssets /* = true */)
+{
+	int32 NumRemovedGlobal = UAbilitySystemGlobals::Get().RemoveGameplayCueNotifyPath(InPath);
+	int32 NumRemovedRuntime = RuntimeGameplayCueObjectLibrary.Paths.Remove(InPath);
+
+	ensureMsgf(NumRemovedGlobal == NumRemovedRuntime, TEXT("Unexpected number of cue paths removed for '%s'"), *InPath);
+	
+	if(bShouldRescanCueAssets && NumRemovedGlobal > 0)
+	{
+		InitializeRuntimeObjectLibrary();		
+	}
+	
+	return NumRemovedRuntime;
+}
+
 void UGameplayCueManager::RefreshObjectLibraries()
 {
 	if (RuntimeGameplayCueObjectLibrary.bHasBeenInitialized)
diff --git a/Engine/Plugins/Runtime/GameplayAbilities/Source/GameplayAbilities/Public/AbilitySystemGlobals.h b/Engine/Plugins/Runtime/GameplayAbilities/Source/GameplayAbilities/Public/AbilitySystemGlobals.h
--- a/Engine/Plugins/Runtime/GameplayAbilities/Source/GameplayAbilities/Public/AbilitySystemGlobals.h
+++ b/Engine/Plugins/Runtime/GameplayAbilities/Source/GameplayAbilities/Public/AbilitySystemGlobals.h
@@ -295,6 +295,16 @@ class GAMEPLAYABILITIES_API UAbilitySystemGlobals : public UObject
 	/** Path where the engine will load gameplay cue notifies from */
 	virtual TArray<FString> GetGameplayCueNotifyPaths() { return GameplayCueNotifyPaths; }
 
+	/** Add a path to the GameplayCueNotifyPaths array. */
+	virtual void AddGameplayCueNotifyPath(const FString& InPath);
+
+	/**
+	 * Remove the given gameplay cue notify path from the GameplayCueNotifyPaths array.
+	 *
+	 * @return Number of paths removed.
+	 */
+	virtual int32 RemoveGameplayCueNotifyPath(const FString& InPath);
+
 	UPROPERTY()
 	FNetSerializeScriptStructCache	TargetDataStructCache;
 
diff --git a/Engine/Plugins/Runtime/GameplayAbilities/Source/GameplayAbilities/Public/GameplayCueManager.h b/Engine/Plugins/Runtime/GameplayAbilities/Source/GameplayAbilities/Public/GameplayCueManager.h
--- a/Engine/Plugins/Runtime/GameplayAbilities/Source/GameplayAbilities/Public/GameplayCueManager.h
+++ b/Engine/Plugins/Runtime/GameplayAbilities/Source/GameplayAbilities/Public/GameplayCueManager.h
@@ -279,6 +279,25 @@ public:
 
 	/** Called before loading any gameplay cue notifies from object libraries. Allows subclasses to skip notifies. */
 	virtual bool ShouldLoadGameplayCueAssetData(const FAssetData& Data) const { return true; }
+
+	/**
+	 * Add a path to the GameplayCueNotifyPaths array.
+	 * Re-initalizes RuntimeObjectLibrary based on the new paths
+	 *
+	 * @param InPath					The path to the directory that should be added to the scan
+	 * @param bShouldRescanCueAssets		If true then the runtime object library will be rebuilt.
+	 */
+	virtual void AddGameplayCueNotifyPath(const FString& InPath, const bool bShouldRescanCueAssets = true);
+
+	/**
+	 * Remove the given gameplay cue notify path from the GameplayCueNotifyPaths array.
+	 * Re-initalizes RuntimeObjectLibrary based on the new paths
+	 * 
+	 * @param InPath					The path to the directory that should be removed from the scan
+	 * @param bShouldRescanCueAssets		If true then the runtime object library will be rebuilt.
+	 * @return Number of paths removed.
+	 */
+	virtual int32 RemoveGameplayCueNotifyPath(const FString& InPath, const bool bShouldRescanCueAssets = true);
 	
 	int32 FinishLoadingGameplayCueNotifies();
 
