AbiltiyTask_WaitDelay:
* Remove unchecked call to GetWorld() in GetDebugString()

diff --git a/Engine/Plugins/Runtime/GameplayAbilities/Source/GameplayAbilities/Private/Abilities/Tasks/AbilityTask_WaitDelay.cpp b/Engine/Plugins/Runtime/GameplayAbilities/Source/GameplayAbilities/Private/Abilities/Tasks/AbilityTask_WaitDelay.cpp
--- a/Engine/Plugins/Runtime/GameplayAbilities/Source/GameplayAbilities/Private/Abilities/Tasks/AbilityTask_WaitDelay.cpp
+++ b/Engine/Plugins/Runtime/GameplayAbilities/Source/GameplayAbilities/Private/Abilities/Tasks/AbilityTask_WaitDelay.cpp
@@ -41,6 +41,13 @@ void UAbilityTask_WaitDelay::OnTimeFinish()
 
 FString UAbilityTask_WaitDelay::GetDebugString() const
 {
-	float TimeLeft = Time - GetWorld()->TimeSince(TimeStarted);
-	return FString::Printf(TEXT("WaitDelay. Time: %.2f. TimeLeft: %.2f"), Time, TimeLeft);
+	if (UWorld* World = GetWorld())
+	{
+		const float TimeLeft = Time - World->TimeSince(TimeStarted);
+		return FString::Printf(TEXT("WaitDelay. Time: %.2f. TimeLeft: %.2f"), Time, TimeLeft);
+	}
+	else
+	{
+		return FString::Printf(TEXT("WaitDelay. Time: %.2f. Time Started: %.2f"), Time, TimeStarted);
+	}
 }
\ No newline at end of file
