AbilitySystem debug improvements:
* Change the global hardcoded 1000 AbilityTask limit into a CVar defaulting to 1000
* Add 'AbilitySystem.AbilityTask.Debug.RecordingEnabled' cvar to enable per-class counting of created AbilityTasks
* Add 'AbilitySystem.AbilityTask.Debug.PrintCounts' console command to print out the current per-class AbilityTask counts

diff --git a/Engine/Plugins/Runtime/GameplayAbilities/Source/GameplayAbilities/Private/Abilities/Tasks/AbilityTask.cpp b/Engine/Plugins/Runtime/GameplayAbilities/Source/GameplayAbilities/Private/Abilities/Tasks/AbilityTask.cpp
--- a/Engine/Plugins/Runtime/GameplayAbilities/Source/GameplayAbilities/Private/Abilities/Tasks/AbilityTask.cpp
+++ b/Engine/Plugins/Runtime/GameplayAbilities/Source/GameplayAbilities/Private/Abilities/Tasks/AbilityTask.cpp
@@ -4,37 +4,95 @@
 #include "AbilitySystemComponent.h"
 #include "AbilitySystemStats.h"
 
-int32 UAbilityTask::GlobalAbilityTaskCount = 0;
+#if !UE_BUILD_SHIPPING
+static void DebugRecordAbilityTaskCreated(const UAbilityTask* NewTask);
+static void DebugRecordAbilityTaskDestroyed(const UAbilityTask* NewTask);
+static void DebugPrintAbilityTasksByClass();
+#endif // !UE_BUILD_SHIPPING
+
+namespace AbilityTaskCVars
+{
+	static int32 AbilityTaskMaxCount = 1000;
+	static FAutoConsoleVariableRef CVarMaxAbilityTaskCount(
+		TEXT("AbilitySystem.AbilityTask.MaxCount"),
+		AbilityTaskMaxCount,
+		TEXT("Global limit on the number of extant AbilityTasks. Use 'AbilitySystem.Debug.RecordingEnabled' and 'AbilitySystem.AbilityTask.Debug.PrintCounts' to debug why you are hitting this before raising the cap.")
+	);
+
+#if !UE_BUILD_SHIPPING
+	static bool bRecordAbilityTaskCounts = false;
+	static FAutoConsoleVariableRef CVarRecordAbilityTaskCounts(
+		TEXT("AbilitySystem.AbilityTask.Debug.RecordingEnabled"),
+		bRecordAbilityTaskCounts,
+		TEXT("If this is enabled, all new AbilityTasks will be counted by type. Use 'AbilitySystem.AbilityTask.Debug.PrintCounts' to print out the current counts.")
+	);
+
+	static FAutoConsoleCommand AbilityTaskPrintAbilityTaskCountsCmd(
+		TEXT("AbilitySystem.AbilityTask.Debug.PrintCounts"),
+		TEXT("Print out the current AbilityTask counts by class. 'AbilitySystem.AbilityTask.Debug.RecordingEnabled' must be turned on for this to function."),
+		FConsoleCommandDelegate::CreateStatic(DebugPrintAbilityTasksByClass)
+	);
+#endif
+}
+
+static int32 GlobalAbilityTaskCount = 0;
 
 UAbilityTask::UAbilityTask(const FObjectInitializer& ObjectInitializer)
 	: Super(ObjectInitializer)
 {
 	WaitStateBitMask = static_cast<uint8>(EAbilityTaskWaitState::WaitingOnGame);
-	bWasSuccessfullyDestroyed = false;
+
+#if !UE_BUILD_SHIPPING
+	if (AbilityTaskCVars::bRecordAbilityTaskCounts)
+	{
+		DebugRecordAbilityTaskCreated(this);
+	}
+#endif  // !UE_BUILD_SHIPPING
 
 	++GlobalAbilityTaskCount;
 	SET_DWORD_STAT(STAT_AbilitySystem_TaskCount, GlobalAbilityTaskCount);
-	if (!ensure(GlobalAbilityTaskCount < 1000))
+	if (!ensure(GlobalAbilityTaskCount < AbilityTaskCVars::AbilityTaskMaxCount))
 	{
 		ABILITY_LOG(Warning, TEXT("Way too many AbilityTasks are currently active! %d. %s"), GlobalAbilityTaskCount, *GetClass()->GetName());
+
+#if !UE_BUILD_SHIPPING
+		// Auto dump the counts if we hit the limit
+		if (AbilityTaskCVars::bRecordAbilityTaskCounts)
+		{
+			static bool bHasDumpedAbilityTasks = false;  // The dump is spammy, so we only want to auto-dump once
+
+			if (!bHasDumpedAbilityTasks)
+			{
+				DebugPrintAbilityTasksByClass();
+				bHasDumpedAbilityTasks = true;
+			}
+		}
+#endif  // !UE_BUILD_SHIPPING
 	}
 }
 
 void UAbilityTask::OnDestroy(bool bInOwnerFinished)
 {
 	checkf(GlobalAbilityTaskCount > 0, TEXT("Mismatched AbilityTask counting"));
-    --GlobalAbilityTaskCount;
+	--GlobalAbilityTaskCount;
 	SET_DWORD_STAT(STAT_AbilitySystem_TaskCount, GlobalAbilityTaskCount);
 
 	bWasSuccessfullyDestroyed = true;
 
+#if !UE_BUILD_SHIPPING
+	if (AbilityTaskCVars::bRecordAbilityTaskCounts)
+	{
+		DebugRecordAbilityTaskDestroyed(this);
+	}
+#endif  // !UE_BUILD_SHIPPING
+
 	Super::OnDestroy(bInOwnerFinished);
 }
 
 void UAbilityTask::BeginDestroy()
 {
 	Super::BeginDestroy();
-	
+
 	if (!bWasSuccessfullyDestroyed)
 	{
 		// this shouldn't happen, it means that ability was destroyed while being active, but we need to keep GlobalAbilityTaskCount in sync anyway
@@ -42,6 +100,13 @@ void UAbilityTask::BeginDestroy()
 		--GlobalAbilityTaskCount;
 		SET_DWORD_STAT(STAT_AbilitySystem_TaskCount, GlobalAbilityTaskCount);
 		bWasSuccessfullyDestroyed = true;
+
+#if !UE_BUILD_SHIPPING
+		if (AbilityTaskCVars::bRecordAbilityTaskCounts)
+		{
+			DebugRecordAbilityTaskDestroyed(this);
+		}
+#endif  // !UE_BUILD_SHIPPING
 	}
 }
 
@@ -144,3 +209,71 @@ bool UAbilityTask::IsWaitingOnAvatar() const
 {
 	return (WaitStateBitMask & (uint8)EAbilityTaskWaitState::WaitingOnAvatar) != 0;
 }
+
+#if !UE_BUILD_SHIPPING
+static TMap<const UClass*, int32> StaticAbilityTasksByClass = {};
+
+static void DebugRecordAbilityTaskCreated(const UAbilityTask* NewTask)
+{
+	const UClass* ClassPtr = (NewTask != nullptr) ? NewTask->GetClass() : nullptr;
+	if (ClassPtr != nullptr)
+	{
+		if (StaticAbilityTasksByClass.Contains(ClassPtr))
+		{
+			StaticAbilityTasksByClass[ClassPtr]++;
+		}
+		else
+		{
+			StaticAbilityTasksByClass.Add(ClassPtr, 1);
+		}
+	}
+
+}
+
+static void DebugRecordAbilityTaskDestroyed(const UAbilityTask* DestroyedTask)
+{
+	const UClass* ClassPtr = (DestroyedTask != nullptr) ? DestroyedTask->GetClass() : nullptr;
+	if (ClassPtr != nullptr)
+	{
+		if (StaticAbilityTasksByClass.Contains(ClassPtr))
+		{
+			StaticAbilityTasksByClass[ClassPtr]--;
+
+			if (StaticAbilityTasksByClass[ClassPtr] <= 0)
+			{
+				StaticAbilityTasksByClass.Remove(ClassPtr);
+			}
+		}
+	}
+}
+
+static void DebugPrintAbilityTasksByClass()
+{
+	if (AbilityTaskCVars::bRecordAbilityTaskCounts)
+	{
+		int32 AccumulatedAbilityTasks = 0;
+		ABILITY_LOG(Display, TEXT("Logging global UAbilityTask counts:"));
+		StaticAbilityTasksByClass.ValueSort(TGreater<int32>());
+		for (const TPair<const UClass*, int32>& Pair : StaticAbilityTasksByClass)
+		{
+			FString SafeName = GetNameSafe(Pair.Key);
+			ABILITY_LOG(Display, TEXT("- Class '%s': %d"), *SafeName, Pair.Value);
+			AccumulatedAbilityTasks += Pair.Value;
+		}
+
+		const int32 UnaccountedAbilityTasks = GlobalAbilityTaskCount - AccumulatedAbilityTasks;
+		if (UnaccountedAbilityTasks > 0)
+		{
+			// It's possible to allocate AbilityTasks before AbilityTaskCVars::bRecordAbilityTaskCounts was set to 'true', even if set via command line.
+			// However, if this value increases during play, there is an issue.
+			ABILITY_LOG(Display, TEXT("- Unknown (allocated before recording): %d"), UnaccountedAbilityTasks);
+		}
+
+		ABILITY_LOG(Display, TEXT("Total AbilityTask count: %d"), GlobalAbilityTaskCount);
+	}
+	else
+	{
+		ABILITY_LOG(Display, TEXT("Recording of UAbilityTask counts is disabled! Enable 'AbilitySystem.AbilityTask.Debug.RecordingEnabled' to turn on recording."))
+	}
+}
+#endif  // !UE_BUILD_SHIPPING
diff --git a/Engine/Plugins/Runtime/GameplayAbilities/Source/GameplayAbilities/Public/Abilities/Tasks/AbilityTask.h b/Engine/Plugins/Runtime/GameplayAbilities/Source/GameplayAbilities/Public/Abilities/Tasks/AbilityTask.h
--- a/Engine/Plugins/Runtime/GameplayAbilities/Source/GameplayAbilities/Public/Abilities/Tasks/AbilityTask.h
+++ b/Engine/Plugins/Runtime/GameplayAbilities/Source/GameplayAbilities/Public/Abilities/Tasks/AbilityTask.h
@@ -163,10 +163,6 @@ class GAMEPLAYABILITIES_API UAbilityTask : public UGameplayTask
 protected:
 	/** Helper method for registering client replicated callbacks */
 	bool CallOrAddReplicatedDelegate(EAbilityGenericReplicatedEvent::Type Event, FSimpleMulticastDelegate::FDelegate Delegate);
-
-private:
-
-	static int32 GlobalAbilityTaskCount;
 };
 
 //For searching through lists of ability instances
