PR #8854: FScalableFloat now returns a reference to self on assignment. (Contributed by [REDACTED_PERSON])

diff --git a/Engine/Plugins/Runtime/GameplayAbilities/Source/GameplayAbilities/Private/ScalableFloat.cpp b/Engine/Plugins/Runtime/GameplayAbilities/Source/GameplayAbilities/Private/ScalableFloat.cpp
--- a/Engine/Plugins/Runtime/GameplayAbilities/Source/GameplayAbilities/Private/ScalableFloat.cpp
+++ b/Engine/Plugins/Runtime/GameplayAbilities/Source/GameplayAbilities/Private/ScalableFloat.cpp
@@ -198,13 +198,15 @@ bool FScalableFloat::operator!=(const FScalableFloat& Other) const
 	return ((Other.Curve != Curve) || (Other.RegistryType != RegistryType) || (Other.Value != Value));
 }
 
-void FScalableFloat::operator=(const FScalableFloat& Src)
+FScalableFloat& FScalableFloat::operator=(const FScalableFloat& Src)
 {
 	Value = Src.Value;
 	Curve = Src.Curve;
 	RegistryType = Src.RegistryType;
 	LocalCachedCurveID = Src.LocalCachedCurveID;
 	CachedCurve = Src.CachedCurve;
+
+	return *this;
 }
 
 
diff --git a/Engine/Plugins/Runtime/GameplayAbilities/Source/GameplayAbilities/Public/ScalableFloat.h b/Engine/Plugins/Runtime/GameplayAbilities/Source/GameplayAbilities/Public/ScalableFloat.h
--- a/Engine/Plugins/Runtime/GameplayAbilities/Source/GameplayAbilities/Public/ScalableFloat.h
+++ b/Engine/Plugins/Runtime/GameplayAbilities/Source/GameplayAbilities/Public/ScalableFloat.h
@@ -94,7 +94,7 @@ public:
 	bool operator!=(const FScalableFloat& Other) const;
 
 	/** copy operator to prevent duplicate handles */
-	void operator=(const FScalableFloat& Src);
+	FScalableFloat& operator=(const FScalableFloat& Src);
 
 	/** Used to upgrade a float or int8/int16/int32 property into an FScalableFloat */
 	bool SerializeFromMismatchedTag(const FPropertyTag& Tag, FStructuredArchive::FSlot Slot);
