Removed some log spam about updating the gameplay tag tree

diff --git a/Engine/Plugins/Runtime/GameplayAbilities/Source/GameplayAbilitiesEditor/Private/GameplayAbilitiesEditorModule.cpp b/Engine/Plugins/Runtime/GameplayAbilities/Source/GameplayAbilitiesEditor/Private/GameplayAbilitiesEditorModule.cpp
--- a/Engine/Plugins/Runtime/GameplayAbilities/Source/GameplayAbilitiesEditor/Private/GameplayAbilitiesEditorModule.cpp
+++ b/Engine/Plugins/Runtime/GameplayAbilities/Source/GameplayAbilitiesEditor/Private/GameplayAbilitiesEditorModule.cpp
@@ -256,11 +256,6 @@ void FGameplayAbilitiesEditorModule::RegisterAssetTypeAction(IAssetTools& AssetT
 
 void FGameplayAbilitiesEditorModule::GameplayTagTreeChanged()
 {
-	// The tag tree changed so we should refresh which actions are provided by the gameplay cue event
-#if STATS
-	FString PerfMessage = FString::Printf(TEXT("FGameplayAbilitiesEditorModule::GameplayTagTreeChanged"));
-	SCOPE_LOG_TIME_IN_SECONDS(*PerfMessage, nullptr)
-#endif
 	if (FBlueprintActionDatabase* BAD = FBlueprintActionDatabase::TryGet())
 	{
 		BAD->RefreshClassActions(UK2Node_GameplayCueEvent::StaticClass());
