Fixed two comment typos.

#UE5 #RNX #NoReleaseNotes

diff --git a/Engine/Plugins/Runtime/GameplayAbilities/Source/GameplayAbilities/Private/Abilities/GameplayAbility.cpp b/Engine/Plugins/Runtime/GameplayAbilities/Source/GameplayAbilities/Private/Abilities/GameplayAbility.cpp
--- a/Engine/Plugins/Runtime/GameplayAbilities/Source/GameplayAbilities/Private/Abilities/GameplayAbility.cpp
+++ b/Engine/Plugins/Runtime/GameplayAbilities/Source/GameplayAbilities/Private/Abilities/GameplayAbility.cpp
@@ -304,7 +304,7 @@ bool UGameplayAbility::CanActivateAbility(const FGameplayAbilitySpecHandle Handl
 
 	FGameplayTagContainer& OutTags = OptionalRelevantTags ? *OptionalRelevantTags : DummyContainer;
 
-	// make sure theability system component is valid, if not bail out.
+	// make sure the ability system component is valid, if not bail out.
 	UAbilitySystemComponent* const AbilitySystemComponent = ActorInfo->AbilitySystemComponent.Get();
 	if (!AbilitySystemComponent)
 	{
@@ -315,7 +315,7 @@ bool UGameplayAbility::CanActivateAbility(const FGameplayAbilitySpecHandle Handl
 	{
 		/**
 		 *	Input is inhibited (UI is pulled up, another ability may be blocking all other input, etc).
-		 *	When we get into triggered abilities, we may need to better differentiate between CanActviate and CanUserActivate or something.
+		 *	When we get into triggered abilities, we may need to better differentiate between CanActivate and CanUserActivate or something.
 		 *	E.g., we would want LMB/RMB to be inhibited while the user is in the menu UI, but we wouldn't want to prevent a 'buff when I am low health'
 		 *	ability to not trigger.
 		 *	
