Change default directory specification for "Add New" button when creating gameplay cues

diff --git a/Engine/Plugins/Runtime/GameplayAbilities/Source/GameplayAbilitiesEditor/Private/SGameplayCueEditor.cpp b/Engine/Plugins/Runtime/GameplayAbilities/Source/GameplayAbilitiesEditor/Private/SGameplayCueEditor.cpp
--- a/Engine/Plugins/Runtime/GameplayAbilities/Source/GameplayAbilitiesEditor/Private/SGameplayCueEditor.cpp
+++ b/Engine/Plugins/Runtime/GameplayAbilities/Source/GameplayAbilitiesEditor/Private/SGameplayCueEditor.cpp
@@ -1473,7 +1473,7 @@ FString SGameplayCueEditor::GetPathNameForGameplayCueTag(FString GameplayCueTagN
 	else
 	{
 		GameplayCueTagName = GameplayCueTagName.Replace(TEXT("GameplayCue."), TEXT(""), ESearchCase::IgnoreCase);
-		NewDefaultPathName = FString::Printf(TEXT("/Game/GC_%s"), *GameplayCueTagName);
+		NewDefaultPathName = FString::Printf(TEXT("/Game/GameplayCueNotifies/GC_%s"), *GameplayCueTagName);
 	}
 	NewDefaultPathName.ReplaceInline(TEXT("."), TEXT("_"));
 	return NewDefaultPathName;
