Fixed an issue where the apply root motion ability task could restore a movement mode without also restoring the custom movement mode.  This would result the player getting into MOVE_Custom where custom mode is 0, which forces the player to float in here.

diff --git a/Engine/Plugins/Runtime/GameplayAbilities/Source/GameplayAbilities/Private/Abilities/Tasks/AbilityTask_ApplyRootMotionMoveToForce.cpp b/Engine/Plugins/Runtime/GameplayAbilities/Source/GameplayAbilities/Private/Abilities/Tasks/AbilityTask_ApplyRootMotionMoveToForce.cpp
--- a/Engine/Plugins/Runtime/GameplayAbilities/Source/GameplayAbilities/Private/Abilities/Tasks/AbilityTask_ApplyRootMotionMoveToForce.cpp
+++ b/Engine/Plugins/Runtime/GameplayAbilities/Source/GameplayAbilities/Private/Abilities/Tasks/AbilityTask_ApplyRootMotionMoveToForce.cpp
@@ -8,16 +8,6 @@
 #include "AbilitySystemGlobals.h"
 #include "Net/UnrealNetwork.h"
 
-UAbilityTask_ApplyRootMotionMoveToForce::UAbilityTask_ApplyRootMotionMoveToForce(const FObjectInitializer& ObjectInitializer)
-: Super(ObjectInitializer)
-{
-	bSetNewMovementMode = false;
-	NewMovementMode = EMovementMode::MOVE_Walking;
-	PreviousMovementMode = EMovementMode::MOVE_None;
-	bRestrictSpeedToExpected = false;
-	PathOffsetCurve = nullptr;
-}
-
 UAbilityTask_ApplyRootMotionMoveToForce* UAbilityTask_ApplyRootMotionMoveToForce::ApplyRootMotionMoveToForce(UGameplayAbility* OwningAbility, FName TaskInstanceName, FVector TargetLocation, float Duration, bool bSetNewMovementMode, EMovementMode MovementMode, bool bRestrictSpeedToExpected, UCurveVector* PathOffsetCurve, ERootMotionFinishVelocityMode VelocityOnFinishMode, FVector SetVelocityOnFinish, float ClampVelocityOnFinish)
 {
 	UAbilitySystemGlobals::NonShipping_ApplyGlobalAbilityScaler_Duration(Duration);
@@ -61,6 +51,7 @@ void UAbilityTask_ApplyRootMotionMoveToForce::SharedInitAndApply()
 			if (bSetNewMovementMode)
 			{
 				PreviousMovementMode = MovementComponent->MovementMode;
+				PreviousCustomMovementMode = MovementComponent->CustomMovementMode;
 				MovementComponent->SetMovementMode(NewMovementMode);
 			}
 
@@ -169,7 +160,7 @@ void UAbilityTask_ApplyRootMotionMoveToForce::OnDestroy(bool AbilityIsEnding)
 
 		if (bSetNewMovementMode)
 		{
-			MovementComponent->SetMovementMode(PreviousMovementMode);
+ 			MovementComponent->SetMovementMode(PreviousMovementMode, PreviousCustomMovementMode);
 		}
 	}
 
diff --git a/Engine/Plugins/Runtime/GameplayAbilities/Source/GameplayAbilities/Public/Abilities/Tasks/AbilityTask_ApplyRootMotionMoveToForce.h b/Engine/Plugins/Runtime/GameplayAbilities/Source/GameplayAbilities/Public/Abilities/Tasks/AbilityTask_ApplyRootMotionMoveToForce.h
--- a/Engine/Plugins/Runtime/GameplayAbilities/Source/GameplayAbilities/Public/Abilities/Tasks/AbilityTask_ApplyRootMotionMoveToForce.h
+++ b/Engine/Plugins/Runtime/GameplayAbilities/Source/GameplayAbilities/Public/Abilities/Tasks/AbilityTask_ApplyRootMotionMoveToForce.h
@@ -22,7 +22,7 @@ class AActor;
 UCLASS()
 class GAMEPLAYABILITIES_API UAbilityTask_ApplyRootMotionMoveToForce : public UAbilityTask_ApplyRootMotion_Base
 {
-	GENERATED_UCLASS_BODY()
+	GENERATED_BODY()
 
 	UPROPERTY(BlueprintAssignable)
 	FApplyRootMotionMoveToForceDelegate OnTimedOut;
@@ -56,19 +56,20 @@ protected:
 	float Duration;
 
 	UPROPERTY(Replicated)
-	bool bSetNewMovementMode;
+	bool bSetNewMovementMode = false;
 
 	UPROPERTY(Replicated)
-	TEnumAsByte<EMovementMode> NewMovementMode;
+	TEnumAsByte<EMovementMode> NewMovementMode = EMovementMode::MOVE_Walking;
 
 	/** If enabled, we limit velocity to the initial expected velocity to go distance to the target over Duration.
 	 *  This prevents cases of getting really high velocity the last few frames of the root motion if you were being blocked by
 	 *  collision. Disabled means we do everything we can to velocity during the move to get to the TargetLocation. */
 	UPROPERTY(Replicated)
-	bool bRestrictSpeedToExpected;
+	bool bRestrictSpeedToExpected = false;
 
 	UPROPERTY(Replicated)
-	UCurveVector* PathOffsetCurve;
+	UCurveVector* PathOffsetCurve = nullptr;
 
-	EMovementMode PreviousMovementMode;
+	EMovementMode PreviousMovementMode = EMovementMode::MOVE_None;
+	uint8 PreviousCustomMovementMode = 0;
 };
