Gameplay Cues: Change  FPreallocationInfo::PreallocatedInstances to be a reflected variable instead of relying only on EndPlay to clean it up, fixing potential issues if cue BPs are reinstanced during PIE
#tests Edited various GCNLs in Lyra during PIE and recompiled / resaved the BPs a couple of times to force reinstancing

diff --git a/Engine/Plugins/Runtime/GameplayAbilities/Source/GameplayAbilities/Private/GameplayCueManager.cpp b/Engine/Plugins/Runtime/GameplayAbilities/Source/GameplayAbilities/Private/GameplayCueManager.cpp
--- a/Engine/Plugins/Runtime/GameplayAbilities/Source/GameplayAbilities/Private/GameplayCueManager.cpp
+++ b/Engine/Plugins/Runtime/GameplayAbilities/Source/GameplayAbilities/Private/GameplayCueManager.cpp
@@ -431,13 +431,13 @@ AGameplayCueNotify_Actor* UGameplayCueManager::GetInstancedCueActor(AActor* Targ
 		if (UseActorRecycling)
 		{
 			FPreallocationInfo& Info = GetPreallocationInfo(World);
-			TArray<AGameplayCueNotify_Actor*>* PreallocatedList = Info.PreallocatedInstances.Find(CueClass);
-			if (PreallocatedList && PreallocatedList->Num() > 0)
+			FGameplayCueNotifyActorArray* PreallocatedList = Info.PreallocatedInstances.Find(CueClass);
+			if (PreallocatedList && PreallocatedList->Actors.Num() > 0)
 			{
 				SpawnedCue = nullptr;
 				while (true)
 				{
-					SpawnedCue = PreallocatedList->Pop(false);
+					SpawnedCue = PreallocatedList->Actors.Pop(false);
 
 					// Temp: tracking down possible memory corruption
 					// null is maybe ok. But invalid low level is bad and we want to crash hard to find out who/why.
@@ -455,7 +455,7 @@ AGameplayCueNotify_Actor* UGameplayCueManager::GetInstancedCueActor(AActor* Targ
 					// outside of replays, this should not happen. GC Notifies should not be actually destroyed.
 					checkf(World->IsPlayingReplay(), TEXT("Spawned Cue is pending kill, garbage or null: %s."), *GetNameSafe(SpawnedCue));
 
-					if (PreallocatedList->Num() <= 0)
+					if (PreallocatedList->Actors.Num() <= 0)
 					{
 						// Ran out of preallocated instances... break and create a new one.
 						break;
@@ -550,12 +550,12 @@ void UGameplayCueManager::NotifyGameplayCueActorFinished(AGameplayCueNotify_Acto
 			UE_CLOG((GameplayCueActorRecycleDebug>0), LogAbilitySystem, Display, TEXT("NotifyGameplayCueActorFinished %s"), *GetNameSafe(Actor));
 
 			FPreallocationInfo& Info = GetPreallocationInfo(Actor->GetWorld());
-			TArray<AGameplayCueNotify_Actor*>& PreAllocatedList = Info.PreallocatedInstances.FindOrAdd(Actor->GetClass());
+			FGameplayCueNotifyActorArray& PreAllocatedList = Info.PreallocatedInstances.FindOrAdd(Actor->GetClass());
 
 			// Put the actor back in the list
-			if (ensureMsgf(PreAllocatedList.Contains(Actor)==false, TEXT("GC Actor PreallocationList already contains Actor %s"), *GetNameSafe(Actor)))
+			if (ensureMsgf(PreAllocatedList.Actors.Contains(Actor)==false, TEXT("GC Actor PreallocationList already contains Actor %s"), *GetNameSafe(Actor)))
 			{
-				PreAllocatedList.Push(Actor);
+				PreAllocatedList.Actors.Push(Actor);
 			}
 			
 #if WITH_EDITOR
@@ -576,8 +576,8 @@ void UGameplayCueManager::NotifyGameplayCueActorEndPlay(AGameplayCueNotify_Actor
 	if (Actor && Actor->bInRecycleQueue)
 	{
 		FPreallocationInfo& Info = GetPreallocationInfo(Actor->GetWorld());
-		TArray<AGameplayCueNotify_Actor*>& PreAllocatedList = Info.PreallocatedInstances.FindOrAdd(Actor->GetClass());
-		PreAllocatedList.Remove(Actor);
+		FGameplayCueNotifyActorArray& PreAllocatedList = Info.PreallocatedInstances.FindOrAdd(Actor->GetClass());
+		PreAllocatedList.Actors.Remove(Actor);
 	}
 }
 
@@ -1557,7 +1557,7 @@ void UGameplayCueManager::UpdatePreallocation(UWorld* World)
 	{
 		TSubclassOf<AGameplayCueNotify_Actor> GCClass = Info.ClassesNeedingPreallocation.Last();
 		AGameplayCueNotify_Actor* CDO = GCClass->GetDefaultObject<AGameplayCueNotify_Actor>();
-		TArray<AGameplayCueNotify_Actor*>& PreallocatedList = Info.PreallocatedInstances.FindOrAdd(CDO->GetClass());
+		FGameplayCueNotifyActorArray& PreallocatedList = Info.PreallocatedInstances.FindOrAdd(CDO->GetClass());
 
 		AGameplayCueNotify_Actor* PrespawnedInstance = Cast<AGameplayCueNotify_Actor>(World->SpawnActor(CDO->GetClass()));
 		if (ensureMsgf(PrespawnedInstance, TEXT("Failed to prespawn GC notify for: %s"), *GetNameSafe(CDO)))
@@ -1570,10 +1570,10 @@ void UGameplayCueManager::UpdatePreallocation(UWorld* World)
 			}
 
 			PrespawnedInstance->bInRecycleQueue = true;
-			PreallocatedList.Push(PrespawnedInstance);
+			PreallocatedList.Actors.Push(PrespawnedInstance);
 			PrespawnedInstance->SetActorHiddenInGame(true);
 
-			if (PreallocatedList.Num() >= CDO->NumPreallocatedInstances)
+			if (PreallocatedList.Actors.Num() >= CDO->NumPreallocatedInstances)
 			{
 				Info.ClassesNeedingPreallocation.Pop(false);
 			}
@@ -1619,22 +1619,6 @@ void UGameplayCueManager::OnWorldCleanup(UWorld* World, bool bSessionEnded, bool
 	IGameplayCueInterface::ClearTagToFunctionMap();
 }
 
-void UGameplayCueManager::AddReferencedObjects(UObject* InThis, FReferenceCollector& Collector)
-{
-	Super::AddReferencedObjects(InThis, Collector);
-
-	UGameplayCueManager* GCManager = Cast<UGameplayCueManager>(InThis);
-	for (int32 idx = 0; idx < GCManager->PreallocationInfoList_Internal.Num(); ++idx)
-	{
-		FPreallocationInfo& Info = GCManager->PreallocationInfoList_Internal[idx];
-		for (auto &It : Info.PreallocatedInstances)
-		{
-			TArray<AGameplayCueNotify_Actor*>& GameplayCueArray = It.Value;
-			Collector.AddReferencedObjects(GameplayCueArray);
-		}
-	}
-}
-
 void UGameplayCueManager::DumpPreallocationStats(UWorld* World)
 {
 	if (World == nullptr)
@@ -1649,7 +1633,7 @@ void UGameplayCueManager::DumpPreallocationStats(UWorld* World)
 		{
 			if (AGameplayCueNotify_Actor* CDO = ThisClass->GetDefaultObject<AGameplayCueNotify_Actor>())
 			{
-				TArray<AGameplayCueNotify_Actor*>& List = It.Value;
+				TArray<AGameplayCueNotify_Actor*>& List = It.Value.Actors;
 				if (List.Num() > CDO->NumPreallocatedInstances)
 				{
 					ABILITY_LOG(Display, TEXT("Notify class: %s was used simultaneously %d times. The CDO default is %d preallocated instanced."), *ThisClass->GetName(), List.Num(),  CDO->NumPreallocatedInstances); 
diff --git a/Engine/Plugins/Runtime/GameplayAbilities/Source/GameplayAbilities/Public/GameplayCueManager.h b/Engine/Plugins/Runtime/GameplayAbilities/Source/GameplayAbilities/Public/GameplayCueManager.h
--- a/Engine/Plugins/Runtime/GameplayAbilities/Source/GameplayAbilities/Public/GameplayCueManager.h
+++ b/Engine/Plugins/Runtime/GameplayAbilities/Source/GameplayAbilities/Public/GameplayCueManager.h
@@ -275,8 +275,6 @@ protected:
 	TSharedPtr<FStreamableHandle> GameplayCueAssetHandle;
 
 public:		
-	static void AddReferencedObjects(UObject* InThis, FReferenceCollector& Collector);
-
 	/** Called before loading any gameplay cue notifies from object libraries. Allows subclasses to skip notifies. */
 	virtual bool ShouldLoadGameplayCueAssetData(const FAssetData& Data) const { return true; }
 
diff --git a/Engine/Plugins/Runtime/GameplayAbilities/Source/GameplayAbilities/Public/GameplayCue_Types.h b/Engine/Plugins/Runtime/GameplayAbilities/Source/GameplayAbilities/Public/GameplayCue_Types.h
--- a/Engine/Plugins/Runtime/GameplayAbilities/Source/GameplayAbilities/Public/GameplayCue_Types.h
+++ b/Engine/Plugins/Runtime/GameplayAbilities/Source/GameplayAbilities/Public/GameplayCue_Types.h
@@ -59,6 +59,14 @@ struct FGameplayCuePendingExecute
 	FGameplayCueParameters CueParameters;
 };
 
+USTRUCT()
+struct FGameplayCueNotifyActorArray
+{
+	GENERATED_BODY()
+
+	UPROPERTY(transient)
+	TArray<AGameplayCueNotify_Actor*> Actors;
+};
 
 /** Struct for pooling and preallocating gameplaycuenotify_actor classes. This data is per world and used to track what actors are available to recycle and which classes need to preallocate instances of those actors */
 USTRUCT()
@@ -67,11 +75,12 @@ struct FPreallocationInfo
 	GENERATED_USTRUCT_BODY()
 
 	/** Raw list of pooled instances. This relies on NotifyGameplayCueActorEndPlay always being called when actor is destroyed */
-	TMap<UClass*, TArray<AGameplayCueNotify_Actor*> >	PreallocatedInstances;
+	UPROPERTY(transient)
+	TMap<UClass*, FGameplayCueNotifyActorArray> PreallocatedInstances;
 
-	/** List of calsses that will be pooled */
+	/** List of classes that will be pooled */
 	UPROPERTY(transient)
-	TArray<TSubclassOf<AGameplayCueNotify_Actor>>	ClassesNeedingPreallocation;
+	TArray<TSubclassOf<AGameplayCueNotify_Actor>> ClassesNeedingPreallocation;
 
 	/** World that owns this list */
 	FObjectKey OwningWorldKey;
