SubObject API
* GameplayTasksComponent and AbilitySystemComponent now support the registered subobject API
* By default the components keep replicating their subobjects via the old way (the ReplicateSubObjects function).
* Deleted some UE_DEPRECATED(4.XX) tags and moved or deleted the code they touched.
* Moved some replicated members to private access. Derived classes should use the Get/Set functions instead.

diff --git a/Engine/Plugins/Runtime/GameplayAbilities/Source/GameplayAbilities/Private/AbilitySystemComponent.cpp b/Engine/Plugins/Runtime/GameplayAbilities/Source/GameplayAbilities/Private/AbilitySystemComponent.cpp
--- a/Engine/Plugins/Runtime/GameplayAbilities/Source/GameplayAbilities/Private/AbilitySystemComponent.cpp
+++ b/Engine/Plugins/Runtime/GameplayAbilities/Source/GameplayAbilities/Private/AbilitySystemComponent.cpp
@@ -95,7 +95,7 @@ const UAttributeSet* UAbilitySystemComponent::GetOrCreateAttributeSubobject(TSub
 		if (!MyAttributes)
 		{
 			UAttributeSet* Attributes = NewObject<UAttributeSet>(OwningActor, AttributeClass);
-			GetSpawnedAttributes_Mutable().AddUnique(Attributes);
+			AddSpawnedAttribute(Attributes);
 			MyAttributes = Attributes;
 			bIsNetDirty = true;
 		}
@@ -232,6 +232,15 @@ void UAbilitySystemComponent::BeginPlay()
 
 	// Cache net role here as well since for map-placed actors on clients, the Role may not be set correctly yet in OnRegister.
 	CacheIsNetSimulated();
+
+	// Register the spawned attributes to the replicated sub object list.
+	if (IsUsingRegisteredSubObjectList())
+	{
+		for (UAttributeSet* ReplicatedAttribute : SpawnedAttributes)
+		{
+			AddReplicatedSubObject(ReplicatedAttribute);
+		}
+	}
 }
 
 void UAbilitySystemComponent::CacheIsNetSimulated()
@@ -1552,11 +1561,10 @@ int32 UAbilitySystemComponent::RemoveActiveEffects(const FGameplayEffectQuery& Q
 	return 0;
 }
 
-void UAbilitySystemComponent::GetLifetimeReplicatedProps(TArray< FLifetimeProperty > & OutLifetimeProps) const
-{	
+void UAbilitySystemComponent::GetLifetimeReplicatedProps(TArray< FLifetimeProperty >& OutLifetimeProps) const
+{
 	// Fast Arrays don't use push model, but there's no harm in marking them with it.
 	// The flag will just be ignored.
-
 	FDoRepLifetimeParams Params;
 	Params.bIsPushBased = true;
 
@@ -1616,7 +1624,7 @@ bool UAbilitySystemComponent::ReplicateSubobjects(class UActorChannel *Channel,
 		}
 	}
 
-	for (UGameplayAbility* Ability : AllReplicatedInstancedAbilities)
+	for (UGameplayAbility* Ability : GetReplicatedInstancedAbilities())
 	{
 		if (IsValid(Ability))
 		{
@@ -2740,7 +2748,6 @@ void UAbilitySystemComponent::Debug_Internal(FAbilitySystemComponentDebugInfo& I
 	Info.YL = MaxCharHeight;
 }
 
-PRAGMA_DISABLE_DEPRECATION_WARNINGS
 void UAbilitySystemComponent::SetOwnerActor(AActor* NewOwnerActor)
 {
 	MARK_PROPERTY_DIRTY_FROM_NAME(UAbilitySystemComponent, OwnerActor, this);
@@ -2755,11 +2762,6 @@ void UAbilitySystemComponent::SetOwnerActor(AActor* NewOwnerActor)
 	}
 }
 
-AActor* UAbilitySystemComponent::GetOwnerActor() const
-{
-	return OwnerActor;
-}
-
 void UAbilitySystemComponent::SetAvatarActor_Direct(AActor* NewAvatarActor)
 {
 	MARK_PROPERTY_DIRTY_FROM_NAME(UAbilitySystemComponent, AvatarActor, this);
@@ -2774,11 +2776,6 @@ void UAbilitySystemComponent::SetAvatarActor_Direct(AActor* NewAvatarActor)
 	}
 }
 
-AActor* UAbilitySystemComponent::GetAvatarActor_Direct() const
-{
-	return AvatarActor;
-}
-
 void UAbilitySystemComponent::OnAvatarActorDestroyed(AActor* InActor)
 {
 	if (InActor == AvatarActor)
@@ -2807,21 +2804,75 @@ void UAbilitySystemComponent::SetSpawnedAttributes(const TArray<UAttributeSet*>&
 			ActorOwner->OnEndPlay.RemoveDynamic(this, &UAbilitySystemComponent::OnSpawnedAttributesEndPlayed);
 		}
 	}
-	GetSpawnedAttributes_Mutable() = NewSpawnedAttributes;
-	for (int32 Index = NewSpawnedAttributes.Num() - 1; Index >= 0; --Index)
+
+	// Clean the previous list
+	RemoveAllSpawnedAttributes();
+
+	// Add the elements from the new list
+	for (UAttributeSet* NewAttribute : NewSpawnedAttributes)
 	{
-		UAttributeSet* AttributeSet = NewSpawnedAttributes[Index];
-		AActor* ActorOwner = AttributeSet->GetTypedOuter<AActor>();
+		AddSpawnedAttribute(NewAttribute);
+
+		AActor* ActorOwner = NewAttribute->GetTypedOuter<AActor>();
 		if (ActorOwner)
 		{
 			ActorOwner->OnEndPlay.AddUniqueDynamic(this, &UAbilitySystemComponent::OnSpawnedAttributesEndPlayed);
 		}
 	}
+
+	SetSpawnedAttributesListDirty();
 }
 
-TArray<UAttributeSet*>& UAbilitySystemComponent::GetSpawnedAttributes_Mutable()
+void UAbilitySystemComponent::AddSpawnedAttribute(UAttributeSet* Attribute)
+{
+
+	if (SpawnedAttributes.Find(Attribute) == INDEX_NONE)
+	{
+		if (IsUsingRegisteredSubObjectList() && HasBegunPlay())
+		{
+			AddReplicatedSubObject(Attribute);
+		}
+
+		SpawnedAttributes.Add(Attribute);
+		SetSpawnedAttributesListDirty();
+	}
+}
+
+void UAbilitySystemComponent::RemoveSpawnedAttribute(UAttributeSet* Attribute)
+{
+	if (SpawnedAttributes.RemoveSingle(Attribute) > 0)
+	{
+		if (IsUsingRegisteredSubObjectList() )
+		{
+			RemoveReplicatedSubObject(Attribute);
+		}
+
+		SetSpawnedAttributesListDirty();
+	}
+}
+
+void UAbilitySystemComponent::RemoveAllSpawnedAttributes()
+{
+	if (IsUsingRegisteredSubObjectList())
+	{
+		for (UAttributeSet* OldAttribute : SpawnedAttributes)
+		{
+			RemoveReplicatedSubObject(OldAttribute);
+		}
+	}
+
+	SpawnedAttributes.Empty();
+	SetSpawnedAttributesListDirty();
+}
+
+void UAbilitySystemComponent::SetSpawnedAttributesListDirty()
 {
 	MARK_PROPERTY_DIRTY_FROM_NAME(UAbilitySystemComponent, SpawnedAttributes, this);
+}
+
+TArray<UAttributeSet*>& UAbilitySystemComponent::GetSpawnedAttributes_Mutable()
+{
+	SetSpawnedAttributesListDirty();
 	return SpawnedAttributes;
 }
 
@@ -2830,6 +2881,76 @@ const TArray<UAttributeSet*>& UAbilitySystemComponent::GetSpawnedAttributes() co
 	return SpawnedAttributes;
 }
 
+void UAbilitySystemComponent::OnRep_SpawnedAttributes(const TArray<UAttributeSet*>& PreviousSpawnedAttributes)
+{
+	if (IsUsingRegisteredSubObjectList())
+	{
+		// Find the attributes that got removed
+		for (UAttributeSet* PreviousAttributeSet : PreviousSpawnedAttributes)
+		{
+			if (PreviousAttributeSet)
+			{
+				const bool bWasRemoved = SpawnedAttributes.Find(PreviousAttributeSet) == INDEX_NONE;
+				if (bWasRemoved)
+				{
+					RemoveReplicatedSubObject(PreviousAttributeSet);
+				}
+			}
+		}
+
+		// Find the attributes that got added
+		for (UAttributeSet* NewAttributeSet : SpawnedAttributes)
+		{
+			if (NewAttributeSet)
+			{
+				const bool bIsAdded = PreviousSpawnedAttributes.Find(NewAttributeSet) == INDEX_NONE;
+				if (bIsAdded)
+				{
+					AddReplicatedSubObject(NewAttributeSet);
+				}
+			}
+		}
+	}
+}
+
+void UAbilitySystemComponent::AddReplicatedInstancedAbility(UGameplayAbility* GameplayAbility)
+{
+	TArray<UGameplayAbility*>& ReplicatedAbilities = GetReplicatedInstancedAbilities_Mutable();
+	if (ReplicatedAbilities.Find(GameplayAbility) == INDEX_NONE)
+	{
+		ReplicatedAbilities.Add(GameplayAbility);
+		
+		if (IsUsingRegisteredSubObjectList())
+		{
+			AddReplicatedSubObject(GameplayAbility);
+		}
+	}
+}
+
+void UAbilitySystemComponent::RemoveReplicatedInstancedAbility(UGameplayAbility* GameplayAbility)
+{
+	const bool bWasRemoved = GetReplicatedInstancedAbilities_Mutable().RemoveSingle(GameplayAbility) > 0;
+	if (bWasRemoved && IsUsingRegisteredSubObjectList())
+	{
+		RemoveReplicatedSubObject(GameplayAbility);
+	}
+}
+
+void UAbilitySystemComponent::RemoveAllReplicatedInstancedAbilities()
+{
+	TArray<UGameplayAbility*>& ReplicatedAbilities = GetReplicatedInstancedAbilities_Mutable();
+
+	if (IsUsingRegisteredSubObjectList())
+	{
+		for (UGameplayAbility* ReplicatedAbility : ReplicatedAbilities)
+		{
+			RemoveReplicatedSubObject(ReplicatedAbility);
+		}
+	}
+
+	ReplicatedAbilities.Empty();
+}
+
 void UAbilitySystemComponent::OnSpawnedAttributesEndPlayed(AActor* InActor, EEndPlayReason::Type EndPlayReason)
 {
 	TArray<UAttributeSet*>& LocalSpawnedAttributes = GetSpawnedAttributes_Mutable();
@@ -2848,6 +2969,7 @@ void UAbilitySystemComponent::SetClientDebugStrings(TArray<FString>&& NewClientD
 	GetClientDebugStrings_Mutable() = NewClientDebugStrings;
 }
 
+PRAGMA_DISABLE_DEPRECATION_WARNINGS
 TArray<FString>& UAbilitySystemComponent::GetClientDebugStrings_Mutable()
 {
 	MARK_PROPERTY_DIRTY_FROM_NAME(UAbilitySystemComponent, ClientDebugStrings, this);
diff --git a/Engine/Plugins/Runtime/GameplayAbilities/Source/GameplayAbilities/Private/AbilitySystemComponent_Abilities.cpp b/Engine/Plugins/Runtime/GameplayAbilities/Source/GameplayAbilities/Private/AbilitySystemComponent_Abilities.cpp
--- a/Engine/Plugins/Runtime/GameplayAbilities/Source/GameplayAbilities/Private/AbilitySystemComponent_Abilities.cpp
+++ b/Engine/Plugins/Runtime/GameplayAbilities/Source/GameplayAbilities/Private/AbilitySystemComponent_Abilities.cpp
@@ -55,12 +55,11 @@ void UAbilitySystemComponent::InitializeComponent()
 	InitAbilityActorInfo(Owner, Owner);	// Default init to our outer owner
 
 	// cleanup any bad data that may have gotten into SpawnedAttributes
-	TArray<UAttributeSet*>& SpawnedAttributesRef = GetSpawnedAttributes_Mutable();
-	for (int32 Idx = SpawnedAttributesRef.Num()-1; Idx >= 0; --Idx)
+	for (int32 Idx = SpawnedAttributes.Num()-1; Idx >= 0; --Idx)
 	{
-		if (SpawnedAttributesRef[Idx] == nullptr)
+		if (SpawnedAttributes[Idx] == nullptr)
 		{
-			SpawnedAttributesRef.RemoveAt(Idx);
+			SpawnedAttributes.RemoveAt(Idx);
 		}
 	}
 
@@ -72,10 +71,12 @@ void UAbilitySystemComponent::InitializeComponent()
 		UAttributeSet* Set = Cast<UAttributeSet>(Obj);
 		if (Set)  
 		{
-			SpawnedAttributesRef.AddUnique(Set);
+			SpawnedAttributes.AddUnique(Set);
 			bIsNetDirty = true;
 		}
 	}
+
+	SetSpawnedAttributesListDirty();
 }
 
 void UAbilitySystemComponent::UninitializeComponent()
@@ -567,7 +568,7 @@ void UAbilitySystemComponent::OnRemoveAbility(FGameplayAbilitySpec& Spec)
 				if (GetOwnerRole() == ROLE_Authority || Instance->GetReplicationPolicy() == EGameplayAbilityReplicationPolicy::ReplicateNo)
 				{
 					// Only destroy if we're the server or this isn't replicated. Can't destroy on the client or replication will fail when it replicates the end state
-					AllReplicatedInstancedAbilities.Remove(Instance);
+					RemoveReplicatedInstancedAbility(Instance);
 					Instance->MarkAsGarbage();
 				}
 			}
@@ -638,13 +639,19 @@ void UAbilitySystemComponent::CheckForClearedAbilities()
 		// We leave around the empty trigger stub, it's likely to be added again
 	}
 
-	for (int32 i = 0; i < AllReplicatedInstancedAbilities.Num(); i++)
+	TArray<UGameplayAbility*>& ReplicatedAbilities = GetReplicatedInstancedAbilities_Mutable();
+	for (int32 i = 0; i < ReplicatedAbilities.Num(); i++)
 	{
-		UGameplayAbility* Ability = AllReplicatedInstancedAbilities[i];
+		UGameplayAbility* Ability = ReplicatedAbilities[i];
 
 		if (!IsValid(Ability))
 		{
-			AllReplicatedInstancedAbilities.RemoveAt(i);
+			if (IsUsingRegisteredSubObjectList())
+			{
+				RemoveReplicatedSubObject(Ability);
+			}
+
+			ReplicatedAbilities.RemoveAt(i);
 			i--;
 		}
 	}
@@ -969,7 +976,7 @@ UGameplayAbility* UAbilitySystemComponent::CreateNewInstanceOfAbility(FGameplayA
 	if (AbilityInstance->GetReplicationPolicy() != EGameplayAbilityReplicationPolicy::ReplicateNo)
 	{
 		Spec.ReplicatedInstances.Add(AbilityInstance);
-		AllReplicatedInstancedAbilities.Add(AbilityInstance);
+		AddReplicatedInstancedAbility(AbilityInstance);
 	}
 	else
 	{
@@ -1020,7 +1027,7 @@ void UAbilitySystemComponent::NotifyAbilityEnded(FGameplayAbilitySpecHandle Hand
 			if (OwnerRole == ROLE_Authority)
 			{
 				Spec->ReplicatedInstances.Remove(Ability);
-				AllReplicatedInstancedAbilities.Remove(Ability);
+				RemoveReplicatedInstancedAbility(Ability);
 				Ability->MarkAsGarbage();
 			}
 		}
diff --git a/Engine/Plugins/Runtime/GameplayAbilities/Source/GameplayAbilities/Public/AbilitySystemComponent.h b/Engine/Plugins/Runtime/GameplayAbilities/Source/GameplayAbilities/Public/AbilitySystemComponent.h
--- a/Engine/Plugins/Runtime/GameplayAbilities/Source/GameplayAbilities/Public/AbilitySystemComponent.h
+++ b/Engine/Plugins/Runtime/GameplayAbilities/Source/GameplayAbilities/Public/AbilitySystemComponent.h
@@ -122,13 +122,6 @@ class GAMEPLAYABILITIES_API UAbilitySystemComponent : public UGameplayTasksCompo
 		return (T*)GetOrCreateAttributeSubobject(T::StaticClass());
 	}
 
-	template <class T>
-	UE_DEPRECATED(4.25, "AddDefaultSubobjectSet is deprecated. Use AddAttributeSetSubobject instead.")
-	const T*	AddDefaultSubobjectSet(T* Subobject)
-	{
-		return AddAttributeSetSubobject(Subobject);
-	}
-
 	/** 
 	 * Manually add a new attribute set that is a subobject of this ability system component.
 	 * All subobjects of this component are automatically added during initialization.
@@ -136,7 +129,7 @@ class GAMEPLAYABILITIES_API UAbilitySystemComponent : public UGameplayTasksCompo
 	template <class T>
 	const T* AddAttributeSetSubobject(T* Subobject)
 	{
-		GetSpawnedAttributes_Mutable().AddUnique(Subobject);
+		AddSpawnedAttribute(Subobject);
 		return Subobject;
 	}
 
@@ -181,15 +174,24 @@ class GAMEPLAYABILITIES_API UAbilitySystemComponent : public UGameplayTasksCompo
 	UPROPERTY(EditAnywhere, Category="AttributeTest")
 	TArray<FAttributeDefaults>	DefaultStartingData;
 
-	/** List of attribute sets */
-	UE_DEPRECATED(4.26, "This will be made private in future engine versions. Use SetSpawnedAttributes, GetSpawnedAttributes, or GetSpawnedAttributes_Mutable instead.")
-	UPROPERTY(Replicated)
-	TArray<UAttributeSet*>	SpawnedAttributes;
-
+	/** Remove all current AttributeSets and register the ones in the passed array. Note that it's better to call Add/Remove directly when possible. */
 	void SetSpawnedAttributes(const TArray<UAttributeSet*>& NewAttributeSet);
+
+	UE_DEPRECATED(5.1, "This function will be made private. Use Add/Remove SpawnedAttributes instead")
 	TArray<UAttributeSet*>& GetSpawnedAttributes_Mutable();
+
+	/** Access the spawned attributes list when you don't intend to modify the list. */
 	const TArray<UAttributeSet*>& GetSpawnedAttributes() const;
 
+	/** Add a new attribute set */
+	void AddSpawnedAttribute(UAttributeSet* Attribute);
+
+	/** Remove an existing attribute set */
+	void RemoveSpawnedAttribute(UAttributeSet* Attribute);
+
+	/** Remove all attribute sets */
+	void RemoveAllSpawnedAttributes();
+
 
 	/** The linked Anim Instance that this component will play montages in. Use NAME_None for the main anim instance. */
 	UPROPERTY(BlueprintReadWrite, EditAnywhere, Category = "Skills")
@@ -1419,21 +1421,23 @@ class GAMEPLAYABILITIES_API UAbilitySystemComponent : public UGameplayTasksCompo
 	//	Actor interaction
 	// ----------------------------------------------------------------------------------------------------------------	
 
+private:
+
 	/** The actor that owns this component logically */
-	UE_DEPRECATED(4.26, "This will be made private in future engine versions. Use SetOwnerActor or GetOwnerActor instead.")
 	UPROPERTY(ReplicatedUsing = OnRep_OwningActor)
 	AActor* OwnerActor;
 
-	void SetOwnerActor(AActor* NewOwnerActor);
-	AActor* GetOwnerActor() const;
-
-	UE_DEPRECATED(4.26, "This will be made private in future engine versions. Use SetAvatarActor or GetAvatarActor instead.")
 	/** The actor that is the physical representation used for abilities. Can be NULL */
 	UPROPERTY(ReplicatedUsing = OnRep_OwningActor)
 	AActor* AvatarActor;
 
+public:
+
+	void SetOwnerActor(AActor* NewOwnerActor);
+	AActor* GetOwnerActor() const { return OwnerActor; }
+
 	void SetAvatarActor_Direct(AActor* NewAvatarActor);
-	AActor* GetAvatarActor_Direct() const;
+	AActor* GetAvatarActor_Direct() const { return AvatarActor; }
 	
 	UFUNCTION()
 	void OnRep_OwningActor();
@@ -1603,9 +1607,24 @@ protected:
 	TArray<TPair<FGameplayTagContainer, FGameplayEventTagMulticastDelegate>> GameplayEventTagContainerDelegates;
 
 	/** Full list of all instance-per-execution gameplay abilities associated with this component */
+	UE_DEPRECATED(5.1, "This array will be made private. Use GetReplicatedInstancedAbilities, AddReplicatedInstancedAbility or RemoveReplicatedInstancedAbility instead.")
 	UPROPERTY()
 	TArray<UGameplayAbility*>	AllReplicatedInstancedAbilities;
 
+	/** Full list of all instance-per-execution gameplay abilities associated with this component */
+	PRAGMA_DISABLE_DEPRECATION_WARNINGS
+	const TArray<UGameplayAbility*>& GetReplicatedInstancedAbilities() const { return AllReplicatedInstancedAbilities; }
+	PRAGMA_ENABLE_DEPRECATION_WARNINGS
+
+	/** Add a gameplay ability associated to this component */
+	void AddReplicatedInstancedAbility(UGameplayAbility* GameplayAbility);
+
+	/** Remove a gameplay ability associated to this component */
+	void RemoveReplicatedInstancedAbility(UGameplayAbility* GameplayAbility);
+
+	/** Unregister all the gameplay abilities of this component */
+	void RemoveAllReplicatedInstancedAbilities();
+
 	/** Will be called from GiveAbility or from OnRep. Initializes events (triggers and inputs) with the given ability */
 	virtual void OnGiveAbility(FGameplayAbilitySpec& AbilitySpec);
 
@@ -1819,6 +1838,24 @@ protected:
 	friend class UAbilitySystemGlobals;
 
 private:
+
+	// Needs to be called when modifying the SpawnedAttributes array for changes to be replicated
+	void SetSpawnedAttributesListDirty();
+
+    // Private accessor to the AllReplicatedInstancedAbilities array until the deprecation tag on it is removed and we can reference the array directly again.
+	PRAGMA_DISABLE_DEPRECATION_WARNINGS
+	TArray<UGameplayAbility*>& GetReplicatedInstancedAbilities_Mutable() { return AllReplicatedInstancedAbilities; }
+	PRAGMA_ENABLE_DEPRECATION_WARNINGS
+
+private:
+
+	/** List of attribute sets */
+	UPROPERTY(Replicated, ReplicatedUsing=OnRep_SpawnedAttributes)
+	TArray<UAttributeSet*>	SpawnedAttributes;
+
+	UFUNCTION()
+	void OnRep_SpawnedAttributes(const TArray<UAttributeSet*>& PreviousSpawnedAttributes);
+
 	FDelegateHandle MonitoredTagChangedDelegateHandle;
 	FTimerHandle    OnRep_ActivateAbilitiesTimerHandle;
 
diff --git a/Engine/Source/Runtime/GameplayTasks/Classes/GameplayTasksComponent.h b/Engine/Source/Runtime/GameplayTasks/Classes/GameplayTasksComponent.h
--- a/Engine/Source/Runtime/GameplayTasks/Classes/GameplayTasksComponent.h
+++ b/Engine/Source/Runtime/GameplayTasks/Classes/GameplayTasksComponent.h
@@ -81,20 +81,23 @@ protected:
 	FGameplayResourceSet CurrentlyClaimedResources;
 
 	/** Tasks that run on simulated proxies */
-	UE_DEPRECATED(4.26, "This will be made private in future versions. Use GetSimulatedTasks, GetSimulatedTasks_Mutable, or SetSimulatedTasks instead.")
-	UPROPERTY(ReplicatedUsing = OnRep_SimulatedTasks)
-	TArray<TObjectPtr<UGameplayTask>> SimulatedTasks;
-
-PRAGMA_DISABLE_DEPRECATION_WARNINGS
 	const TArray<UGameplayTask*>& GetSimulatedTasks()
 	{
 		return SimulatedTasks;
 	}
-PRAGMA_ENABLE_DEPRECATION_WARNINGS
 
+	UE_DEPRECATED(5.1, "This will be removed in future versions. Use AddSimulatedTask or RemoveSimulatedTask to modify the array")
 	TArray<UGameplayTask*>& GetSimulatedTasks_Mutable();
+
+	/** Remove all current tasks and register the one in the passed array. It's optimal to use Add/Remove of single tasks directly if possible.*/
 	void SetSimulatedTasks(const TArray<UGameplayTask*>& NewSimulatedTasks);
 
+	/** Add a new simulated task. Returns true if the task was added to the list. Returns false if the task was already registered. */
+	bool AddSimulatedTask(UGameplayTask* NewTask);
+
+	/** Remove an existing simulated task */
+	void RemoveSimulatedTask(UGameplayTask* NewTask);
+
 	UPROPERTY()
 	TArray<TObjectPtr<UGameplayTask>> TaskPriorityQueue;
 	
@@ -119,7 +122,7 @@ public:
 	UGameplayTasksComponent(const FObjectInitializer& ObjectInitializer);
 	
 	UFUNCTION()
-	void OnRep_SimulatedTasks();
+	void OnRep_SimulatedTasks(const TArray<UGameplayTask*>& PreviousSimulatedTasks);
 
 	virtual void GetLifetimeReplicatedProps(TArray< FLifetimeProperty >& OutLifetimeProps) const override;
 	virtual bool ReplicateSubobjects(UActorChannel *Channel, FOutBunch* Bunch, FReplicationFlags* RepFlags) override;
@@ -155,6 +158,9 @@ public:
 	virtual void OnGameplayTaskDeactivated(UGameplayTask& Task) override;
 	// END IGameplayTaskOwnerInterface
 
+	// ActorComponent overrides
+	virtual void BeginPlay() override;
+
 	UFUNCTION(BlueprintCallable, DisplayName="Run Gameplay Task", meta=(ScriptName="RunGameplayTask"), Category = "Gameplay Tasks", meta = (AutoCreateRefTerm = "AdditionalRequiredResources, AdditionalClaimedResources", AdvancedDisplay = "AdditionalRequiredResources, AdditionalClaimedResources"))
 	static EGameplayTaskRunResult K2_RunGameplayTask(TScriptInterface<IGameplayTaskOwnerInterface> TaskOwner, UGameplayTask* Task, uint8 Priority, TArray<TSubclassOf<UGameplayTaskResource> > AdditionalRequiredResources, TArray<TSubclassOf<UGameplayTaskResource> > AdditionalClaimedResources);
 
@@ -199,6 +205,13 @@ private:
 	void RemoveTaskFromPriorityQueue(UGameplayTask& Task);
 
 	FORCEINLINE bool CanProcessEvents() const { return !bInEventProcessingInProgress && (EventLockCounter == 0); }
+
+	void SetSimulatedTasksNetDirty();
+
+	/** Tasks that run on simulated proxies */
+	UPROPERTY(ReplicatedUsing = OnRep_SimulatedTasks)
+	TArray<TObjectPtr<UGameplayTask>> SimulatedTasks;
+
 };
 
 typedef UGameplayTasksComponent::GameplayTaskContainerType::TConstIterator FConstGameplayTaskIterator;
diff --git a/Engine/Source/Runtime/GameplayTasks/Private/GameplayTasksComponent.cpp b/Engine/Source/Runtime/GameplayTasks/Private/GameplayTasksComponent.cpp
--- a/Engine/Source/Runtime/GameplayTasks/Private/GameplayTasksComponent.cpp
+++ b/Engine/Source/Runtime/GameplayTasks/Private/GameplayTasksComponent.cpp
@@ -56,6 +56,22 @@ UGameplayTasksComponent::UGameplayTasksComponent(const FObjectInitializer& Objec
 	bInEventProcessingInProgress = false;
 	TopActivePriority = 0;
 }
+
+void UGameplayTasksComponent::BeginPlay()
+{
+	Super::BeginPlay();
+
+	if (IsUsingRegisteredSubObjectList())
+	{
+		for (UGameplayTask* SimulatedTask : SimulatedTasks)
+		{
+			if (SimulatedTask)
+			{
+				AddReplicatedSubObject(SimulatedTask, COND_SkipOwner);
+			}
+		}
+	}	
+}
 	
 void UGameplayTasksComponent::OnGameplayTaskActivated(UGameplayTask& Task)
 {
@@ -77,9 +93,8 @@ void UGameplayTasksComponent::OnGameplayTaskActivated(UGameplayTask& Task)
 	
 	if (Task.IsSimulatedTask())
 	{
-		TArray<UGameplayTask*>& MutableSimulatedTasks = GetSimulatedTasks_Mutable();
-		check(MutableSimulatedTasks.Contains(&Task) == false);
-		MutableSimulatedTasks.Add(&Task);
+		const bool bWasAdded = AddSimulatedTask(&Task);
+		check(bWasAdded == true);
 		bIsNetDirty = true;
 	}
 
@@ -124,7 +139,7 @@ void UGameplayTasksComponent::OnGameplayTaskDeactivated(UGameplayTask& Task)
 
 	if (Task.IsSimulatedTask())
 	{
-		GetSimulatedTasks_Mutable().RemoveSingleSwap(&Task);
+		RemoveSimulatedTask(&Task);
 		bIsNetDirty = true;
 	}
 
@@ -154,10 +169,7 @@ void UGameplayTasksComponent::GetLifetimeReplicatedProps(TArray< FLifetimeProper
 	// Intentionally not calling super: We do not want to replicate bActive which controls ticking. We sometimes need to tick on client predictively.
 	DISABLE_ALL_CLASS_REPLICATED_PROPERTIES(Super, EFieldIteratorFlags::IncludeSuper);
 
-	FDoRepLifetimeParams Params;
-	Params.bIsPushBased = true;
-	Params.Condition = COND_SkipOwner;
-
+	const FDoRepLifetimeParams Params{ COND_SkipOwner, REPNOTIFY_Always, true };
 	DOREPLIFETIME_WITH_PARAMS_FAST(UGameplayTasksComponent, SimulatedTasks, Params);
 }
 
@@ -179,21 +191,47 @@ bool UGameplayTasksComponent::ReplicateSubobjects(UActorChannel* Channel, class
 	return WroteSomething;
 }
 
-void UGameplayTasksComponent::OnRep_SimulatedTasks()
+void UGameplayTasksComponent::OnRep_SimulatedTasks(const TArray<UGameplayTask*>& PreviousSimulatedTasks)
 {
+	if (IsUsingRegisteredSubObjectList())
+	{
+		// Find if any tasks got removed
+		for (UGameplayTask* OldSimulatedTask : PreviousSimulatedTasks)
+		{
+			const bool bIsRemoved = SimulatedTasks.Find(OldSimulatedTask) == INDEX_NONE;
+			if (bIsRemoved)
+			{
+				RemoveReplicatedSubObject(OldSimulatedTask);
+			}
+		}
+	}
+
 	for (UGameplayTask* SimulatedTask : GetSimulatedTasks())
 	{
-		// Temp check 
-		if (SimulatedTask && SimulatedTask->IsTickingTask() && TickingTasks.Contains(SimulatedTask) == false)
+		if (SimulatedTask)
 		{
-			SimulatedTask->InitSimulatedTask(*this);
+			// If the task needs to be ticked and isn't yet.
+			if (SimulatedTask->IsTickingTask() && TickingTasks.Contains(SimulatedTask) == false)
+			{
+				SimulatedTask->InitSimulatedTask(*this);
+
+				TickingTasks.Add(SimulatedTask);
 
-			TickingTasks.Add(SimulatedTask);
+				// If this is our first ticking task, set this component as active so it begins ticking
+				if (TickingTasks.Num() == 1)
+				{
+					UpdateShouldTick();
+				}
+			}
 
-			// If this is our first ticking task, set this component as active so it begins ticking
-			if (TickingTasks.Num() == 1)
+			// See if it's a new task that needs to be registered
+			if (IsUsingRegisteredSubObjectList())
 			{
-				UpdateShouldTick();
+				const bool bIsNew = PreviousSimulatedTasks.Find(SimulatedTask) == INDEX_NONE;
+				if (bIsNew)
+				{
+					AddReplicatedSubObject(SimulatedTask, COND_SkipOwner);
+				}
 			}
 		}
 	}
@@ -680,17 +718,73 @@ EGameplayTaskRunResult UGameplayTasksComponent::RunGameplayTask(IGameplayTaskOwn
 	return EGameplayTaskRunResult::Error;
 }
 
-PRAGMA_DISABLE_DEPRECATION_WARNINGS
 TArray<UGameplayTask*>& UGameplayTasksComponent::GetSimulatedTasks_Mutable()
 {
 	MARK_PROPERTY_DIRTY_FROM_NAME(UGameplayTasksComponent, SimulatedTasks, this);
 	return SimulatedTasks;
 }
-PRAGMA_ENABLE_DEPRECATION_WARNINGS
+
+bool UGameplayTasksComponent::AddSimulatedTask(UGameplayTask* NewTask)
+{
+	if (SimulatedTasks.Find(NewTask) == INDEX_NONE)
+	{
+		SimulatedTasks.Add(NewTask);
+		SetSimulatedTasksNetDirty();
+
+		if (IsUsingRegisteredSubObjectList())
+		{
+			AddReplicatedSubObject(NewTask, COND_SkipOwner);
+		}
+
+		return true;
+	}
+
+	return false;
+}
+
+void UGameplayTasksComponent::RemoveSimulatedTask(UGameplayTask* NewTask)
+{
+	if (SimulatedTasks.RemoveSingle(NewTask) > 0)
+	{
+		SetSimulatedTasksNetDirty();
+
+		if (IsUsingRegisteredSubObjectList())
+		{
+			RemoveReplicatedSubObject(NewTask);
+		}
+	}
+}
 
 void UGameplayTasksComponent::SetSimulatedTasks(const TArray<UGameplayTask*>& NewSimulatedTasks)
 {
-	GetSimulatedTasks_Mutable() = NewSimulatedTasks;
+	if (IsUsingRegisteredSubObjectList())
+	{
+		// Unregister all current tasks
+		for (UGameplayTask* OldGameplayTask : SimulatedTasks)
+		{
+			RemoveReplicatedSubObject(OldGameplayTask);
+		}
+
+		SimulatedTasks.Reset(NewSimulatedTasks.Num());
+
+		// Register the new tasks
+		for (UGameplayTask* NewGameplayTask : NewSimulatedTasks)
+		{
+			AddReplicatedSubObject(NewGameplayTask, COND_SkipOwner);
+			SimulatedTasks.Add(NewGameplayTask);
+		}
+	}
+	else
+	{
+		SimulatedTasks = NewSimulatedTasks;
+	}
+
+	SetSimulatedTasksNetDirty();
+}
+
+void UGameplayTasksComponent::SetSimulatedTasksNetDirty()
+{
+	MARK_PROPERTY_DIRTY_FROM_NAME(UGameplayTasksComponent, SimulatedTasks, this);
 }
 
 //----------------------------------------------------------------------//
