Fix deprecation warnings in merged code

#rnx

diff --git a/Engine/Plugins/Runtime/GameplayAbilities/Source/GameplayAbilities/Private/AbilitySystemComponent.cpp b/Engine/Plugins/Runtime/GameplayAbilities/Source/GameplayAbilities/Private/AbilitySystemComponent.cpp
--- a/Engine/Plugins/Runtime/GameplayAbilities/Source/GameplayAbilities/Private/AbilitySystemComponent.cpp
+++ b/Engine/Plugins/Runtime/GameplayAbilities/Source/GameplayAbilities/Private/AbilitySystemComponent.cpp
@@ -2794,7 +2794,10 @@ void UAbilitySystemComponent::OnOwnerActorDestroyed(AActor* InActor)
 
 void UAbilitySystemComponent::SetSpawnedAttributes(const TArray<UAttributeSet*>& NewSpawnedAttributes)
 {
+	PRAGMA_DISABLE_DEPRECATION_WARNINGS
 	TArray<UAttributeSet*>& LocalSpawnedAttributes = GetSpawnedAttributes_Mutable();
+	PRAGMA_ENABLE_DEPRECATION_WARNINGS
+	
 	for (int32 Index = LocalSpawnedAttributes.Num() - 1; Index >= 0; --Index)
 	{
 		UAttributeSet* AttributeSet = LocalSpawnedAttributes[Index];
@@ -2953,7 +2956,9 @@ void UAbilitySystemComponent::RemoveAllReplicatedInstancedAbilities()
 
 void UAbilitySystemComponent::OnSpawnedAttributesEndPlayed(AActor* InActor, EEndPlayReason::Type EndPlayReason)
 {
+	PRAGMA_DISABLE_DEPRECATION_WARNINGS
 	TArray<UAttributeSet*>& LocalSpawnedAttributes = GetSpawnedAttributes_Mutable();
+	PRAGMA_ENABLE_DEPRECATION_WARNINGS
 	for (int32 Index = LocalSpawnedAttributes.Num() - 1; Index >= 0; --Index)
 	{
 		UAttributeSet* AttributeSet = LocalSpawnedAttributes[Index];
