ASC
* Fix references to deprecated function GetSpawnedAttributes_Mutable after merge
* Fix replication of  OwnerActor and AvatarActor when they get nulled after being destroyed.

#rnx

diff --git a/Engine/Plugins/Runtime/GameplayAbilities/Source/GameplayAbilities/Private/AbilitySystemComponent.cpp b/Engine/Plugins/Runtime/GameplayAbilities/Source/GameplayAbilities/Private/AbilitySystemComponent.cpp
--- a/Engine/Plugins/Runtime/GameplayAbilities/Source/GameplayAbilities/Private/AbilitySystemComponent.cpp
+++ b/Engine/Plugins/Runtime/GameplayAbilities/Source/GameplayAbilities/Private/AbilitySystemComponent.cpp
@@ -2781,6 +2781,7 @@ void UAbilitySystemComponent::OnAvatarActorDestroyed(AActor* InActor)
 	if (InActor == AvatarActor)
 	{
 		AvatarActor = nullptr;
+		MARK_PROPERTY_DIRTY_FROM_NAME(UAbilitySystemComponent, AvatarActor, this);
 	}
 }
 
@@ -2789,22 +2790,21 @@ void UAbilitySystemComponent::OnOwnerActorDestroyed(AActor* InActor)
 	if (InActor == OwnerActor)
 	{
 		OwnerActor = nullptr;
+		MARK_PROPERTY_DIRTY_FROM_NAME(UAbilitySystemComponent, OwnerActor, this);
 	}
 }
 
 void UAbilitySystemComponent::SetSpawnedAttributes(const TArray<UAttributeSet*>& NewSpawnedAttributes)
 {
-	PRAGMA_DISABLE_DEPRECATION_WARNINGS
-	TArray<UAttributeSet*>& LocalSpawnedAttributes = GetSpawnedAttributes_Mutable();
-	PRAGMA_ENABLE_DEPRECATION_WARNINGS
-	
-	for (int32 Index = LocalSpawnedAttributes.Num() - 1; Index >= 0; --Index)
+	for (UAttributeSet* AttributeSet : SpawnedAttributes)
 	{
-		UAttributeSet* AttributeSet = LocalSpawnedAttributes[Index];
-		AActor* ActorOwner = AttributeSet->GetTypedOuter<AActor>();
-		if (ActorOwner)
+		if (AttributeSet)
 		{
-			ActorOwner->OnEndPlay.RemoveDynamic(this, &UAbilitySystemComponent::OnSpawnedAttributesEndPlayed);
+			AActor* ActorOwner = AttributeSet->GetTypedOuter<AActor>();
+			if (ActorOwner)
+			{
+				ActorOwner->OnEndPlay.RemoveDynamic(this, &UAbilitySystemComponent::OnSpawnedAttributesEndPlayed);
+			}
 		}
 	}
 
@@ -2956,17 +2956,21 @@ void UAbilitySystemComponent::RemoveAllReplicatedInstancedAbilities()
 
 void UAbilitySystemComponent::OnSpawnedAttributesEndPlayed(AActor* InActor, EEndPlayReason::Type EndPlayReason)
 {
-	PRAGMA_DISABLE_DEPRECATION_WARNINGS
-	TArray<UAttributeSet*>& LocalSpawnedAttributes = GetSpawnedAttributes_Mutable();
-	PRAGMA_ENABLE_DEPRECATION_WARNINGS
-	for (int32 Index = LocalSpawnedAttributes.Num() - 1; Index >= 0; --Index)
+	for (int32 Index = SpawnedAttributes.Num() - 1; Index >= 0; --Index)
 	{
-		UAttributeSet* AttributeSet = LocalSpawnedAttributes[Index];
+		UAttributeSet* AttributeSet = SpawnedAttributes[Index];
 		if (AttributeSet->GetTypedOuter<AActor>() == InActor)
 		{
-			LocalSpawnedAttributes[Index] = nullptr;
+			if (IsUsingRegisteredSubObjectList())
+			{
+				RemoveReplicatedSubObject(AttributeSet);
+			}
+			
+			SpawnedAttributes[Index] = nullptr;
 		}
 	}
+
+	SetSpawnedAttributesListDirty();
 }
 
 void UAbilitySystemComponent::SetClientDebugStrings(TArray<FString>&& NewClientDebugStrings)
