Additional logging to track down sources of too many ability tasks: can now turn on extra logging that will log the source abilities that are creating each ability task to make it easier to figure out where any excess tasks are coming from.

NOTE: To turn on the additional logging, make sure that AbilitySystem.AbilityTask.Debug.RecordingEnabled is true, and then also set AbilitySystem.AbilityTask.Debug.SourceRecordingEnabled to true.

#UE5 #RNX #NoReleaseNotes

diff --git a/Engine/Plugins/Runtime/GameplayAbilities/Source/GameplayAbilities/Private/Abilities/Tasks/AbilityTask.cpp b/Engine/Plugins/Runtime/GameplayAbilities/Source/GameplayAbilities/Private/Abilities/Tasks/AbilityTask.cpp
--- a/Engine/Plugins/Runtime/GameplayAbilities/Source/GameplayAbilities/Private/Abilities/Tasks/AbilityTask.cpp
+++ b/Engine/Plugins/Runtime/GameplayAbilities/Source/GameplayAbilities/Private/Abilities/Tasks/AbilityTask.cpp
@@ -27,6 +27,13 @@ namespace AbilityTaskCVars
 		TEXT("If this is enabled, all new AbilityTasks will be counted by type. Use 'AbilitySystem.AbilityTask.Debug.PrintCounts' to print out the current counts.")
 	);
 
+	static bool bRecordAbilityTaskSourceAbilityCounts = false;
+	static FAutoConsoleVariableRef CVarRecordAbilityTaskSourceAbilityCounts(
+		TEXT("AbilitySystem.AbilityTask.Debug.SourceRecordingEnabled"),
+		bRecordAbilityTaskSourceAbilityCounts,
+		TEXT("Requires bRecordAbilityTaskCounts to be set to true for this value to do anything.  If both are enabled, all new AbilityTasks (after InitTask is called in NewAbilityTask) will be counted by the class of the ability that created them.  Use 'AbilitySystem.AbilityTask.Debug.PrintCounts' to print out the current counts.")
+	);
+
 	static FAutoConsoleCommand AbilityTaskPrintAbilityTaskCountsCmd(
 		TEXT("AbilitySystem.AbilityTask.Debug.PrintCounts"),
 		TEXT("Print out the current AbilityTask counts by class. 'AbilitySystem.AbilityTask.Debug.RecordingEnabled' must be turned on for this to function."),
@@ -212,8 +219,9 @@ bool UAbilityTask::IsWaitingOnAvatar() const
 
 #if !UE_BUILD_SHIPPING
 static TMap<const UClass*, int32> StaticAbilityTasksByClass = {};
+static TMap<const UClass*, int32> StaticAbilityTasksByAbilityClass = {};
 
-static void DebugRecordAbilityTaskCreated(const UAbilityTask* NewTask)
+void DebugRecordAbilityTaskCreated(const UAbilityTask* NewTask)
 {
 	const UClass* ClassPtr = (NewTask != nullptr) ? NewTask->GetClass() : nullptr;
 	if (ClassPtr != nullptr)
@@ -227,7 +235,27 @@ static void DebugRecordAbilityTaskCreated(const UAbilityTask* NewTask)
 			StaticAbilityTasksByClass.Add(ClassPtr, 1);
 		}
 	}
+}
+
+void UAbilityTask::DebugRecordAbilityTaskCreatedByAbility(const UObject* Ability)
+{
+	if (!AbilityTaskCVars::bRecordAbilityTaskSourceAbilityCounts || !AbilityTaskCVars::bRecordAbilityTaskCounts)
+	{	// Both the more detailed and the basic recording is required for the detailed recording to work properly.
+		return;
+	}
 
+	const UClass* ClassPtr = (Ability != nullptr) ? Ability->GetClass() : nullptr;
+	if (ClassPtr != nullptr)
+	{
+		if (StaticAbilityTasksByAbilityClass.Contains(ClassPtr))
+		{
+			StaticAbilityTasksByAbilityClass[ClassPtr]++;
+		}
+		else
+		{
+			StaticAbilityTasksByAbilityClass.Add(ClassPtr, 1);
+		}
+	}
 }
 
 static void DebugRecordAbilityTaskDestroyed(const UAbilityTask* DestroyedTask)
@@ -235,6 +263,23 @@ static void DebugRecordAbilityTaskDestroyed(const UAbilityTask* DestroyedTask)
 	const UClass* ClassPtr = (DestroyedTask != nullptr) ? DestroyedTask->GetClass() : nullptr;
 	if (ClassPtr != nullptr)
 	{
+		if (AbilityTaskCVars::bRecordAbilityTaskSourceAbilityCounts)
+		{
+			const UClass* AbilityClassPtr = (DestroyedTask->Ability != nullptr) ? DestroyedTask->Ability->GetClass() : nullptr;
+			if (AbilityClassPtr != nullptr)
+			{
+				if (StaticAbilityTasksByAbilityClass.Contains(AbilityClassPtr))
+				{
+					StaticAbilityTasksByAbilityClass[AbilityClassPtr]--;
+
+					if (StaticAbilityTasksByAbilityClass[AbilityClassPtr] <= 0)
+					{
+						StaticAbilityTasksByAbilityClass.Remove(AbilityClassPtr);
+					}
+				}
+			}
+		}
+
 		if (StaticAbilityTasksByClass.Contains(ClassPtr))
 		{
 			StaticAbilityTasksByClass[ClassPtr]--;
@@ -269,6 +314,17 @@ static void DebugPrintAbilityTasksByClass()
 			ABILITY_LOG(Display, TEXT("- Unknown (allocated before recording): %d"), UnaccountedAbilityTasks);
 		}
 
+		if (AbilityTaskCVars::bRecordAbilityTaskSourceAbilityCounts)
+		{
+			ABILITY_LOG(Display, TEXT("UAbilityTask counts per Ability Class:"));
+			StaticAbilityTasksByAbilityClass.ValueSort(TGreater<int32>());
+			for (const TPair<const UClass*, int32>& Pair : StaticAbilityTasksByAbilityClass)
+			{
+				FString SafeName = GetNameSafe(Pair.Key);
+				ABILITY_LOG(Display, TEXT("- Ability Class '%s': %d"), *SafeName, Pair.Value);
+			}
+		}
+
 		ABILITY_LOG(Display, TEXT("Total AbilityTask count: %d"), GlobalAbilityTaskCount);
 	}
 	else
diff --git a/Engine/Plugins/Runtime/GameplayAbilities/Source/GameplayAbilities/Public/Abilities/Tasks/AbilityTask.h b/Engine/Plugins/Runtime/GameplayAbilities/Source/GameplayAbilities/Public/Abilities/Tasks/AbilityTask.h
--- a/Engine/Plugins/Runtime/GameplayAbilities/Source/GameplayAbilities/Public/Abilities/Tasks/AbilityTask.h
+++ b/Engine/Plugins/Runtime/GameplayAbilities/Source/GameplayAbilities/Public/Abilities/Tasks/AbilityTask.h
@@ -121,6 +121,10 @@ class GAMEPLAYABILITIES_API UAbilityTask : public UGameplayTask
 
 	virtual void InitSimulatedTask(UGameplayTasksComponent& InGameplayTasksComponent) override;
 
+#if !UE_BUILD_SHIPPING
+	static void DebugRecordAbilityTaskCreatedByAbility(const UObject* Ability);
+#endif  // !UE_BUILD_SHIPPING
+
 	/** Helper function for instantiating and initializing a new task */
 	template <class T>
 	static T* NewAbilityTask(UGameplayAbility* ThisAbility, FName InstanceName = FName())
@@ -129,6 +133,11 @@ class GAMEPLAYABILITIES_API UAbilityTask : public UGameplayTask
 
 		T* MyObj = NewObject<T>();
 		MyObj->InitTask(*ThisAbility, ThisAbility->GetGameplayTaskDefaultPriority());
+
+#if !UE_BUILD_SHIPPING
+		UAbilityTask::DebugRecordAbilityTaskCreatedByAbility(ThisAbility);
+#endif  // !UE_BUILD_SHIPPING
+
 		MyObj->InstanceName = InstanceName;
 		return MyObj;
 	}
