Fix for ensures when adding duplicate OnEndPlayDelegates

diff --git a/Engine/Plugins/Runtime/GameplayAbilities/Source/GameplayAbilities/Private/AbilitySystemComponent.cpp b/Engine/Plugins/Runtime/GameplayAbilities/Source/GameplayAbilities/Private/AbilitySystemComponent.cpp
--- a/Engine/Plugins/Runtime/GameplayAbilities/Source/GameplayAbilities/Private/AbilitySystemComponent.cpp
+++ b/Engine/Plugins/Runtime/GameplayAbilities/Source/GameplayAbilities/Private/AbilitySystemComponent.cpp
@@ -2758,7 +2758,7 @@ void UAbilitySystemComponent::SetOwnerActor(AActor* NewOwnerActor)
 	OwnerActor = NewOwnerActor;
 	if (OwnerActor)
 	{
-		OwnerActor->OnDestroyed.AddDynamic(this, &UAbilitySystemComponent::OnOwnerActorDestroyed);
+		OwnerActor->OnDestroyed.AddUniqueDynamic(this, &UAbilitySystemComponent::OnOwnerActorDestroyed);
 	}
 }
 
@@ -2772,7 +2772,7 @@ void UAbilitySystemComponent::SetAvatarActor_Direct(AActor* NewAvatarActor)
 	AvatarActor = NewAvatarActor;
 	if (AvatarActor)
 	{
-		AvatarActor->OnDestroyed.AddDynamic(this, &UAbilitySystemComponent::OnAvatarActorDestroyed);
+		AvatarActor->OnDestroyed.AddUniqueDynamic(this, &UAbilitySystemComponent::OnAvatarActorDestroyed);
 	}
 }
 
