Fixing SA warning in AbilityTask_PlayMontageAndWait

diff --git a/Engine/Plugins/Runtime/GameplayAbilities/Source/GameplayAbilities/Private/Abilities/Tasks/AbilityTask_PlayMontageAndWait.cpp b/Engine/Plugins/Runtime/GameplayAbilities/Source/GameplayAbilities/Private/Abilities/Tasks/AbilityTask_PlayMontageAndWait.cpp
--- a/Engine/Plugins/Runtime/GameplayAbilities/Source/GameplayAbilities/Private/Abilities/Tasks/AbilityTask_PlayMontageAndWait.cpp
+++ b/Engine/Plugins/Runtime/GameplayAbilities/Source/GameplayAbilities/Private/Abilities/Tasks/AbilityTask_PlayMontageAndWait.cpp
@@ -184,8 +184,13 @@ void UAbilityTask_PlayMontageAndWait::OnDestroy(bool AbilityEnded)
 
 bool UAbilityTask_PlayMontageAndWait::StopPlayingMontage()
 {
+	if (Ability == nullptr)
+	{
+		return false;
+	}
+
 	const FGameplayAbilityActorInfo* ActorInfo = Ability->GetCurrentActorInfo();
-	if (!ActorInfo)
+	if (ActorInfo == nullptr)
 	{
 		return false;
 	}
