Fix non-unity errors
#rnx

diff --git a/Engine/Plugins/Runtime/GameplayAbilities/Source/GameplayAbilities/Private/ScalableFloat.cpp b/Engine/Plugins/Runtime/GameplayAbilities/Source/GameplayAbilities/Private/ScalableFloat.cpp
--- a/Engine/Plugins/Runtime/GameplayAbilities/Source/GameplayAbilities/Private/ScalableFloat.cpp
+++ b/Engine/Plugins/Runtime/GameplayAbilities/Source/GameplayAbilities/Private/ScalableFloat.cpp
@@ -9,6 +9,7 @@
 #include "AbilitySystemLog.h"
 #include "UObject/UObjectIterator.h"
 #include "DataRegistrySubsystem.h"
+#include "HAL/IConsoleManager.h"
 
 #if WITH_EDITOR
 #include "EditorReimportHandler.h"
diff --git a/Engine/Source/Developer/FunctionalTesting/Classes/TraceQueryTestResults.h b/Engine/Source/Developer/FunctionalTesting/Classes/TraceQueryTestResults.h
--- a/Engine/Source/Developer/FunctionalTesting/Classes/TraceQueryTestResults.h
+++ b/Engine/Source/Developer/FunctionalTesting/Classes/TraceQueryTestResults.h
@@ -6,6 +6,7 @@
 #include "UObject/ObjectMacros.h"
 #include "UObject/UObjectGlobals.h"
 #include "Engine/EngineTypes.h"
+#include "Engine/HitResult.h"
 #include "TraceQueryTestResults.generated.h"
 
 class AFunctionalTest;
diff --git a/Engine/Source/Runtime/Engine/Classes/Components/InterpToMovementComponent.h b/Engine/Source/Runtime/Engine/Classes/Components/InterpToMovementComponent.h
--- a/Engine/Source/Runtime/Engine/Classes/Components/InterpToMovementComponent.h
+++ b/Engine/Source/Runtime/Engine/Classes/Components/InterpToMovementComponent.h
@@ -5,6 +5,7 @@
 #include "CoreMinimal.h"
 #include "UObject/ObjectMacros.h"
 #include "GameFramework/MovementComponent.h"
+#include "Engine/HitResult.h"
 #include "InterpToMovementComponent.generated.h"
 
 /**
diff --git a/Engine/Source/Runtime/Engine/Classes/Engine/MemberReference.h b/Engine/Source/Runtime/Engine/Classes/Engine/MemberReference.h
--- a/Engine/Source/Runtime/Engine/Classes/Engine/MemberReference.h
+++ b/Engine/Source/Runtime/Engine/Classes/Engine/MemberReference.h
@@ -11,6 +11,7 @@
 #include "UObject/Package.h"
 #include "UObject/UObjectIterator.h"
 #include "Engine/Blueprint.h"
+#include "EngineLogs.h"
 #include "MemberReference.generated.h"
 
 /** Helper struct to allow us to redirect properties and functions through renames and additionally between classes if necessary */
diff --git a/Engine/Source/Runtime/Engine/Classes/GameFramework/ProjectileMovementComponent.h b/Engine/Source/Runtime/Engine/Classes/GameFramework/ProjectileMovementComponent.h
--- a/Engine/Source/Runtime/Engine/Classes/GameFramework/ProjectileMovementComponent.h
+++ b/Engine/Source/Runtime/Engine/Classes/GameFramework/ProjectileMovementComponent.h
@@ -5,6 +5,7 @@
 #include "CoreMinimal.h"
 #include "UObject/ObjectMacros.h"
 #include "GameFramework/MovementComponent.h"
+#include "Engine/HitResult.h"
 #include "ProjectileMovementComponent.generated.h"
 
 /**
diff --git a/Engine/Source/Runtime/Engine/Classes/Kismet/GameplayStaticsTypes.h b/Engine/Source/Runtime/Engine/Classes/Kismet/GameplayStaticsTypes.h
--- a/Engine/Source/Runtime/Engine/Classes/Kismet/GameplayStaticsTypes.h
+++ b/Engine/Source/Runtime/Engine/Classes/Kismet/GameplayStaticsTypes.h
@@ -2,6 +2,7 @@
 
 #pragma once
 #include "KismetSystemLibrary.h"
+#include "Engine/HitResult.h"
 #include "GameplayStaticsTypes.generated.h"
 
 // Trace option to SuggestProjectileVelocity functions.
diff --git a/Engine/Source/Runtime/Engine/Private/TimelineTemplate.cpp b/Engine/Source/Runtime/Engine/Private/TimelineTemplate.cpp
--- a/Engine/Source/Runtime/Engine/Private/TimelineTemplate.cpp
+++ b/Engine/Source/Runtime/Engine/Private/TimelineTemplate.cpp
@@ -4,6 +4,7 @@
 #include "UObject/FortniteMainBranchObjectVersion.h"
 #include "UObject/Package.h"
 #include "Engine/Blueprint.h"
+#include "EngineLogs.h"
 #include "Curves/CurveFloat.h"
 #include "Curves/CurveLinearColor.h"
 #include "Curves/CurveVector.h"
