Add a new version of GetGameplayEffectCount to AbilitySystemComponent that takes a soft reference - to avoid creating unnecessary hard references between assets.

diff --git a/Engine/Plugins/Runtime/GameplayAbilities/Source/GameplayAbilities/Private/AbilitySystemComponent.cpp b/Engine/Plugins/Runtime/GameplayAbilities/Source/GameplayAbilities/Private/AbilitySystemComponent.cpp
--- a/Engine/Plugins/Runtime/GameplayAbilities/Source/GameplayAbilities/Private/AbilitySystemComponent.cpp
+++ b/Engine/Plugins/Runtime/GameplayAbilities/Source/GameplayAbilities/Private/AbilitySystemComponent.cpp
@@ -392,7 +392,7 @@ FGameplayEffectContextHandle UAbilitySystemComponent::MakeEffectContext() const
 	return Context;
 }
 
-int32 UAbilitySystemComponent::GetGameplayEffectCount(TSubclassOf<UGameplayEffect> SourceGameplayEffect, UAbilitySystemComponent* OptionalInstigatorFilterComponent, bool bEnforceOnGoingCheck)
+int32 UAbilitySystemComponent::GetGameplayEffectCount(TSubclassOf<UGameplayEffect> SourceGameplayEffect, UAbilitySystemComponent* OptionalInstigatorFilterComponent, bool bEnforceOnGoingCheck) const
 {
 	int32 Count = 0;
 
@@ -426,6 +426,19 @@ int32 UAbilitySystemComponent::GetGameplayEffectCount(TSubclassOf<UGameplayEffec
 	return Count;
 }
 
+int32 UAbilitySystemComponent::GetGameplayEffectCount_IfLoaded(TSoftClassPtr<UGameplayEffect> SoftSourceGameplayEffect, UAbilitySystemComponent* OptionalInstigatorFilterComponent, bool bEnforceOnGoingCheck) const
+{
+	TSubclassOf<UGameplayEffect> SourceGameplayEffect = SoftSourceGameplayEffect.Get();
+
+	//if the gameplay effect is not loaded, then there must be none active
+	if (!SourceGameplayEffect)
+	{
+		return 0;
+	}
+
+	return GetGameplayEffectCount(SourceGameplayEffect, OptionalInstigatorFilterComponent, bEnforceOnGoingCheck);
+}
+
 int32 UAbilitySystemComponent::GetAggregatedStackCount(const FGameplayEffectQuery& Query) const
 {
 	return ActiveGameplayEffects.GetActiveEffectCount(Query);
diff --git a/Engine/Plugins/Runtime/GameplayAbilities/Source/GameplayAbilities/Public/AbilitySystemComponent.h b/Engine/Plugins/Runtime/GameplayAbilities/Source/GameplayAbilities/Public/AbilitySystemComponent.h
--- a/Engine/Plugins/Runtime/GameplayAbilities/Source/GameplayAbilities/Public/AbilitySystemComponent.h
+++ b/Engine/Plugins/Runtime/GameplayAbilities/Source/GameplayAbilities/Public/AbilitySystemComponent.h
@@ -345,7 +345,19 @@ class GAMEPLAYABILITIES_API UAbilitySystemComponent : public UGameplayTasksCompo
 	 * @return Count of the specified source effect
 	 */
 	UFUNCTION(BlueprintCallable, BlueprintPure, Category=GameplayEffects)
-	int32 GetGameplayEffectCount(TSubclassOf<UGameplayEffect> SourceGameplayEffect, UAbilitySystemComponent* OptionalInstigatorFilterComponent, bool bEnforceOnGoingCheck = true);
+	int32 GetGameplayEffectCount(TSubclassOf<UGameplayEffect> SourceGameplayEffect, UAbilitySystemComponent* OptionalInstigatorFilterComponent, bool bEnforceOnGoingCheck = true) const;
+
+	/**
+	 * Get the count of the specified source effect on the ability system component. For non-stacking effects, this is the sum of all active instances.
+	 * For stacking effects, this is the sum of all valid stack counts. If an instigator is specified, only effects from that instigator are counted.
+	 * 
+	 * @param SoftSourceGameplayEffect				Effect to get the count of. If this is not currently loaded, the count is 0
+	 * @param OptionalInstigatorFilterComponent		If specified, only count effects applied by this ability system component
+	 * 
+	 * @return Count of the specified source effect
+	 */
+	UFUNCTION(BlueprintCallable, BlueprintPure, Category = GameplayEffects)
+	int32 GetGameplayEffectCount_IfLoaded(TSoftClassPtr<UGameplayEffect> SoftSourceGameplayEffect, UAbilitySystemComponent* OptionalInstigatorFilterComponent, bool bEnforceOnGoingCheck = true) const;
 
 	/** Returns the sum of StackCount of all gameplay effects that pass query */
 	int32 GetAggregatedStackCount(const FGameplayEffectQuery& Query) const;
