Adding BP library function to get a debug string representation of a gameplay attribute. Partial adaptation from PR #7444

diff --git a/Engine/Plugins/Runtime/GameplayAbilities/Source/GameplayAbilities/Private/AbilitySystemBlueprintLibrary.cpp b/Engine/Plugins/Runtime/GameplayAbilities/Source/GameplayAbilities/Private/AbilitySystemBlueprintLibrary.cpp
--- a/Engine/Plugins/Runtime/GameplayAbilities/Source/GameplayAbilities/Private/AbilitySystemBlueprintLibrary.cpp
+++ b/Engine/Plugins/Runtime/GameplayAbilities/Source/GameplayAbilities/Private/AbilitySystemBlueprintLibrary.cpp
@@ -139,6 +139,16 @@ bool UAbilitySystemBlueprintLibrary::NotEqual_GameplayAttributeGameplayAttribute
 	return (AttributeA != AttributeB);
 }
 
+FString UAbilitySystemBlueprintLibrary::GetDebugStringFromGameplayAttribute(const FGameplayAttribute& Attribute)
+{
+	if (const UClass* AttributeSetClass = Attribute.GetAttributeSetClass())
+	{
+		return FString::Format(TEXT("{0}.{1}"), { AttributeSetClass->GetName(), Attribute.GetName() });
+	}
+
+	return Attribute.GetName();
+}
+
 FGameplayAbilityTargetDataHandle UAbilitySystemBlueprintLibrary::AppendTargetDataHandle(FGameplayAbilityTargetDataHandle TargetHandle, const FGameplayAbilityTargetDataHandle& HandleToAdd)
 {
 	TargetHandle.Append(HandleToAdd);
diff --git a/Engine/Plugins/Runtime/GameplayAbilities/Source/GameplayAbilities/Public/AbilitySystemBlueprintLibrary.h b/Engine/Plugins/Runtime/GameplayAbilities/Source/GameplayAbilities/Public/AbilitySystemBlueprintLibrary.h
--- a/Engine/Plugins/Runtime/GameplayAbilities/Source/GameplayAbilities/Public/AbilitySystemBlueprintLibrary.h
+++ b/Engine/Plugins/Runtime/GameplayAbilities/Source/GameplayAbilities/Public/AbilitySystemBlueprintLibrary.h
@@ -78,6 +78,10 @@ class GAMEPLAYABILITIES_API UAbilitySystemBlueprintLibrary : public UBlueprintFu
 	UFUNCTION(BlueprintPure, meta=(DisplayName = "Not Equal (Gameplay Attribute)", CompactNodeTitle = "!=", Keywords = "!= not equal"), Category="Ability|Attribute")
 	static bool NotEqual_GameplayAttributeGameplayAttribute(FGameplayAttribute AttributeA, FGameplayAttribute AttributeB);
 
+	/** Returns FString representation of a gameplay attribute's set class and name, in the form of AttrSetName.AttrName (or just AttrName if not part of a set).  */
+	UFUNCTION(BlueprintPure, Category="Ability|Attribute")
+	static FString GetDebugStringFromGameplayAttribute(const FGameplayAttribute& Attribute);
+
 	// -------------------------------------------------------------------------------
 	//		TargetData
 	// -------------------------------------------------------------------------------
