Added support for Gameplay Effects to add blocked ability tags.

diff --git a/Engine/Plugins/Runtime/GameplayAbilities/Source/GameplayAbilities/Private/GameplayEffect.cpp b/Engine/Plugins/Runtime/GameplayAbilities/Source/GameplayAbilities/Private/GameplayEffect.cpp
--- a/Engine/Plugins/Runtime/GameplayAbilities/Source/GameplayAbilities/Private/GameplayEffect.cpp
+++ b/Engine/Plugins/Runtime/GameplayAbilities/Source/GameplayAbilities/Private/GameplayEffect.cpp
@@ -69,6 +69,11 @@ void UGameplayEffect::GetOwnedGameplayTags(FGameplayTagContainer& TagContainer)
 	TagContainer.AppendTags(InheritableOwnedTagsContainer.CombinedTags);
 }
 
+void UGameplayEffect::GetBlockedAbilityTags(FGameplayTagContainer& OutTagContainer) const
+{
+	OutTagContainer.AppendTags(InheritableBlockedAbilityTagsContainer.CombinedTags);
+}
+
 void UGameplayEffect::PostLoad()
 {
 	Super::PostLoad();
@@ -131,6 +136,7 @@ void UGameplayEffect::PostInitProperties()
 
 	InheritableGameplayEffectTags.PostInitProperties();
 	InheritableOwnedTagsContainer.PostInitProperties();
+	InheritableBlockedAbilityTagsContainer.PostInitProperties();
 	RemoveGameplayEffectsWithTags.PostInitProperties();
 }
 
@@ -153,6 +159,10 @@ void UGameplayEffect::PostEditChangeProperty(FPropertyChangedEvent& PropertyChan
 		{
 			InheritableOwnedTagsContainer.UpdateInheritedTagProperties(Parent ? &Parent->InheritableOwnedTagsContainer : NULL);
 		}
+		else if (PropName == GET_MEMBER_NAME_CHECKED(UGameplayEffect, InheritableBlockedAbilityTagsContainer))
+		{
+			InheritableBlockedAbilityTagsContainer.UpdateInheritedTagProperties(Parent ? &Parent->InheritableBlockedAbilityTagsContainer : nullptr);
+		}
 		else if (PropName == GET_MEMBER_NAME_CHECKED(UGameplayEffect, RemoveGameplayEffectsWithTags))
 		{
 			RemoveGameplayEffectsWithTags.UpdateInheritedTagProperties(Parent ? &Parent->RemoveGameplayEffectsWithTags : NULL);
@@ -187,6 +197,7 @@ void UGameplayEffect::UpdateInheritedTagProperties()
 
 	InheritableGameplayEffectTags.UpdateInheritedTagProperties(Parent ? &Parent->InheritableGameplayEffectTags : NULL);
 	InheritableOwnedTagsContainer.UpdateInheritedTagProperties(Parent ? &Parent->InheritableOwnedTagsContainer : NULL);
+	InheritableBlockedAbilityTagsContainer.UpdateInheritedTagProperties(Parent ? &Parent->InheritableBlockedAbilityTagsContainer : nullptr);
 	RemoveGameplayEffectsWithTags.UpdateInheritedTagProperties(Parent ? &Parent->RemoveGameplayEffectsWithTags : NULL);
 }
 
@@ -1233,6 +1244,14 @@ void FGameplayEffectSpec::GetAllGrantedTags(OUT FGameplayTagContainer& Container
 	}
 }
 
+void FGameplayEffectSpec::GetAllBlockedAbilityTags(FGameplayTagContainer& OutContainer) const
+{
+	if (Def)
+	{
+		OutContainer.AppendTags(Def->InheritableBlockedAbilityTagsContainer.CombinedTags);
+	}
+}
+
 void FGameplayEffectSpec::GetAllAssetTags(OUT FGameplayTagContainer& Container) const
 {
 	Container.AppendTags(GetDynamicAssetTags());
@@ -3245,6 +3264,11 @@ void FActiveGameplayEffectsContainer::AddActiveGameplayEffectGrantedTagsAndModif
 	// Update our owner with the tags this GameplayEffect grants them
 	Owner->UpdateTagMap(Effect.Spec.Def->InheritableOwnedTagsContainer.CombinedTags, 1);
 	Owner->UpdateTagMap(Effect.Spec.DynamicGrantedTags, 1);
+
+	// Update our owner with the blocked ability tags this GameplayEffect adds to them
+	Owner->BlockAbilitiesWithTags(Effect.Spec.Def->InheritableBlockedAbilityTagsContainer.CombinedTags);
+
+	// Update minimal replication if needed.
 	if (ShouldUseMinimalReplication())
 	{
 		Owner->AddMinimalReplicationGameplayTags(Effect.Spec.Def->InheritableOwnedTagsContainer.CombinedTags);
@@ -3541,6 +3565,10 @@ void FActiveGameplayEffectsContainer::RemoveActiveGameplayEffectGrantedTagsAndMo
 	Owner->UpdateTagMap(Effect.Spec.Def->InheritableOwnedTagsContainer.CombinedTags, -1);
 	Owner->UpdateTagMap(Effect.Spec.DynamicGrantedTags, -1);
 
+	// Update our owner with the blocked ability tags this GameplayEffect adds to them
+	Owner->UnBlockAbilitiesWithTags(Effect.Spec.Def->InheritableBlockedAbilityTagsContainer.CombinedTags);
+
+	// Update minimal replication if needed.
 	if (ShouldUseMinimalReplication())
 	{
 		Owner->RemoveMinimalReplicationGameplayTags(Effect.Spec.Def->InheritableOwnedTagsContainer.CombinedTags);
diff --git a/Engine/Plugins/Runtime/GameplayAbilities/Source/GameplayAbilities/Public/GameplayEffect.h b/Engine/Plugins/Runtime/GameplayAbilities/Source/GameplayAbilities/Public/GameplayEffect.h
--- a/Engine/Plugins/Runtime/GameplayAbilities/Source/GameplayAbilities/Public/GameplayEffect.h
+++ b/Engine/Plugins/Runtime/GameplayAbilities/Source/GameplayAbilities/Public/GameplayEffect.h
@@ -1012,6 +1012,9 @@ struct GAMEPLAYABILITIES_API FGameplayEffectSpec
 	/** Appends all tags granted by this gameplay effect spec */
 	void GetAllGrantedTags(OUT FGameplayTagContainer& Container) const;
 
+	/** Appends all blocked ability tags granted by this gameplay effect spec */
+	void GetAllBlockedAbilityTags(FGameplayTagContainer& OutContainer) const;
+
 	/** Appends all tags that apply to this gameplay effect spec */
 	void GetAllAssetTags(OUT FGameplayTagContainer& Container) const;
 
@@ -1937,6 +1940,7 @@ public:
 	static const float INVALID_LEVEL;
 
 	virtual void GetOwnedGameplayTags(FGameplayTagContainer& TagContainer) const override;
+	virtual void GetBlockedAbilityTags(FGameplayTagContainer& OutTagContainer) const;
 	virtual void PostInitProperties() override;
 	virtual void PostLoad() override;
 	PRAGMA_DISABLE_DEPRECATION_WARNINGS // Suppress compiler warning on override of deprecated function
@@ -2041,10 +2045,14 @@ public:
 	UPROPERTY(EditDefaultsOnly, BlueprintReadOnly, Category = Tags, meta = (DisplayName = "GameplayEffectAssetTag", Categories="GameplayEffectTagsCategory"))
 	FInheritedTagContainer InheritableGameplayEffectTags;
 	
-	/** "These tags are applied to the actor I am applied to" */
+	/** These tags are applied to the actor I am applied to */
 	UPROPERTY(EditDefaultsOnly, BlueprintReadOnly, Category = Tags, meta=(DisplayName="GrantedTags", Categories="OwnedTagsCategory"))
 	FInheritedTagContainer InheritableOwnedTagsContainer;
 	
+	/** These blocked ability tags are applied to the actor I am applied to */
+	UPROPERTY(EditDefaultsOnly, BlueprintReadOnly, Category = Tags, meta=(DisplayName="GrantedBlockedAbilityTags", Categories="BlockedAbilityTagsCategory"))
+	FInheritedTagContainer InheritableBlockedAbilityTagsContainer;
+	
 	/** Once Applied, these tags requirements are used to determined if the GameplayEffect is "on" or "off". A GameplayEffect can be off and do nothing, but still applied. */
 	UPROPERTY(EditDefaultsOnly, BlueprintReadOnly, Category = Tags, meta=(Categories="OngoingTagRequirementsCategory"))
 	FGameplayTagRequirements OngoingTagRequirements;
