Static Analysis Compile fix: Remove unnecessary test for TargetActor != null, since it is already tested in an outer if condition in the current scope on line 417. The presence of the unnecessary test confuses static analysis which then compilains about a missing test in the usage on 469.
#rnx

diff --git a/Engine/Plugins/Runtime/GameplayAbilities/Source/GameplayAbilities/Private/GameplayCueManager.cpp b/Engine/Plugins/Runtime/GameplayAbilities/Source/GameplayAbilities/Private/GameplayCueManager.cpp
--- a/Engine/Plugins/Runtime/GameplayAbilities/Source/GameplayAbilities/Private/GameplayCueManager.cpp
+++ b/Engine/Plugins/Runtime/GameplayAbilities/Source/GameplayAbilities/Private/GameplayCueManager.cpp
@@ -421,7 +421,7 @@ AGameplayCueNotify_Actor* UGameplayCueManager::GetInstancedCueActor(AActor* Targ
 		
 #if WITH_EDITOR	
 		// Animtion preview hack. If we are trying to play the GC on a CDO, then don't use actor recycling and don't set the owner (to the CDO, which would cause problems)
-		if (TargetActor && TargetActor->HasAnyFlags(RF_ClassDefaultObject))
+		if (TargetActor->HasAnyFlags(RF_ClassDefaultObject))
 		{
 			NewOwnerActor = nullptr;
 			UseActorRecycling = false;
