Added a method to allow ability system components to perform a predictive effect execution. Requires an instant or periodic effect and a valid prediction key. Has an option to execute the gameplay cues or not.

#rnx

diff --git a/Engine/Plugins/Runtime/GameplayAbilities/Source/GameplayAbilities/Private/GameplayEffect.cpp b/Engine/Plugins/Runtime/GameplayAbilities/Source/GameplayAbilities/Private/GameplayEffect.cpp
--- a/Engine/Plugins/Runtime/GameplayAbilities/Source/GameplayAbilities/Private/GameplayEffect.cpp
+++ b/Engine/Plugins/Runtime/GameplayAbilities/Source/GameplayAbilities/Private/GameplayEffect.cpp
@@ -1989,6 +1989,130 @@ void FActiveGameplayEffectsContainer::RegisterWithOwner(UAbilitySystemComponent*
 	}
 }
 
+void FActiveGameplayEffectsContainer::PredictivelyExecuteEffectSpec(FGameplayEffectSpec& Spec, FPredictionKey PredictionKey, const bool bPredictGameplayCues /*= false*/)
+{
+	if (!Owner)
+	{
+		return;
+	}
+
+	if (!PredictionKey.IsValidForMorePrediction())
+	{
+		return;
+	}
+
+	// Should only ever execute effects that are instant application or periodic application
+	// Effects with no period and that aren't instant application should never be executed
+	const bool bNotInstantEffect = (Spec.GetDuration() > UGameplayEffect::INSTANT_APPLICATION);
+	const bool bNoPeriodEffect = (Spec.GetPeriod() != UGameplayEffect::NO_PERIOD);
+	if (bNotInstantEffect && bNoPeriodEffect)
+	{
+		return;
+	}
+
+	FGameplayEffectSpec& SpecToUse = Spec;
+
+	// Capture our own tags.
+	// TODO: We should only capture them if we need to. We may have snapshotted target tags (?) (in the case of dots with exotic setups?)
+
+	SpecToUse.CapturedTargetTags.GetActorTags().Reset();
+	Owner->GetOwnedGameplayTags(SpecToUse.CapturedTargetTags.GetActorTags());
+
+	SpecToUse.CalculateModifierMagnitudes();
+
+	// ------------------------------------------------------
+	//	Modifiers
+	//		These will modify the base value of attributes
+	// ------------------------------------------------------
+
+	bool ModifierSuccessfullyExecuted = false;
+
+	for (int32 ModIdx = 0; ModIdx < SpecToUse.Modifiers.Num(); ++ModIdx)
+	{
+		const FGameplayModifierInfo& ModDef = SpecToUse.Def->Modifiers[ModIdx];
+
+		FGameplayModifierEvaluatedData EvalData(ModDef.Attribute, ModDef.ModifierOp, SpecToUse.GetModifierMagnitude(ModIdx, true));
+		ModifierSuccessfullyExecuted |= InternalExecuteMod(SpecToUse, EvalData);
+	}
+
+	// ------------------------------------------------------
+	//	Executions
+	//		This will run custom code to 'do stuff'
+	// ------------------------------------------------------
+
+	bool GameplayCuesWereManuallyHandled = false;
+
+	for (const FGameplayEffectExecutionDefinition& CurExecDef : SpecToUse.Def->Executions)
+	{
+		if (CurExecDef.CalculationClass)
+		{
+			const UGameplayEffectExecutionCalculation* ExecCDO = CurExecDef.CalculationClass->GetDefaultObject<UGameplayEffectExecutionCalculation>();
+			check(ExecCDO);
+
+			// Run the custom execution
+			FGameplayEffectCustomExecutionParameters ExecutionParams(SpecToUse, CurExecDef.CalculationModifiers, Owner, CurExecDef.PassedInTags, PredictionKey);
+			FGameplayEffectCustomExecutionOutput ExecutionOutput;
+			ExecCDO->Execute(ExecutionParams, ExecutionOutput);
+
+			// Execute any mods the custom execution yielded
+			TArray<FGameplayModifierEvaluatedData>& OutModifiers = ExecutionOutput.GetOutputModifiersRef();
+
+			const bool bApplyStackCountToEmittedMods = !ExecutionOutput.IsStackCountHandledManually();
+			const int32 SpecStackCount = SpecToUse.StackCount;
+
+			for (FGameplayModifierEvaluatedData& CurExecMod : OutModifiers)
+			{
+				// If the execution didn't manually handle the stack count, automatically apply it here
+				if (bApplyStackCountToEmittedMods && SpecStackCount > 1)
+				{
+					CurExecMod.Magnitude = GameplayEffectUtilities::ComputeStackedModifierMagnitude(CurExecMod.Magnitude, SpecStackCount, CurExecMod.ModifierOp);
+				}
+				ModifierSuccessfullyExecuted |= InternalExecuteMod(SpecToUse, CurExecMod);
+			}
+
+			// If execution handled GameplayCues, we dont have to.
+			if (ExecutionOutput.AreGameplayCuesHandledManually())
+			{
+				GameplayCuesWereManuallyHandled = true;
+			}
+		}
+	}
+
+	// ------------------------------------------------------
+	//	Invoke GameplayCue events
+	// ------------------------------------------------------
+	if (bPredictGameplayCues)
+	{
+		// If there are no modifiers or we don't require modifier success to trigger, we apply the GameplayCue.
+		const bool bHasModifiers = SpecToUse.Modifiers.Num() > 0;
+		const bool bHasExecutions = SpecToUse.Def->Executions.Num() > 0;
+		const bool bHasModifiersOrExecutions = bHasModifiers || bHasExecutions;
+
+		// If there are no modifiers or we don't require modifier success to trigger, we apply the GameplayCue.
+		bool InvokeGameplayCueExecute = (!bHasModifiersOrExecutions) || !Spec.Def->bRequireModifierSuccessToTriggerCues;
+
+		if (bHasModifiersOrExecutions && ModifierSuccessfullyExecuted)
+		{
+			InvokeGameplayCueExecute = true;
+		}
+
+		// Don't trigger gameplay cues if one of the executions says it manually handled them
+		if (GameplayCuesWereManuallyHandled)
+		{
+			InvokeGameplayCueExecute = false;
+		}
+
+		if (InvokeGameplayCueExecute && SpecToUse.Def->GameplayCues.Num())
+		{
+			// TODO: check replication policy. Right now we will replicate every execute via a multicast RPC
+
+			ABILITY_LOG(Log, TEXT("Invoking Execute GameplayCue for %s"), *SpecToUse.ToSimpleString());
+
+			UAbilitySystemGlobals::Get().GetGameplayCueManager()->InvokeGameplayCueExecuted_FromSpec(Owner, SpecToUse, PredictionKey);
+		}
+	}
+}
+
 /** This is the main function that executes a GameplayEffect on Attributes and ActiveGameplayEffects */
 void FActiveGameplayEffectsContainer::ExecuteActiveEffectsFrom(FGameplayEffectSpec &Spec, FPredictionKey PredictionKey)
 {
diff --git a/Engine/Plugins/Runtime/GameplayAbilities/Source/GameplayAbilities/Public/GameplayEffect.h b/Engine/Plugins/Runtime/GameplayAbilities/Source/GameplayAbilities/Public/GameplayEffect.h
--- a/Engine/Plugins/Runtime/GameplayAbilities/Source/GameplayAbilities/Public/GameplayEffect.h
+++ b/Engine/Plugins/Runtime/GameplayAbilities/Source/GameplayAbilities/Public/GameplayEffect.h
@@ -1578,6 +1578,13 @@ struct GAMEPLAYABILITIES_API FActiveGameplayEffectsContainer : public FFastArray
 
 	const FActiveGameplayEffect* GetActiveGameplayEffect(const FActiveGameplayEffectHandle Handle) const;
 
+	/** Predictively execute a given effect spec. Any attribute modifications and effect execution calculations in the effect will run and then if desired predict gameplay cues
+		@note: This method will not predictively run any conditional effects that may be set up in the effect that apply post execution and will only happen if/when this spec is
+		applied on the server. 
+		
+		@note: WARNING: This will locally perform attribute changes on your client so beware. */
+	void PredictivelyExecuteEffectSpec(FGameplayEffectSpec& Spec, FPredictionKey PredictionKey, const bool bPredictGameplayCues = false);
+
 	void ExecuteActiveEffectsFrom(FGameplayEffectSpec &Spec, FPredictionKey PredictionKey = FPredictionKey() );
 	
 	void ExecutePeriodicGameplayEffect(FActiveGameplayEffectHandle Handle);	// This should not be outward facing to the skill system API, should only be called by the owning AbilitySystemComponent
