GameplayAbilities
* Prevent null subobjects from getting registered for replication

diff --git a/Engine/Plugins/Runtime/GameplayAbilities/Source/GameplayAbilities/Private/AbilitySystemComponent.cpp b/Engine/Plugins/Runtime/GameplayAbilities/Source/GameplayAbilities/Private/AbilitySystemComponent.cpp
--- a/Engine/Plugins/Runtime/GameplayAbilities/Source/GameplayAbilities/Private/AbilitySystemComponent.cpp
+++ b/Engine/Plugins/Runtime/GameplayAbilities/Source/GameplayAbilities/Private/AbilitySystemComponent.cpp
@@ -238,7 +238,10 @@
 	{
 		for (UAttributeSet* ReplicatedAttribute : SpawnedAttributes)
 		{
-			AddReplicatedSubObject(ReplicatedAttribute);
+			if (ReplicatedAttribute)
+			{
+				AddReplicatedSubObject(ReplicatedAttribute);
+			}
 		}
 	}
 }
@@ -2827,12 +2830,15 @@
 	// Add the elements from the new list
 	for (UAttributeSet* NewAttribute : NewSpawnedAttributes)
 	{
-		AddSpawnedAttribute(NewAttribute);
+		if (NewAttribute)
+		{
+			AddSpawnedAttribute(NewAttribute);
 
-		AActor* ActorOwner = NewAttribute->GetTypedOuter<AActor>();
-		if (ActorOwner)
-		{
-			ActorOwner->OnEndPlay.AddUniqueDynamic(this, &UAbilitySystemComponent::OnSpawnedAttributesEndPlayed);
+			AActor* ActorOwner = NewAttribute->GetTypedOuter<AActor>();
+			if (ActorOwner)
+			{
+				ActorOwner->OnEndPlay.AddUniqueDynamic(this, &UAbilitySystemComponent::OnSpawnedAttributesEndPlayed);
+			}
 		}
 	}
 
@@ -2841,6 +2847,10 @@
 
 void UAbilitySystemComponent::AddSpawnedAttribute(UAttributeSet* Attribute)
 {
+	if (Attribute == nullptr)
+	{
+		return;
+	}
 
 	if (SpawnedAttributes.Find(Attribute) == INDEX_NONE)
 	{
@@ -2972,7 +2982,7 @@
 	for (int32 Index = SpawnedAttributes.Num() - 1; Index >= 0; --Index)
 	{
 		UAttributeSet* AttributeSet = SpawnedAttributes[Index];
-		if (AttributeSet->GetTypedOuter<AActor>() == InActor)
+		if (AttributeSet && AttributeSet->GetTypedOuter<AActor>() == InActor)
 		{
 			if (IsUsingRegisteredSubObjectList())
 			{
diff --git a/Engine/Source/Runtime/GameplayTasks/Private/GameplayTasksComponent.cpp b/Engine/Source/Runtime/GameplayTasks/Private/GameplayTasksComponent.cpp
--- a/Engine/Source/Runtime/GameplayTasks/Private/GameplayTasksComponent.cpp
+++ b/Engine/Source/Runtime/GameplayTasks/Private/GameplayTasksComponent.cpp
@@ -198,10 +198,13 @@
 		// Find if any tasks got removed
 		for (UGameplayTask* OldSimulatedTask : PreviousSimulatedTasks)
 		{
-			const bool bIsRemoved = SimulatedTasks.Find(OldSimulatedTask) == INDEX_NONE;
-			if (bIsRemoved)
+			if (OldSimulatedTask)
 			{
-				RemoveReplicatedSubObject(OldSimulatedTask);
+				const bool bIsRemoved = SimulatedTasks.Find(OldSimulatedTask) == INDEX_NONE;
+				if (bIsRemoved)
+				{
+					RemoveReplicatedSubObject(OldSimulatedTask);
+				}
 			}
 		}
 	}
@@ -726,6 +729,11 @@
 
 bool UGameplayTasksComponent::AddSimulatedTask(UGameplayTask* NewTask)
 {
+	if (NewTask == nullptr)
+	{
+		return false;
+	}
+
 	if (SimulatedTasks.Find(NewTask) == INDEX_NONE)
 	{
 		SimulatedTasks.Add(NewTask);
@@ -762,7 +770,10 @@
 		// Unregister all current tasks
 		for (UGameplayTask* OldGameplayTask : SimulatedTasks)
 		{
-			RemoveReplicatedSubObject(OldGameplayTask);
+			if (OldGameplayTask)
+			{
+				RemoveReplicatedSubObject(OldGameplayTask);
+			}
 		}
 
 		SimulatedTasks.Reset(NewSimulatedTasks.Num());
@@ -770,8 +781,11 @@
 		// Register the new tasks
 		for (UGameplayTask* NewGameplayTask : NewSimulatedTasks)
 		{
-			AddReplicatedSubObject(NewGameplayTask, COND_SkipOwner);
-			SimulatedTasks.Add(NewGameplayTask);
+			if (NewGameplayTask)
+			{
+				AddReplicatedSubObject(NewGameplayTask, COND_SkipOwner);
+				SimulatedTasks.Add(NewGameplayTask);
+			}
 		}
 	}
 	else
