Potential server crash fix.
#rnx

diff --git a/Engine/Plugins/Runtime/GameplayAbilities/Source/GameplayAbilities/Private/GameplayEffect.cpp b/Engine/Plugins/Runtime/GameplayAbilities/Source/GameplayAbilities/Private/GameplayEffect.cpp
--- a/Engine/Plugins/Runtime/GameplayAbilities/Source/GameplayAbilities/Private/GameplayEffect.cpp
+++ b/Engine/Plugins/Runtime/GameplayAbilities/Source/GameplayAbilities/Private/GameplayEffect.cpp
@@ -3979,7 +3979,8 @@ void FActiveGameplayEffectsContainer::OnOwnerTagChange(FGameplayTag TagChange, i
 	TSet<FActiveGameplayEffectHandle>* Ptr = ActiveEffectTagDependencies.Find(TagChange);
 	if (Ptr)
 	{
-		TSet<FActiveGameplayEffectHandle>& Handles = *Ptr;
+		// Copy the set in case it's modified while iterating.
+		TSet<FActiveGameplayEffectHandle> Handles = *Ptr;
 		for (const FActiveGameplayEffectHandle& Handle : Handles)
 		{
 			FActiveGameplayEffect* ActiveEffect = GetActiveGameplayEffect(Handle);
