- Fixed an issue where UCancellableAsyncAction could be destroyed, but never give the subclasses a chance to end and cleanup.  (for example, the subclass never had a chance to unregister any delegates it registered on other actors)
 - Added WaitGameplayTagQuery nodes, one based off of UAbilityTask and one off of UAbilityAsync.  This node specifies a TagQuery, and will trigger it's output pin when the query becomes true or false, based on configuration.

diff --git a/Engine/Plugins/Runtime/GameplayAbilities/Source/GameplayAbilities/Private/Abilities/Async/AbilityAsync_WaitGameplayTagQuery.cpp b/Engine/Plugins/Runtime/GameplayAbilities/Source/GameplayAbilities/Private/Abilities/Async/AbilityAsync_WaitGameplayTagQuery.cpp
new file mode 100644
--- /dev/null
+++ b/Engine/Plugins/Runtime/GameplayAbilities/Source/GameplayAbilities/Private/Abilities/Async/AbilityAsync_WaitGameplayTagQuery.cpp
@@ -0,0 +1,134 @@
+// Copyright Epic Games, Inc. All Rights Reserved.
+
+#include "Abilities/Async/AbilityAsync_WaitGameplayTagQuery.h"
+
+#include "AbilitySystemComponent.h"
+#include "AbilitySystemGlobals.h"
+
+void UAbilityAsync_WaitGameplayTagQuery::Activate()
+{
+	Super::Activate();
+
+	UAbilitySystemComponent* ASC = ShouldBroadcastDelegates() ? GetAbilitySystemComponent() : nullptr;
+	if (ASC != nullptr)
+	{
+		//Start the query state to the opposite of the triggered condition so it can be triggered immediately
+		bQueryState = TriggerCondition != EWaitGameplayTagQueryAsyncTriggerCondition::WhenTrue;
+
+		TArray<FGameplayTag> QueryTags;
+		TagQuery.GetGameplayTagArray(QueryTags);
+
+		for (const FGameplayTag& Tag : QueryTags)
+		{
+			if (!TagHandleMap.Contains(Tag))
+			{
+				UpdateTargetTags(Tag, ASC->GetTagCount(Tag));
+				TagHandleMap.Add(Tag, ASC->RegisterGameplayTagEvent(Tag).AddUObject(this, &UAbilityAsync_WaitGameplayTagQuery::UpdateTargetTags));
+			}
+		}
+
+		EvaluateTagQuery();
+
+		bRegisteredCallbacks = true;
+	}
+	else
+	{
+		EndAction();
+	}
+}
+
+void UAbilityAsync_WaitGameplayTagQuery::EndAction()
+{
+	if (UAbilitySystemComponent* ASC = GetAbilitySystemComponent())
+	{
+		for (TPair<FGameplayTag, FDelegateHandle> Pair : TagHandleMap)
+		{
+			if (Pair.Value.IsValid())
+			{
+				ASC->UnregisterGameplayTagEvent(Pair.Value, Pair.Key);
+			}
+		}
+	}
+
+	TagHandleMap.Empty();
+	TargetTags.Reset();
+
+	Super::EndAction();
+}
+
+bool UAbilityAsync_WaitGameplayTagQuery::ShouldBroadcastDelegates() const
+{
+	return !TagQuery.IsEmpty() && Super::ShouldBroadcastDelegates();
+}
+
+void UAbilityAsync_WaitGameplayTagQuery::UpdateTargetTags(const FGameplayTag Tag, int32 NewCount)
+{
+	if (ShouldBroadcastDelegates())
+	{
+		if (NewCount <= 0)
+		{
+			TargetTags.RemoveTag(Tag);
+		}
+		else
+		{
+			TargetTags.AddTag(Tag);
+		}
+
+		if (bRegisteredCallbacks)
+		{
+			EvaluateTagQuery();
+		}
+	}
+	else
+	{
+		EndAction();
+	}
+}
+
+void UAbilityAsync_WaitGameplayTagQuery::EvaluateTagQuery()
+{
+	if (TagQuery.IsEmpty())
+	{
+		return;
+	}
+
+	const bool bMatchesQuery = !TargetTags.IsEmpty() && TagQuery.Matches(TargetTags);
+	const bool bStateChanged = bMatchesQuery != bQueryState;
+	bQueryState = bMatchesQuery;
+
+	bool bTriggerDelegate = false;
+	if (bStateChanged)
+	{
+		if (TriggerCondition == EWaitGameplayTagQueryAsyncTriggerCondition::WhenTrue && bQueryState)
+		{
+			bTriggerDelegate = true;
+		}
+		else if (TriggerCondition == EWaitGameplayTagQueryAsyncTriggerCondition::WhenFalse && !bQueryState)
+		{
+			bTriggerDelegate = true;
+		}
+	}
+
+	if (bTriggerDelegate)
+	{
+		Triggered.Broadcast();
+		if (bOnlyTriggerOnce)
+		{
+			EndAction();
+		}
+	}
+}
+
+UAbilityAsync_WaitGameplayTagQuery* UAbilityAsync_WaitGameplayTagQuery::WaitGameplayTagQueryOnActor(AActor* TargetActor, 
+																									const FGameplayTagQuery TagQuery, 
+																									const EWaitGameplayTagQueryAsyncTriggerCondition TriggerCondition /*= EWaitGameplayTagQueryTriggerCondition::WhenTrue*/, 
+																									const bool bOnlyTriggerOnce/*=false*/)
+{
+	UAbilityAsync_WaitGameplayTagQuery* MyObj = NewObject<UAbilityAsync_WaitGameplayTagQuery>();
+	MyObj->SetAbilityActor(TargetActor);
+	MyObj->TagQuery = TagQuery;
+	MyObj->TriggerCondition = TriggerCondition;
+	MyObj->bOnlyTriggerOnce = bOnlyTriggerOnce;
+
+	return MyObj;
+}
diff --git a/Engine/Plugins/Runtime/GameplayAbilities/Source/GameplayAbilities/Private/Abilities/Tasks/AbilityTask_WaitGameplayTagQuery.cpp b/Engine/Plugins/Runtime/GameplayAbilities/Source/GameplayAbilities/Private/Abilities/Tasks/AbilityTask_WaitGameplayTagQuery.cpp
new file mode 100644
--- /dev/null
+++ b/Engine/Plugins/Runtime/GameplayAbilities/Source/GameplayAbilities/Private/Abilities/Tasks/AbilityTask_WaitGameplayTagQuery.cpp
@@ -0,0 +1,133 @@
+// Copyright Epic Games, Inc. All Rights Reserved.
+
+#include "Abilities/Tasks/AbilityTask_WaitGameplayTagQuery.h"
+
+#include "AbilitySystemComponent.h"
+#include "AbilitySystemGlobals.h"
+
+void UAbilityTask_WaitGameplayTagQuery::Activate()
+{
+	Super::Activate();
+
+	if (UAbilitySystemComponent* ASC = GetTargetASC())
+	{
+		//Start the query state to the opposite of the triggered condition so it can be triggered immediately
+		bQueryState = TriggerCondition != EWaitGameplayTagQueryTriggerCondition::WhenTrue;
+
+		TArray<FGameplayTag> QueryTags;
+		TagQuery.GetGameplayTagArray(QueryTags);
+
+		for (const FGameplayTag& Tag : QueryTags)
+		{
+			if (!TagHandleMap.Contains(Tag))
+			{
+				UpdateTargetTags(Tag, ASC->GetTagCount(Tag));
+				TagHandleMap.Add(Tag, ASC->RegisterGameplayTagEvent(Tag).AddUObject(this, &UAbilityTask_WaitGameplayTagQuery::UpdateTargetTags));
+			}
+		}
+
+		EvaluateTagQuery();
+
+		bRegisteredCallbacks = true;
+	}
+}
+
+void UAbilityTask_WaitGameplayTagQuery::UpdateTargetTags(const FGameplayTag Tag, int32 NewCount)
+{
+	if (NewCount <= 0)
+	{
+		TargetTags.RemoveTag(Tag);
+	}
+	else
+	{
+		TargetTags.AddTag(Tag);
+	}
+
+	if (bRegisteredCallbacks)
+	{
+		EvaluateTagQuery();
+	}
+}
+
+void UAbilityTask_WaitGameplayTagQuery::EvaluateTagQuery()
+{
+	if (TagQuery.IsEmpty())
+	{
+		return;
+	}
+
+	const bool bMatchesQuery = !TargetTags.IsEmpty() && TagQuery.Matches(TargetTags);
+	const bool bStateChanged = bMatchesQuery != bQueryState;
+	bQueryState = bMatchesQuery;
+
+	bool bTriggerDelegate = false;
+	if (bStateChanged)
+	{
+		if (TriggerCondition == EWaitGameplayTagQueryTriggerCondition::WhenTrue && bQueryState)
+		{
+			bTriggerDelegate = true;
+		}
+		else if (TriggerCondition == EWaitGameplayTagQueryTriggerCondition::WhenFalse && !bQueryState)
+		{
+			bTriggerDelegate = true;
+		}
+	}
+
+	if (bTriggerDelegate)
+	{
+		Triggered.Broadcast();
+		if (bOnlyTriggerOnce)
+		{
+			EndTask();
+		}
+	}
+}
+
+void UAbilityTask_WaitGameplayTagQuery::OnDestroy(bool AbilityIsEnding)
+{
+	UAbilitySystemComponent* ASC = bRegisteredCallbacks ? GetTargetASC() : nullptr;
+	if (ASC != nullptr)
+	{
+		for (TPair<FGameplayTag, FDelegateHandle> Pair : TagHandleMap)
+		{
+			if (Pair.Value.IsValid())
+			{
+				ASC->UnregisterGameplayTagEvent(Pair.Value, Pair.Key);
+			}
+		}
+	}
+
+	TagHandleMap.Empty();
+	TargetTags.Reset();
+
+	Super::OnDestroy(AbilityIsEnding);
+}
+
+UAbilitySystemComponent* UAbilityTask_WaitGameplayTagQuery::GetTargetASC()
+{	
+	return bUseExternalTarget ? OptionalExternalTarget : AbilitySystemComponent;
+}
+
+void UAbilityTask_WaitGameplayTagQuery::SetExternalTarget(const AActor* Actor)
+{
+	if (Actor != nullptr)
+	{
+		bUseExternalTarget = true;
+		OptionalExternalTarget = UAbilitySystemGlobals::GetAbilitySystemComponentFromActor(Actor);
+	}
+}
+
+UAbilityTask_WaitGameplayTagQuery* UAbilityTask_WaitGameplayTagQuery::WaitGameplayTagQuery(UGameplayAbility* OwningAbility, 
+																						   const FGameplayTagQuery TagQuery, 
+																						   const AActor* InOptionalExternalTarget /*= nullptr*/, 
+																						   const EWaitGameplayTagQueryTriggerCondition TriggerCondition /*= EWaitGameplayTagQueryTriggerCondition::WhenTrue*/, 
+																						   const bool bOnlyTriggerOnce /*= false*/)
+{
+	UAbilityTask_WaitGameplayTagQuery* MyObj = NewAbilityTask<UAbilityTask_WaitGameplayTagQuery>(OwningAbility);
+	MyObj->TagQuery = TagQuery;
+	MyObj->SetExternalTarget(InOptionalExternalTarget);
+	MyObj->TriggerCondition = TriggerCondition;
+	MyObj->bOnlyTriggerOnce = bOnlyTriggerOnce;
+
+	return MyObj;
+}
diff --git a/Engine/Plugins/Runtime/GameplayAbilities/Source/GameplayAbilities/Public/Abilities/Async/AbilityAsync_WaitGameplayTagQuery.h b/Engine/Plugins/Runtime/GameplayAbilities/Source/GameplayAbilities/Public/Abilities/Async/AbilityAsync_WaitGameplayTagQuery.h
new file mode 100644
--- /dev/null
+++ b/Engine/Plugins/Runtime/GameplayAbilities/Source/GameplayAbilities/Public/Abilities/Async/AbilityAsync_WaitGameplayTagQuery.h
@@ -0,0 +1,79 @@
+// Copyright Epic Games, Inc. All Rights Reserved.
+#pragma once
+
+#include "AbilityAsync.h"
+
+#include "AbilityAsync_WaitGameplayTagQuery.generated.h"
+
+class UAbilitySystemComponent;
+
+
+DECLARE_DYNAMIC_MULTICAST_DELEGATE(FAsyncWaitGameplayTagQueryDelegate);
+
+/** This enum defines the condition in which the async wait gameplay tag query node will trigger. */
+UENUM()
+enum class EWaitGameplayTagQueryAsyncTriggerCondition : uint8
+{
+	WhenTrue = 0,
+	WhenFalse,
+};
+
+/** This class defines an async node to wait on a gameplay tag query. */
+UCLASS()
+class GAMEPLAYABILITIES_API UAbilityAsync_WaitGameplayTagQuery : public UAbilityAsync
+{
+	GENERATED_BODY()
+protected:
+
+	/** Activates this AbilityAsync node. */
+	virtual void Activate() override;
+
+	/** Ends this AbilityAsync node, unregistering all tag changed delagates. */
+	virtual void EndAction() override;
+
+	/** Returns if this AbilityAsync node should currently broadcast it's delegates. */
+	virtual bool ShouldBroadcastDelegates() const;
+
+	/** This will update the tags in the TargetTags container to reflect that tags that are on the target ASC. */
+	void UpdateTargetTags(const FGameplayTag Tag, int32 NewCount);
+
+	/** This will evaluate the TargetTags using the given TagQuery, executing the Trigger delegate if needed. */
+	void EvaluateTagQuery();
+
+	/** This is the tag query to evaluate for triggering this node. */
+	FGameplayTagQuery TagQuery;
+
+	/** This indicates when to Trigger the Triggered output for this node. */
+	EWaitGameplayTagQueryAsyncTriggerCondition TriggerCondition = EWaitGameplayTagQueryAsyncTriggerCondition::WhenTrue;
+
+	/** This indicates if this node should only trigger once, or any number of times. */
+	bool bOnlyTriggerOnce = false;
+
+	/** Indicates if the callbacks for tag changes have been registered. */
+	bool bRegisteredCallbacks = false;
+
+	/** This was the last result of evaluating the TagQuery. */
+	bool bQueryState = false;
+
+	/** This is the tag container of the targets the target ASC currently has.
+		NOTE: This will only contain the tags referenced in the TagQuery. */
+	FGameplayTagContainer TargetTags;
+
+	/** This is the handles to the tag changed delegate for each gameplay tag in the TagQuery. */
+	TMap<FGameplayTag, FDelegateHandle> TagHandleMap;
+
+	/** This delegate will be triggered when the trigger condition has been met. */
+	UPROPERTY(BlueprintAssignable)
+	FAsyncWaitGameplayTagQueryDelegate Triggered;
+
+	/**
+	 * 	Wait until the given gameplay tag query has become true or false, based on TriggerCondition, looking at the target actors ASC.
+	 *  If the the tag query already satisfies the TriggerCondition when this task is started, it will immediately broadcast the Triggered
+	 *  event. It will keep listening as long as bOnlyTriggerOnce = false.
+	 */
+ 	UFUNCTION(BlueprintCallable, Category = "Ability|Tasks", meta = (DefaultToSelf = "TargetActor", BlueprintInternalUseOnly = "TRUE"))
+ 	static UAbilityAsync_WaitGameplayTagQuery* WaitGameplayTagQueryOnActor(AActor* TargetActor, 
+																		   const FGameplayTagQuery TagQuery, 
+																		   const EWaitGameplayTagQueryAsyncTriggerCondition TriggerCondition = EWaitGameplayTagQueryAsyncTriggerCondition::WhenTrue, 
+																		   const bool bOnlyTriggerOnce=false);
+};
\ No newline at end of file
diff --git a/Engine/Plugins/Runtime/GameplayAbilities/Source/GameplayAbilities/Public/Abilities/Tasks/AbilityTask_WaitGameplayTagQuery.h b/Engine/Plugins/Runtime/GameplayAbilities/Source/GameplayAbilities/Public/Abilities/Tasks/AbilityTask_WaitGameplayTagQuery.h
new file mode 100644
--- /dev/null
+++ b/Engine/Plugins/Runtime/GameplayAbilities/Source/GameplayAbilities/Public/Abilities/Tasks/AbilityTask_WaitGameplayTagQuery.h
@@ -0,0 +1,91 @@
+// Copyright Epic Games, Inc. All Rights Reserved.
+#pragma once
+
+#include "Abilities/Tasks/AbilityTask.h"
+
+#include "AbilityTask_WaitGameplayTagQuery.generated.h"
+
+DECLARE_DYNAMIC_MULTICAST_DELEGATE(FWaitGameplayTagQueryDelegate);
+
+/** This enum defines the condition in which the wait gameplay tag query node will trigger. */
+UENUM()
+enum class EWaitGameplayTagQueryTriggerCondition : uint8
+{
+	WhenTrue = 0,
+	WhenFalse,
+};
+
+/** This class defines a node to wait on a gameplay tag query. */
+UCLASS()
+class GAMEPLAYABILITIES_API UAbilityTask_WaitGameplayTagQuery : public UAbilityTask
+{
+	GENERATED_BODY()
+
+	/** Activates this AbilityTask. */
+	virtual void Activate() override;
+	
+	/** Sets the external target that this node whould use for checking tags. */
+	void SetExternalTarget(const AActor* Actor);
+
+	/**
+	 * 	Wait until the given gameplay tag query has become true or false, based on TriggerCondition. 
+	 *  By default this will look at the owner of this ability. OptionalExternalTarget can be set to
+	 *  make this look at another actor's tags for changes.  If the the tag query already satisfies 
+	 *  the TriggerCondition when this task is started, it will immediately broadcast the Triggered 
+	 *  event. It will keep listening as long as bOnlyTriggerOnce = false.
+	 */
+	UFUNCTION(BlueprintCallable, Category = "Ability|Tasks", meta = (HidePin = "OwningAbility", DefaultToSelf = "OwningAbility", BlueprintInternalUseOnly = "TRUE"))
+	static UAbilityTask_WaitGameplayTagQuery* WaitGameplayTagQuery(UGameplayAbility* OwningAbility, 
+																   const FGameplayTagQuery TagQuery, 
+																   const AActor* InOptionalExternalTarget = nullptr, 
+																   const EWaitGameplayTagQueryTriggerCondition TriggerCondition = EWaitGameplayTagQueryTriggerCondition::WhenTrue, 
+																   const bool bOnlyTriggerOnce = false);
+
+protected:
+
+	/** This will update the tags in the TargetTags container to reflect that tags that are on the target ASC. */
+	UFUNCTION()
+	void UpdateTargetTags(const FGameplayTag Tag, int32 NewCount);
+
+	/** This will evaluate the TargetTags using the given TagQuery, executing the Trigger delegate if needed. */
+	void EvaluateTagQuery();
+
+	/** This gets the ASC to use to listen to tag changed events. */
+	UAbilitySystemComponent* GetTargetASC();
+
+	/** This will handle cleaning up any registered delegates. */
+	virtual void OnDestroy(bool AbilityIsEnding) override;
+
+	/** This is the tag query to evaluate for triggering this node. */
+	FGameplayTagQuery TagQuery;
+
+	/** Indicates if the callbacks for tag changes have been registered. */
+	bool bRegisteredCallbacks = false;
+
+	/** This was the last result of evaluating the TagQuery. */
+	bool bQueryState = false;
+
+	/** This is the tag container of the targets the target ASC currently has. 
+		NOTE: This will only contain the tags referenced in the TagQuery. */
+	FGameplayTagContainer TargetTags;
+
+	/** This is the handles to the tag changed delegate for each gameplay tag in the TagQuery. */
+	TMap<FGameplayTag, FDelegateHandle> TagHandleMap;
+
+	/** This delegate will be triggered when the trigger condition has been met. */
+	UPROPERTY(BlueprintAssignable)
+	FWaitGameplayTagQueryDelegate Triggered;
+	
+	/** This is the optional external target to use when getting the ASC to get tags from. */
+	UPROPERTY()
+	UAbilitySystemComponent* OptionalExternalTarget = nullptr;
+
+	/** This indicates when to Trigger the Triggered output for this node. */
+	EWaitGameplayTagQueryTriggerCondition TriggerCondition = EWaitGameplayTagQueryTriggerCondition::WhenTrue;
+
+	/** This indicates if we should use the external target. */
+	bool bUseExternalTarget = false;	
+
+	/** This indicates if this node should only trigger once, or any number of times. */
+	bool bOnlyTriggerOnce = false;
+};
\ No newline at end of file
diff --git a/Engine/Source/Runtime/Engine/Classes/Engine/CancellableAsyncAction.h b/Engine/Source/Runtime/Engine/Classes/Engine/CancellableAsyncAction.h
--- a/Engine/Source/Runtime/Engine/Classes/Engine/CancellableAsyncAction.h
+++ b/Engine/Source/Runtime/Engine/Classes/Engine/CancellableAsyncAction.h
@@ -17,6 +17,9 @@ class ENGINE_API UCancellableAsyncAction : public UBlueprintAsyncActionBase
 	GENERATED_BODY()
 public:
 
+	/** Handle when this action is being destroyed to ensure that the action is canceled and child classes can clean up. */
+	virtual void BeginDestroy() override;
+
 	/** Cancel an asynchronous action, this attempts to cancel any lower level processes and also prevents delegates from being fired */
 	UFUNCTION(BlueprintCallable, Category = "Async Action")
 	virtual void Cancel();
diff --git a/Engine/Source/Runtime/Engine/Private/Engine/CancellableAsyncAction.cpp b/Engine/Source/Runtime/Engine/Private/Engine/CancellableAsyncAction.cpp
--- a/Engine/Source/Runtime/Engine/Private/Engine/CancellableAsyncAction.cpp
+++ b/Engine/Source/Runtime/Engine/Private/Engine/CancellableAsyncAction.cpp
@@ -12,6 +12,13 @@ void UCancellableAsyncAction::Cancel()
 	SetReadyToDestroy();
 }
 
+void UCancellableAsyncAction::BeginDestroy()
+{
+	Cancel();
+
+	Super::BeginDestroy();
+}
+
 bool UCancellableAsyncAction::IsActive() const
 {
 	return ShouldBroadcastDelegates();
