Removed redundant enum EWaitGameplayTagQueryAsyncTriggerCondition.

diff --git a/Engine/Plugins/Runtime/GameplayAbilities/Source/GameplayAbilities/Private/Abilities/Async/AbilityAsync_WaitGameplayTagQuery.cpp b/Engine/Plugins/Runtime/GameplayAbilities/Source/GameplayAbilities/Private/Abilities/Async/AbilityAsync_WaitGameplayTagQuery.cpp
--- a/Engine/Plugins/Runtime/GameplayAbilities/Source/GameplayAbilities/Private/Abilities/Async/AbilityAsync_WaitGameplayTagQuery.cpp
+++ b/Engine/Plugins/Runtime/GameplayAbilities/Source/GameplayAbilities/Private/Abilities/Async/AbilityAsync_WaitGameplayTagQuery.cpp
@@ -13,7 +13,7 @@ void UAbilityAsync_WaitGameplayTagQuery::Activate()
 	if (ASC != nullptr)
 	{
 		//Start the query state to the opposite of the triggered condition so it can be triggered immediately
-		bQueryState = TriggerCondition != EWaitGameplayTagQueryAsyncTriggerCondition::WhenTrue;
+		bQueryState = TriggerCondition != EWaitGameplayTagQueryTriggerCondition::WhenTrue;
 
 		TArray<FGameplayTag> QueryTags;
 		TagQuery.GetGameplayTagArray(QueryTags);
@@ -99,11 +99,11 @@ void UAbilityAsync_WaitGameplayTagQuery::EvaluateTagQuery()
 	bool bTriggerDelegate = false;
 	if (bStateChanged)
 	{
-		if (TriggerCondition == EWaitGameplayTagQueryAsyncTriggerCondition::WhenTrue && bQueryState)
+		if (TriggerCondition == EWaitGameplayTagQueryTriggerCondition::WhenTrue && bQueryState)
 		{
 			bTriggerDelegate = true;
 		}
-		else if (TriggerCondition == EWaitGameplayTagQueryAsyncTriggerCondition::WhenFalse && !bQueryState)
+		else if (TriggerCondition == EWaitGameplayTagQueryTriggerCondition::WhenFalse && !bQueryState)
 		{
 			bTriggerDelegate = true;
 		}
@@ -121,7 +121,7 @@ void UAbilityAsync_WaitGameplayTagQuery::EvaluateTagQuery()
 
 UAbilityAsync_WaitGameplayTagQuery* UAbilityAsync_WaitGameplayTagQuery::WaitGameplayTagQueryOnActor(AActor* TargetActor, 
 																									const FGameplayTagQuery TagQuery, 
-																									const EWaitGameplayTagQueryAsyncTriggerCondition TriggerCondition /*= EWaitGameplayTagQueryTriggerCondition::WhenTrue*/, 
+																									const EWaitGameplayTagQueryTriggerCondition TriggerCondition /*= EWaitGameplayTagQueryTriggerCondition::WhenTrue*/, 
 																									const bool bOnlyTriggerOnce/*=false*/)
 {
 	UAbilityAsync_WaitGameplayTagQuery* MyObj = NewObject<UAbilityAsync_WaitGameplayTagQuery>();
diff --git a/Engine/Plugins/Runtime/GameplayAbilities/Source/GameplayAbilities/Public/Abilities/Async/AbilityAsync_WaitGameplayTagQuery.h b/Engine/Plugins/Runtime/GameplayAbilities/Source/GameplayAbilities/Public/Abilities/Async/AbilityAsync_WaitGameplayTagQuery.h
--- a/Engine/Plugins/Runtime/GameplayAbilities/Source/GameplayAbilities/Public/Abilities/Async/AbilityAsync_WaitGameplayTagQuery.h
+++ b/Engine/Plugins/Runtime/GameplayAbilities/Source/GameplayAbilities/Public/Abilities/Async/AbilityAsync_WaitGameplayTagQuery.h
@@ -1,23 +1,15 @@
 // Copyright Epic Games, Inc. All Rights Reserved.
 #pragma once
 
+#include "Abilities/Tasks/AbilityTask_WaitGameplayTagQuery.h"
 #include "AbilityAsync.h"
 
 #include "AbilityAsync_WaitGameplayTagQuery.generated.h"
 
 class UAbilitySystemComponent;
 
-
 DECLARE_DYNAMIC_MULTICAST_DELEGATE(FAsyncWaitGameplayTagQueryDelegate);
 
-/** This enum defines the condition in which the async wait gameplay tag query node will trigger. */
-UENUM()
-enum class EWaitGameplayTagQueryAsyncTriggerCondition : uint8
-{
-	WhenTrue = 0,
-	WhenFalse,
-};
-
 /** This class defines an async node to wait on a gameplay tag query. */
 UCLASS()
 class GAMEPLAYABILITIES_API UAbilityAsync_WaitGameplayTagQuery : public UAbilityAsync
@@ -44,7 +36,7 @@ protected:
 	FGameplayTagQuery TagQuery;
 
 	/** This indicates when to Trigger the Triggered output for this node. */
-	EWaitGameplayTagQueryAsyncTriggerCondition TriggerCondition = EWaitGameplayTagQueryAsyncTriggerCondition::WhenTrue;
+	EWaitGameplayTagQueryTriggerCondition TriggerCondition = EWaitGameplayTagQueryTriggerCondition::WhenTrue;
 
 	/** This indicates if this node should only trigger once, or any number of times. */
 	bool bOnlyTriggerOnce = false;
@@ -74,6 +66,6 @@ protected:
  	UFUNCTION(BlueprintCallable, Category = "Ability|Tasks", meta = (DefaultToSelf = "TargetActor", BlueprintInternalUseOnly = "TRUE"))
  	static UAbilityAsync_WaitGameplayTagQuery* WaitGameplayTagQueryOnActor(AActor* TargetActor, 
 																		   const FGameplayTagQuery TagQuery, 
-																		   const EWaitGameplayTagQueryAsyncTriggerCondition TriggerCondition = EWaitGameplayTagQueryAsyncTriggerCondition::WhenTrue, 
+																		   const EWaitGameplayTagQueryTriggerCondition TriggerCondition = EWaitGameplayTagQueryTriggerCondition::WhenTrue,
 																		   const bool bOnlyTriggerOnce=false);
 };
\ No newline at end of file
