[KillPendingKill] Make the Ability System Component Reference in the Ability Task a weak Pointer, following the pattern for the TasksComponent in UgameplayTask.

diff --git a/Engine/Plugins/Runtime/GameplayAbilities/Source/GameplayAbilities/Private/Abilities/Tasks/AbilityTask.cpp b/Engine/Plugins/Runtime/GameplayAbilities/Source/GameplayAbilities/Private/Abilities/Tasks/AbilityTask.cpp
--- a/Engine/Plugins/Runtime/GameplayAbilities/Source/GameplayAbilities/Private/Abilities/Tasks/AbilityTask.cpp
+++ b/Engine/Plugins/Runtime/GameplayAbilities/Source/GameplayAbilities/Private/Abilities/Tasks/AbilityTask.cpp
@@ -181,7 +181,7 @@ bool UAbilityTask::CallOrAddReplicatedDelegate(EAbilityGenericReplicatedEvent::T
 
 void UAbilityTask::SetWaitingOnRemotePlayerData()
 {
-	if (IsValid(Ability) && AbilitySystemComponent)
+	if (IsValid(Ability) && AbilitySystemComponent.IsValid())
 	{
 		WaitStateBitMask |= (uint8)EAbilityTaskWaitState::WaitingOnUser;
 		Ability->NotifyAbilityTaskWaitingOnPlayerData(this);
@@ -200,7 +200,7 @@ bool UAbilityTask::IsWaitingOnRemotePlayerdata() const
 
 void UAbilityTask::SetWaitingOnAvatar()
 {
-	if (IsValid(Ability) && AbilitySystemComponent)
+	if (IsValid(Ability) && AbilitySystemComponent.IsValid())
 	{
 		WaitStateBitMask |= (uint8)EAbilityTaskWaitState::WaitingOnAvatar;
 		Ability->NotifyAbilityTaskWaitingOnAvatar(this);
diff --git a/Engine/Plugins/Runtime/GameplayAbilities/Source/GameplayAbilities/Private/Abilities/Tasks/AbilityTask_NetworkSyncPoint.cpp b/Engine/Plugins/Runtime/GameplayAbilities/Source/GameplayAbilities/Private/Abilities/Tasks/AbilityTask_NetworkSyncPoint.cpp
--- a/Engine/Plugins/Runtime/GameplayAbilities/Source/GameplayAbilities/Private/Abilities/Tasks/AbilityTask_NetworkSyncPoint.cpp
+++ b/Engine/Plugins/Runtime/GameplayAbilities/Source/GameplayAbilities/Private/Abilities/Tasks/AbilityTask_NetworkSyncPoint.cpp
@@ -11,7 +11,7 @@ UAbilityTask_NetworkSyncPoint::UAbilityTask_NetworkSyncPoint(const FObjectInitia
 
 void UAbilityTask_NetworkSyncPoint::OnSignalCallback()
 {
-	if (AbilitySystemComponent)
+	if (AbilitySystemComponent.IsValid())
 	{
 		AbilitySystemComponent->ConsumeGenericReplicatedEvent(ReplicatedEventToListenFor, GetAbilitySpecHandle(), GetActivationPredictionKey());
 	}
@@ -27,9 +27,9 @@ UAbilityTask_NetworkSyncPoint* UAbilityTask_NetworkSyncPoint::WaitNetSync(class
 
 void UAbilityTask_NetworkSyncPoint::Activate()
 {
-	FScopedPredictionWindow ScopedPrediction(AbilitySystemComponent, IsPredictingClient());
+	FScopedPredictionWindow ScopedPrediction(AbilitySystemComponent.Get(), IsPredictingClient());
 
-	if (AbilitySystemComponent)
+	if (AbilitySystemComponent.IsValid())
 	{
 		if (IsPredictingClient())
 		{
diff --git a/Engine/Plugins/Runtime/GameplayAbilities/Source/GameplayAbilities/Private/Abilities/Tasks/AbilityTask_PlayMontageAndWait.cpp b/Engine/Plugins/Runtime/GameplayAbilities/Source/GameplayAbilities/Private/Abilities/Tasks/AbilityTask_PlayMontageAndWait.cpp
--- a/Engine/Plugins/Runtime/GameplayAbilities/Source/GameplayAbilities/Private/Abilities/Tasks/AbilityTask_PlayMontageAndWait.cpp
+++ b/Engine/Plugins/Runtime/GameplayAbilities/Source/GameplayAbilities/Private/Abilities/Tasks/AbilityTask_PlayMontageAndWait.cpp
@@ -98,7 +98,7 @@ void UAbilityTask_PlayMontageAndWait::Activate()
 
 	bool bPlayedMontage = false;
 
-	if (AbilitySystemComponent)
+	if (AbilitySystemComponent.IsValid())
 	{
 		const FGameplayAbilityActorInfo* ActorInfo = Ability->GetCurrentActorInfo();
 		UAnimInstance* AnimInstance = ActorInfo->GetAnimInstance();
@@ -154,7 +154,7 @@ void UAbilityTask_PlayMontageAndWait::Activate()
 
 void UAbilityTask_PlayMontageAndWait::ExternalCancel()
 {
-	check(AbilitySystemComponent);
+	check(AbilitySystemComponent.IsValid());
 
 	if (ShouldBroadcastAbilityTaskDelegates())
 	{
@@ -203,7 +203,7 @@ bool UAbilityTask_PlayMontageAndWait::StopPlayingMontage()
 
 	// Check if the montage is still playing
 	// The ability would have been interrupted, in which case we should automatically stop the montage
-	if (AbilitySystemComponent && Ability)
+	if (AbilitySystemComponent.IsValid() && Ability)
 	{
 		if (AbilitySystemComponent->GetAnimatingAbility() == Ability
 			&& AbilitySystemComponent->GetCurrentMontage() == MontageToPlay)
diff --git a/Engine/Plugins/Runtime/GameplayAbilities/Source/GameplayAbilities/Private/Abilities/Tasks/AbilityTask_WaitAbilityActivate.cpp b/Engine/Plugins/Runtime/GameplayAbilities/Source/GameplayAbilities/Private/Abilities/Tasks/AbilityTask_WaitAbilityActivate.cpp
--- a/Engine/Plugins/Runtime/GameplayAbilities/Source/GameplayAbilities/Private/Abilities/Tasks/AbilityTask_WaitAbilityActivate.cpp
+++ b/Engine/Plugins/Runtime/GameplayAbilities/Source/GameplayAbilities/Private/Abilities/Tasks/AbilityTask_WaitAbilityActivate.cpp
@@ -42,7 +42,7 @@ UAbilityTask_WaitAbilityActivate* UAbilityTask_WaitAbilityActivate::WaitForAbili
 
 void UAbilityTask_WaitAbilityActivate::Activate()
 {
-	if (AbilitySystemComponent)
+	if (AbilitySystemComponent.IsValid())
 	{
 		OnAbilityActivateDelegateHandle = AbilitySystemComponent->AbilityActivatedCallbacks.AddUObject(this, &UAbilityTask_WaitAbilityActivate::OnAbilityActivate);
 	}
@@ -95,7 +95,7 @@ void UAbilityTask_WaitAbilityActivate::OnAbilityActivate(UGameplayAbility* Activ
 
 void UAbilityTask_WaitAbilityActivate::OnDestroy(bool AbilityEnded)
 {
-	if (AbilitySystemComponent)
+	if (AbilitySystemComponent.IsValid())
 	{
 		AbilitySystemComponent->AbilityActivatedCallbacks.Remove(OnAbilityActivateDelegateHandle);
 	}
diff --git a/Engine/Plugins/Runtime/GameplayAbilities/Source/GameplayAbilities/Private/Abilities/Tasks/AbilityTask_WaitAbilityCommit.cpp b/Engine/Plugins/Runtime/GameplayAbilities/Source/GameplayAbilities/Private/Abilities/Tasks/AbilityTask_WaitAbilityCommit.cpp
--- a/Engine/Plugins/Runtime/GameplayAbilities/Source/GameplayAbilities/Private/Abilities/Tasks/AbilityTask_WaitAbilityCommit.cpp
+++ b/Engine/Plugins/Runtime/GameplayAbilities/Source/GameplayAbilities/Private/Abilities/Tasks/AbilityTask_WaitAbilityCommit.cpp
@@ -31,7 +31,7 @@ UAbilityTask_WaitAbilityCommit* UAbilityTask_WaitAbilityCommit::WaitForAbilityCo
 
 void UAbilityTask_WaitAbilityCommit::Activate()
 {
-	if (AbilitySystemComponent)	
+	if (AbilitySystemComponent.IsValid())	
 	{		
 		OnAbilityCommitDelegateHandle = AbilitySystemComponent->AbilityCommittedCallbacks.AddUObject(this, &UAbilityTask_WaitAbilityCommit::OnAbilityCommit);
 	}
@@ -39,7 +39,7 @@ void UAbilityTask_WaitAbilityCommit::Activate()
 
 void UAbilityTask_WaitAbilityCommit::OnDestroy(bool AbilityEnded)
 {
-	if (AbilitySystemComponent)
+	if (AbilitySystemComponent.IsValid())
 	{
 		AbilitySystemComponent->AbilityCommittedCallbacks.Remove(OnAbilityCommitDelegateHandle);
 	}
diff --git a/Engine/Plugins/Runtime/GameplayAbilities/Source/GameplayAbilities/Private/Abilities/Tasks/AbilityTask_WaitAttributeChange.cpp b/Engine/Plugins/Runtime/GameplayAbilities/Source/GameplayAbilities/Private/Abilities/Tasks/AbilityTask_WaitAttributeChange.cpp
--- a/Engine/Plugins/Runtime/GameplayAbilities/Source/GameplayAbilities/Private/Abilities/Tasks/AbilityTask_WaitAttributeChange.cpp
+++ b/Engine/Plugins/Runtime/GameplayAbilities/Source/GameplayAbilities/Private/Abilities/Tasks/AbilityTask_WaitAttributeChange.cpp
@@ -111,7 +111,7 @@ void UAbilityTask_WaitAttributeChange::OnAttributeChange(const FOnAttributeChang
 
 UAbilitySystemComponent* UAbilityTask_WaitAttributeChange::GetFocusedASC()
 {
-	return ExternalOwner ? ExternalOwner : AbilitySystemComponent;
+	return ExternalOwner ? ExternalOwner : AbilitySystemComponent.Get();
 }
 
 void UAbilityTask_WaitAttributeChange::OnDestroy(bool AbilityEnded)
diff --git a/Engine/Plugins/Runtime/GameplayAbilities/Source/GameplayAbilities/Private/Abilities/Tasks/AbilityTask_WaitAttributeChangeRatioThreshold.cpp b/Engine/Plugins/Runtime/GameplayAbilities/Source/GameplayAbilities/Private/Abilities/Tasks/AbilityTask_WaitAttributeChangeRatioThreshold.cpp
--- a/Engine/Plugins/Runtime/GameplayAbilities/Source/GameplayAbilities/Private/Abilities/Tasks/AbilityTask_WaitAttributeChangeRatioThreshold.cpp
+++ b/Engine/Plugins/Runtime/GameplayAbilities/Source/GameplayAbilities/Private/Abilities/Tasks/AbilityTask_WaitAttributeChangeRatioThreshold.cpp
@@ -50,7 +50,7 @@ void UAbilityTask_WaitAttributeChangeRatioThreshold::Activate()
 
 UAbilitySystemComponent* UAbilityTask_WaitAttributeChangeRatioThreshold::GetFocusedASC()
 {
-	return ExternalOwner ? ExternalOwner : AbilitySystemComponent;
+	return ExternalOwner ? ExternalOwner : AbilitySystemComponent.Get();
 }
 
 void UAbilityTask_WaitAttributeChangeRatioThreshold::OnAttributeChange()
diff --git a/Engine/Plugins/Runtime/GameplayAbilities/Source/GameplayAbilities/Private/Abilities/Tasks/AbilityTask_WaitAttributeChangeThreshold.cpp b/Engine/Plugins/Runtime/GameplayAbilities/Source/GameplayAbilities/Private/Abilities/Tasks/AbilityTask_WaitAttributeChangeThreshold.cpp
--- a/Engine/Plugins/Runtime/GameplayAbilities/Source/GameplayAbilities/Private/Abilities/Tasks/AbilityTask_WaitAttributeChangeThreshold.cpp
+++ b/Engine/Plugins/Runtime/GameplayAbilities/Source/GameplayAbilities/Private/Abilities/Tasks/AbilityTask_WaitAttributeChangeThreshold.cpp
@@ -92,7 +92,7 @@ bool UAbilityTask_WaitAttributeChangeThreshold::DoesValuePassComparison(float Va
 
 UAbilitySystemComponent* UAbilityTask_WaitAttributeChangeThreshold::GetFocusedASC()
 {
-	return ExternalOwner ? ExternalOwner : AbilitySystemComponent;
+	return ExternalOwner ? ExternalOwner : AbilitySystemComponent.Get();
 }
 
 void UAbilityTask_WaitAttributeChangeThreshold::OnDestroy(bool AbilityEnded)
diff --git a/Engine/Plugins/Runtime/GameplayAbilities/Source/GameplayAbilities/Private/Abilities/Tasks/AbilityTask_WaitCancel.cpp b/Engine/Plugins/Runtime/GameplayAbilities/Source/GameplayAbilities/Private/Abilities/Tasks/AbilityTask_WaitCancel.cpp
--- a/Engine/Plugins/Runtime/GameplayAbilities/Source/GameplayAbilities/Private/Abilities/Tasks/AbilityTask_WaitCancel.cpp
+++ b/Engine/Plugins/Runtime/GameplayAbilities/Source/GameplayAbilities/Private/Abilities/Tasks/AbilityTask_WaitCancel.cpp
@@ -14,7 +14,7 @@ UAbilityTask_WaitCancel::UAbilityTask_WaitCancel(const FObjectInitializer& Objec
 
 void UAbilityTask_WaitCancel::OnCancelCallback()
 {
-	if (AbilitySystemComponent)
+	if (AbilitySystemComponent.IsValid())
 	{
 		AbilitySystemComponent->ConsumeGenericReplicatedEvent(EAbilityGenericReplicatedEvent::GenericCancel, GetAbilitySpecHandle(), GetActivationPredictionKey());
 		if (ShouldBroadcastAbilityTaskDelegates())
@@ -27,9 +27,9 @@ void UAbilityTask_WaitCancel::OnCancelCallback()
 
 void UAbilityTask_WaitCancel::OnLocalCancelCallback()
 {
-	FScopedPredictionWindow ScopedPrediction(AbilitySystemComponent, IsPredictingClient());
+	FScopedPredictionWindow ScopedPrediction(AbilitySystemComponent.Get(), IsPredictingClient());
 
-	if (AbilitySystemComponent && IsPredictingClient())
+	if (AbilitySystemComponent.IsValid() && IsPredictingClient())
 	{
 		AbilitySystemComponent->ServerSetReplicatedEvent(EAbilityGenericReplicatedEvent::GenericCancel, GetAbilitySpecHandle(), GetActivationPredictionKey() ,AbilitySystemComponent->ScopedPredictionKey);
 	}
@@ -43,7 +43,7 @@ UAbilityTask_WaitCancel* UAbilityTask_WaitCancel::WaitCancel(UGameplayAbility* O
 
 void UAbilityTask_WaitCancel::Activate()
 {
-	if (AbilitySystemComponent)
+	if (AbilitySystemComponent.IsValid())
 	{
 		const FGameplayAbilityActorInfo* Info = Ability->GetCurrentActorInfo();
 
@@ -68,7 +68,7 @@ void UAbilityTask_WaitCancel::Activate()
 
 void UAbilityTask_WaitCancel::OnDestroy(bool AbilityEnding)
 {
-	if (RegisteredCallbacks && AbilitySystemComponent)
+	if (RegisteredCallbacks && AbilitySystemComponent.IsValid())
 	{
 		AbilitySystemComponent->GenericLocalCancelCallbacks.RemoveDynamic(this, &UAbilityTask_WaitCancel::OnLocalCancelCallback);
 	}
diff --git a/Engine/Plugins/Runtime/GameplayAbilities/Source/GameplayAbilities/Private/Abilities/Tasks/AbilityTask_WaitConfirmCancel.cpp b/Engine/Plugins/Runtime/GameplayAbilities/Source/GameplayAbilities/Private/Abilities/Tasks/AbilityTask_WaitConfirmCancel.cpp
--- a/Engine/Plugins/Runtime/GameplayAbilities/Source/GameplayAbilities/Private/Abilities/Tasks/AbilityTask_WaitConfirmCancel.cpp
+++ b/Engine/Plugins/Runtime/GameplayAbilities/Source/GameplayAbilities/Private/Abilities/Tasks/AbilityTask_WaitConfirmCancel.cpp
@@ -15,7 +15,7 @@ UAbilityTask_WaitConfirmCancel::UAbilityTask_WaitConfirmCancel(const FObjectInit
 void UAbilityTask_WaitConfirmCancel::OnConfirmCallback()
 {
 	
-	if (AbilitySystemComponent)
+	if (AbilitySystemComponent.IsValid())
 	{
 		AbilitySystemComponent->ConsumeGenericReplicatedEvent(EAbilityGenericReplicatedEvent::GenericConfirm, GetAbilitySpecHandle(), GetActivationPredictionKey());
 		if (ShouldBroadcastAbilityTaskDelegates())
@@ -28,7 +28,7 @@ void UAbilityTask_WaitConfirmCancel::OnConfirmCallback()
 
 void UAbilityTask_WaitConfirmCancel::OnCancelCallback()
 {
-	if (AbilitySystemComponent)
+	if (AbilitySystemComponent.IsValid())
 	{
 		AbilitySystemComponent->ConsumeGenericReplicatedEvent(EAbilityGenericReplicatedEvent::GenericCancel, GetAbilitySpecHandle(), GetActivationPredictionKey());
 		if (ShouldBroadcastAbilityTaskDelegates())
@@ -41,9 +41,9 @@ void UAbilityTask_WaitConfirmCancel::OnCancelCallback()
 
 void UAbilityTask_WaitConfirmCancel::OnLocalConfirmCallback()
 {
-	FScopedPredictionWindow ScopedPrediction(AbilitySystemComponent, IsPredictingClient());
+	FScopedPredictionWindow ScopedPrediction(AbilitySystemComponent.Get(), IsPredictingClient());
 
-	if (AbilitySystemComponent && IsPredictingClient())
+	if (AbilitySystemComponent.IsValid() && IsPredictingClient())
 	{
 		AbilitySystemComponent->ServerSetReplicatedEvent(EAbilityGenericReplicatedEvent::GenericConfirm, GetAbilitySpecHandle(), GetActivationPredictionKey() ,AbilitySystemComponent->ScopedPredictionKey);
 	}
@@ -52,9 +52,9 @@ void UAbilityTask_WaitConfirmCancel::OnLocalConfirmCallback()
 
 void UAbilityTask_WaitConfirmCancel::OnLocalCancelCallback()
 {
-	FScopedPredictionWindow ScopedPrediction(AbilitySystemComponent, IsPredictingClient());
+	FScopedPredictionWindow ScopedPrediction(AbilitySystemComponent.Get(), IsPredictingClient());
 
-	if (AbilitySystemComponent && IsPredictingClient())
+	if (AbilitySystemComponent.IsValid() && IsPredictingClient())
 	{
 		AbilitySystemComponent->ServerSetReplicatedEvent(EAbilityGenericReplicatedEvent::GenericCancel, GetAbilitySpecHandle(), GetActivationPredictionKey() ,AbilitySystemComponent->ScopedPredictionKey);
 	}
@@ -68,7 +68,7 @@ UAbilityTask_WaitConfirmCancel* UAbilityTask_WaitConfirmCancel::WaitConfirmCance
 
 void UAbilityTask_WaitConfirmCancel::Activate()
 {
-	if (AbilitySystemComponent && Ability)
+	if (AbilitySystemComponent.IsValid() && Ability)
 	{
 		const FGameplayAbilityActorInfo* Info = Ability->GetCurrentActorInfo();
 
@@ -101,7 +101,7 @@ void UAbilityTask_WaitConfirmCancel::Activate()
 
 void UAbilityTask_WaitConfirmCancel::OnDestroy(bool AbilityEnding)
 {
-	if (RegisteredCallbacks && AbilitySystemComponent)
+	if (RegisteredCallbacks && AbilitySystemComponent.IsValid())
 	{
 		AbilitySystemComponent->GenericLocalConfirmCallbacks.RemoveDynamic(this, &UAbilityTask_WaitConfirmCancel::OnLocalConfirmCallback);
 		AbilitySystemComponent->GenericLocalCancelCallbacks.RemoveDynamic(this, &UAbilityTask_WaitConfirmCancel::OnLocalCancelCallback);
diff --git a/Engine/Plugins/Runtime/GameplayAbilities/Source/GameplayAbilities/Private/Abilities/Tasks/AbilityTask_WaitGameplayEffectApplied.cpp b/Engine/Plugins/Runtime/GameplayAbilities/Source/GameplayAbilities/Private/Abilities/Tasks/AbilityTask_WaitGameplayEffectApplied.cpp
--- a/Engine/Plugins/Runtime/GameplayAbilities/Source/GameplayAbilities/Private/Abilities/Tasks/AbilityTask_WaitGameplayEffectApplied.cpp
+++ b/Engine/Plugins/Runtime/GameplayAbilities/Source/GameplayAbilities/Private/Abilities/Tasks/AbilityTask_WaitGameplayEffectApplied.cpp
@@ -97,5 +97,5 @@ UAbilitySystemComponent* UAbilityTask_WaitGameplayEffectApplied::GetASC()
 	{
 		return ExternalOwner;
 	}
-	return AbilitySystemComponent;
+	return AbilitySystemComponent.Get();
 }
diff --git a/Engine/Plugins/Runtime/GameplayAbilities/Source/GameplayAbilities/Private/Abilities/Tasks/AbilityTask_WaitGameplayEffectBlockedImmunity.cpp b/Engine/Plugins/Runtime/GameplayAbilities/Source/GameplayAbilities/Private/Abilities/Tasks/AbilityTask_WaitGameplayEffectBlockedImmunity.cpp
--- a/Engine/Plugins/Runtime/GameplayAbilities/Source/GameplayAbilities/Private/Abilities/Tasks/AbilityTask_WaitGameplayEffectBlockedImmunity.cpp
+++ b/Engine/Plugins/Runtime/GameplayAbilities/Source/GameplayAbilities/Private/Abilities/Tasks/AbilityTask_WaitGameplayEffectBlockedImmunity.cpp
@@ -81,7 +81,7 @@ UAbilitySystemComponent* UAbilityTask_WaitGameplayEffectBlockedImmunity::GetASC(
 	{
 		return ExternalOwner;
 	}
-	return AbilitySystemComponent;
+	return AbilitySystemComponent.Get();
 }
 
 void UAbilityTask_WaitGameplayEffectBlockedImmunity::RegisterDelegate()
diff --git a/Engine/Plugins/Runtime/GameplayAbilities/Source/GameplayAbilities/Private/Abilities/Tasks/AbilityTask_WaitGameplayEvent.cpp b/Engine/Plugins/Runtime/GameplayAbilities/Source/GameplayAbilities/Private/Abilities/Tasks/AbilityTask_WaitGameplayEvent.cpp
--- a/Engine/Plugins/Runtime/GameplayAbilities/Source/GameplayAbilities/Private/Abilities/Tasks/AbilityTask_WaitGameplayEvent.cpp
+++ b/Engine/Plugins/Runtime/GameplayAbilities/Source/GameplayAbilities/Private/Abilities/Tasks/AbilityTask_WaitGameplayEvent.cpp
@@ -76,7 +76,7 @@ UAbilitySystemComponent* UAbilityTask_WaitGameplayEvent::GetTargetASC()
 		return OptionalExternalTarget;
 	}
 
-	return AbilitySystemComponent;
+	return AbilitySystemComponent.Get();
 }
 
 void UAbilityTask_WaitGameplayEvent::OnDestroy(bool AbilityEnding)
diff --git a/Engine/Plugins/Runtime/GameplayAbilities/Source/GameplayAbilities/Private/Abilities/Tasks/AbilityTask_WaitGameplayTagBase.cpp b/Engine/Plugins/Runtime/GameplayAbilities/Source/GameplayAbilities/Private/Abilities/Tasks/AbilityTask_WaitGameplayTagBase.cpp
--- a/Engine/Plugins/Runtime/GameplayAbilities/Source/GameplayAbilities/Private/Abilities/Tasks/AbilityTask_WaitGameplayTagBase.cpp
+++ b/Engine/Plugins/Runtime/GameplayAbilities/Source/GameplayAbilities/Private/Abilities/Tasks/AbilityTask_WaitGameplayTagBase.cpp
@@ -44,7 +44,7 @@ UAbilitySystemComponent* UAbilityTask_WaitGameplayTag::GetTargetASC()
 		return OptionalExternalTarget;
 	}
 
-	return AbilitySystemComponent;
+	return AbilitySystemComponent.Get();
 }
 
 void UAbilityTask_WaitGameplayTag::SetExternalTarget(AActor* Actor)
diff --git a/Engine/Plugins/Runtime/GameplayAbilities/Source/GameplayAbilities/Private/Abilities/Tasks/AbilityTask_WaitInputPress.cpp b/Engine/Plugins/Runtime/GameplayAbilities/Source/GameplayAbilities/Private/Abilities/Tasks/AbilityTask_WaitInputPress.cpp
--- a/Engine/Plugins/Runtime/GameplayAbilities/Source/GameplayAbilities/Private/Abilities/Tasks/AbilityTask_WaitInputPress.cpp
+++ b/Engine/Plugins/Runtime/GameplayAbilities/Source/GameplayAbilities/Private/Abilities/Tasks/AbilityTask_WaitInputPress.cpp
@@ -14,14 +14,14 @@ void UAbilityTask_WaitInputPress::OnPressCallback()
 {
 	float ElapsedTime = GetWorld()->GetTimeSeconds() - StartTime;
 
-	if (!Ability || !AbilitySystemComponent)
+	if (!Ability || !AbilitySystemComponent.IsValid())
 	{
 		return;
 	}
 
 	AbilitySystemComponent->AbilityReplicatedEventDelegate(EAbilityGenericReplicatedEvent::InputPressed, GetAbilitySpecHandle(), GetActivationPredictionKey()).Remove(DelegateHandle);
 
-	FScopedPredictionWindow ScopedPrediction(AbilitySystemComponent, IsPredictingClient());
+	FScopedPredictionWindow ScopedPrediction(AbilitySystemComponent.Get(), IsPredictingClient());
 
 	if (IsPredictingClient())
 	{
diff --git a/Engine/Plugins/Runtime/GameplayAbilities/Source/GameplayAbilities/Private/Abilities/Tasks/AbilityTask_WaitInputRelease.cpp b/Engine/Plugins/Runtime/GameplayAbilities/Source/GameplayAbilities/Private/Abilities/Tasks/AbilityTask_WaitInputRelease.cpp
--- a/Engine/Plugins/Runtime/GameplayAbilities/Source/GameplayAbilities/Private/Abilities/Tasks/AbilityTask_WaitInputRelease.cpp
+++ b/Engine/Plugins/Runtime/GameplayAbilities/Source/GameplayAbilities/Private/Abilities/Tasks/AbilityTask_WaitInputRelease.cpp
@@ -14,14 +14,14 @@ void UAbilityTask_WaitInputRelease::OnReleaseCallback()
 {
 	float ElapsedTime = GetWorld()->GetTimeSeconds() - StartTime;
 
-	if (!Ability || !AbilitySystemComponent)
+	if (!Ability || !AbilitySystemComponent.IsValid())
 	{
 		return;
 	}
 
 	AbilitySystemComponent->AbilityReplicatedEventDelegate(EAbilityGenericReplicatedEvent::InputReleased, GetAbilitySpecHandle(), GetActivationPredictionKey()).Remove(DelegateHandle);
 
-	FScopedPredictionWindow ScopedPrediction(AbilitySystemComponent, IsPredictingClient());
+	FScopedPredictionWindow ScopedPrediction(AbilitySystemComponent.Get(), IsPredictingClient());
 
 	if (IsPredictingClient())
 	{
diff --git a/Engine/Plugins/Runtime/GameplayAbilities/Source/GameplayAbilities/Private/Abilities/Tasks/AbilityTask_WaitTargetData.cpp b/Engine/Plugins/Runtime/GameplayAbilities/Source/GameplayAbilities/Private/Abilities/Tasks/AbilityTask_WaitTargetData.cpp
--- a/Engine/Plugins/Runtime/GameplayAbilities/Source/GameplayAbilities/Private/Abilities/Tasks/AbilityTask_WaitTargetData.cpp
+++ b/Engine/Plugins/Runtime/GameplayAbilities/Source/GameplayAbilities/Private/Abilities/Tasks/AbilityTask_WaitTargetData.cpp
@@ -209,7 +209,7 @@ void UAbilityTask_WaitTargetData::RegisterTargetDataCallbacks()
 /** Valid TargetData was replicated to use (we are server, was sent from client) */
 void UAbilityTask_WaitTargetData::OnTargetDataReplicatedCallback(const FGameplayAbilityTargetDataHandle& Data, FGameplayTag ActivationTag)
 {
-	check(AbilitySystemComponent);
+	check(AbilitySystemComponent.IsValid());
 
 	FGameplayAbilityTargetDataHandle MutableData = Data;
 	AbilitySystemComponent->ConsumeClientReplicatedTargetData(GetAbilitySpecHandle(), GetActivationPredictionKey());
@@ -247,7 +247,7 @@ void UAbilityTask_WaitTargetData::OnTargetDataReplicatedCallback(const FGameplay
 /** Client canceled this Targeting Task (we are the server) */
 void UAbilityTask_WaitTargetData::OnTargetDataReplicatedCancelledCallback()
 {
-	check(AbilitySystemComponent);
+	check(AbilitySystemComponent.IsValid());
 	if (ShouldBroadcastAbilityTaskDelegates())
 	{
 		Cancelled.Broadcast(FGameplayAbilityTargetDataHandle());
@@ -258,13 +258,13 @@ void UAbilityTask_WaitTargetData::OnTargetDataReplicatedCancelledCallback()
 /** The TargetActor we spawned locally has called back with valid target data */
 void UAbilityTask_WaitTargetData::OnTargetDataReadyCallback(const FGameplayAbilityTargetDataHandle& Data)
 {
-	check(AbilitySystemComponent);
+	check(AbilitySystemComponent.IsValid());
 	if (!Ability)
 	{
 		return;
 	}
 
-	FScopedPredictionWindow	ScopedPrediction(AbilitySystemComponent, ShouldReplicateDataToServer());
+	FScopedPredictionWindow	ScopedPrediction(AbilitySystemComponent.Get(), ShouldReplicateDataToServer());
 	
 	const FGameplayAbilityActorInfo* Info = Ability->GetCurrentActorInfo();
 	if (IsPredictingClient())
@@ -295,9 +295,9 @@ void UAbilityTask_WaitTargetData::OnTargetDataReadyCallback(const FGameplayAbili
 /** The TargetActor we spawned locally has called back with a cancel event (they still include the 'last/best' targetdata but the consumer of this may want to discard it) */
 void UAbilityTask_WaitTargetData::OnTargetDataCancelledCallback(const FGameplayAbilityTargetDataHandle& Data)
 {
-	check(AbilitySystemComponent);
+	check(AbilitySystemComponent.IsValid());
 
-	FScopedPredictionWindow ScopedPrediction(AbilitySystemComponent, IsPredictingClient());
+	FScopedPredictionWindow ScopedPrediction(AbilitySystemComponent.Get(), IsPredictingClient());
 
 	if (IsPredictingClient())
 	{
@@ -318,7 +318,7 @@ void UAbilityTask_WaitTargetData::OnTargetDataCancelledCallback(const FGameplayA
 /** Called when the ability is asked to confirm from an outside node. What this means depends on the individual task. By default, this does nothing other than ending if bEndTask is true. */
 void UAbilityTask_WaitTargetData::ExternalConfirm(bool bEndTask)
 {
-	check(AbilitySystemComponent);
+	check(AbilitySystemComponent.IsValid());
 	if (TargetActor)
 	{
 		if (TargetActor->ShouldProduceTargetData())
@@ -332,7 +332,7 @@ void UAbilityTask_WaitTargetData::ExternalConfirm(bool bEndTask)
 /** Called when the ability is asked to confirm from an outside node. What this means depends on the individual task. By default, this does nothing other than ending if bEndTask is true. */
 void UAbilityTask_WaitTargetData::ExternalCancel()
 {
-	check(AbilitySystemComponent);
+	check(AbilitySystemComponent.IsValid());
 	if (ShouldBroadcastAbilityTaskDelegates())
 	{
 		Cancelled.Broadcast(FGameplayAbilityTargetDataHandle());
diff --git a/Engine/Plugins/Runtime/GameplayAbilities/Source/GameplayAbilities/Public/Abilities/Tasks/AbilityTask.h b/Engine/Plugins/Runtime/GameplayAbilities/Source/GameplayAbilities/Public/Abilities/Tasks/AbilityTask.h
--- a/Engine/Plugins/Runtime/GameplayAbilities/Source/GameplayAbilities/Public/Abilities/Tasks/AbilityTask.h
+++ b/Engine/Plugins/Runtime/GameplayAbilities/Source/GameplayAbilities/Public/Abilities/Tasks/AbilityTask.h
@@ -102,7 +102,7 @@ class GAMEPLAYABILITIES_API UAbilityTask : public UGameplayTask
 	UGameplayAbility* Ability;
 
 	UPROPERTY()
-	UAbilitySystemComponent* AbilitySystemComponent;
+	TWeakObjectPtr<UAbilitySystemComponent> AbilitySystemComponent;
 
 	/** Returns true if the ability is a locally predicted ability running on a client. Usually this means we need to tell the server something. */
 	bool IsPredictingClient() const;
