[KillPendingKill] Updating the only AbilitySystemComponetn reference to use Get() as it is a weakptr now.

diff --git a/Engine/Plugins/Runtime/GameplayAbilities/Source/GameplayAbilities/Private/Abilities/Tasks/AbilityTask_WaitGameplayTagQuery.cpp b/Engine/Plugins/Runtime/GameplayAbilities/Source/GameplayAbilities/Private/Abilities/Tasks/AbilityTask_WaitGameplayTagQuery.cpp
--- a/Engine/Plugins/Runtime/GameplayAbilities/Source/GameplayAbilities/Private/Abilities/Tasks/AbilityTask_WaitGameplayTagQuery.cpp
+++ b/Engine/Plugins/Runtime/GameplayAbilities/Source/GameplayAbilities/Private/Abilities/Tasks/AbilityTask_WaitGameplayTagQuery.cpp
@@ -105,7 +105,7 @@ void UAbilityTask_WaitGameplayTagQuery::OnDestroy(bool AbilityIsEnding)
 
 UAbilitySystemComponent* UAbilityTask_WaitGameplayTagQuery::GetTargetASC()
 {	
-	return bUseExternalTarget ? OptionalExternalTarget : AbilitySystemComponent;
+	return bUseExternalTarget ? OptionalExternalTarget : AbilitySystemComponent.Get();
 }
 
 void UAbilityTask_WaitGameplayTagQuery::SetExternalTarget(const AActor* Actor)
