Fix includes for GameplayAbilityTypes
Use UE_ENABLE_INCLUDE_ORDER_DEPRECATED_IN_5_1

diff --git a/Engine/Plugins/Runtime/GameplayAbilities/Source/GameplayAbilities/Private/GameplayAbilityTargetTypes.cpp b/Engine/Plugins/Runtime/GameplayAbilities/Source/GameplayAbilities/Private/GameplayAbilityTargetTypes.cpp
--- a/Engine/Plugins/Runtime/GameplayAbilities/Source/GameplayAbilities/Private/GameplayAbilityTargetTypes.cpp
+++ b/Engine/Plugins/Runtime/GameplayAbilities/Source/GameplayAbilities/Private/GameplayAbilityTargetTypes.cpp
@@ -5,7 +5,7 @@
 #include "AbilitySystemComponent.h"
 #include "AbilitySystemBlueprintLibrary.h"
 #include "AbilitySystemGlobals.h"
-
+#include "Components/MeshComponent.h"
 
 TArray<FActiveGameplayEffectHandle> FGameplayAbilityTargetData::ApplyGameplayEffect(const UGameplayEffect* GameplayEffect, const FGameplayEffectContextHandle& InEffectContext, float Level, FPredictionKey PredictionKey)
 {
@@ -79,6 +79,35 @@ FString FGameplayAbilityTargetData::ToString() const
 	return TEXT("BASE CLASS");
 }
 
+FTransform FGameplayAbilityTargetingLocationInfo::GetTargetingTransform() const
+{
+	//Return or calculate based on LocationType.
+	switch (LocationType)
+	{
+	case EGameplayAbilityTargetingLocationType::ActorTransform:
+		if (SourceActor)
+		{
+			return SourceActor->GetTransform();
+		}
+		break;
+	case EGameplayAbilityTargetingLocationType::SocketTransform:
+		if (SourceComponent)
+		{
+			// Bad socket name will just return component transform anyway, so we're safe
+			return SourceComponent->GetSocketTransform(SourceSocketName);
+		}
+		break;
+	case EGameplayAbilityTargetingLocationType::LiteralTransform:
+		return LiteralTransform;
+	default:
+		check(false);
+		break;
+	}
+
+	// It cannot get here
+	return FTransform::Identity;
+}
+
 FGameplayAbilityTargetDataHandle FGameplayAbilityTargetingLocationInfo::MakeTargetDataHandleFromHitResult(TWeakObjectPtr<UGameplayAbility> Ability, const FHitResult& HitResult) const
 {
 	TArray<FHitResult> HitResults;
diff --git a/Engine/Plugins/Runtime/GameplayAbilities/Source/GameplayAbilities/Private/GameplayAbilityTypes.cpp b/Engine/Plugins/Runtime/GameplayAbilities/Source/GameplayAbilities/Private/GameplayAbilityTypes.cpp
--- a/Engine/Plugins/Runtime/GameplayAbilities/Source/GameplayAbilities/Private/GameplayAbilityTypes.cpp
+++ b/Engine/Plugins/Runtime/GameplayAbilities/Source/GameplayAbilities/Private/GameplayAbilityTypes.cpp
@@ -2,11 +2,13 @@
 
 #include "Abilities/GameplayAbilityTypes.h"
 
-#include "GameFramework/PlayerController.h"
-#include "GameFramework/MovementComponent.h"
 #include "Abilities/GameplayAbility.h"
 #include "AbilitySystemComponent.h"
+#include "Animation/AnimInstance.h"
 #include "Animation/AnimMontage.h"
+#include "Components/SkeletalMeshComponent.h"
+#include "GameFramework/MovementComponent.h"
+#include "GameFramework/PlayerController.h"
 #include "GameplayPrediction.h"
 #include "Misc/NetworkVersion.h"
 
@@ -76,6 +78,26 @@ void FGameplayAbilityActorInfo::ClearActorInfo()
 	MovementComponent = nullptr;
 }
 
+UAnimInstance* FGameplayAbilityActorInfo::GetAnimInstance() const
+{ 
+	const USkeletalMeshComponent* SKMC = SkeletalMeshComponent.Get();
+
+	if (SKMC)
+	{
+		if (AffectedAnimInstanceTag != NAME_None)
+		{
+			if(UAnimInstance* Instance = SKMC->GetAnimInstance())
+			{
+				return Instance->GetLinkedAnimGraphInstanceByTag(AffectedAnimInstanceTag);
+			}
+		}
+
+		return SKMC->GetAnimInstance();
+	}
+
+	return nullptr;
+}
+
 bool FGameplayAbilityActorInfo::IsLocallyControlled() const
 {
 	if (const APlayerController* PC = PlayerController.Get())
diff --git a/Engine/Plugins/Runtime/GameplayAbilities/Source/GameplayAbilities/Public/Abilities/GameplayAbilityTargetTypes.h b/Engine/Plugins/Runtime/GameplayAbilities/Source/GameplayAbilities/Public/Abilities/GameplayAbilityTargetTypes.h
--- a/Engine/Plugins/Runtime/GameplayAbilities/Source/GameplayAbilities/Public/Abilities/GameplayAbilityTargetTypes.h
+++ b/Engine/Plugins/Runtime/GameplayAbilities/Source/GameplayAbilities/Public/Abilities/GameplayAbilityTargetTypes.h
@@ -10,11 +10,14 @@
 #include "GameplayTagContainer.h"
 #include "GameplayEffectTypes.h"
 #include "GameplayPrediction.h"
+#if UE_ENABLE_INCLUDE_ORDER_DEPRECATED_IN_5_1
 #include "Components/MeshComponent.h"
+#endif
 #include "GameplayAbilityTargetTypes.generated.h"
 
 class UGameplayAbility;
 class UGameplayEffect;
+class UMeshComponent;
 struct FGameplayEffectSpec;
 
 UENUM(BlueprintType)
@@ -320,35 +323,7 @@ struct GAMEPLAYABILITIES_API FGameplayAbilityTargetingLocationInfo
 	}
 
 	/** Converts internal format into a literal world space transform */
-	FTransform GetTargetingTransform() const
-	{
-		//Return or calculate based on LocationType.
-		switch (LocationType)
-		{
-		case EGameplayAbilityTargetingLocationType::ActorTransform:
-			if (SourceActor)
-			{
-				return SourceActor->GetTransform();
-			}
-			break;
-		case EGameplayAbilityTargetingLocationType::SocketTransform:
-			if (SourceComponent)
-			{
-				// Bad socket name will just return component transform anyway, so we're safe
-				return SourceComponent->GetSocketTransform(SourceSocketName);
-			}
-			break;
-		case EGameplayAbilityTargetingLocationType::LiteralTransform:
-			return LiteralTransform;
-		default:
-			check(false);
-			break;
-		}
-
-		// It cannot get here
-		return FTransform::Identity;
-	}
-
+	FTransform GetTargetingTransform() const;
 	/** Initializes new target data and fills in with hit results */
 	FGameplayAbilityTargetDataHandle MakeTargetDataHandleFromHitResult(TWeakObjectPtr<UGameplayAbility> Ability, const FHitResult& HitResult) const;
 	FGameplayAbilityTargetDataHandle MakeTargetDataHandleFromHitResults(TWeakObjectPtr<UGameplayAbility> Ability, const TArray<FHitResult>& HitResults) const;
diff --git a/Engine/Plugins/Runtime/GameplayAbilities/Source/GameplayAbilities/Public/Abilities/GameplayAbilityTypes.h b/Engine/Plugins/Runtime/GameplayAbilities/Source/GameplayAbilities/Public/Abilities/GameplayAbilityTypes.h
--- a/Engine/Plugins/Runtime/GameplayAbilities/Source/GameplayAbilities/Public/Abilities/GameplayAbilityTypes.h
+++ b/Engine/Plugins/Runtime/GameplayAbilities/Source/GameplayAbilities/Public/Abilities/GameplayAbilityTypes.h
@@ -11,19 +11,23 @@
 #include "GameplayEffectTypes.h"
 #include "GameplayPrediction.h"
 #include "GameplayAbilitySpec.h"
+#if UE_ENABLE_INCLUDE_ORDER_DEPRECATED_IN_5_1
 #include "UObject/Package.h"
 #include "Animation/AnimInstance.h"
 #include "Components/SkeletalMeshComponent.h"
+#endif
 #include "Abilities/GameplayAbilityTargetTypes.h"
 #include "GameplayAbilityTypes.generated.h"
 
 class APlayerController;
 class UAbilitySystemComponent;
+class UAnimInstance;
 class UAnimMontage;
-class UGameplayAbility;
-class UMovementComponent;
 class UDataTable;
+class UGameplayAbility;
 class UGameplayTask;
+class UMovementComponent;
+class USkeletalMeshComponent;
 
 GAMEPLAYABILITIES_API DECLARE_LOG_CATEGORY_EXTERN(LogAbilitySystemComponent, Log, All);
 
@@ -172,26 +176,8 @@ struct GAMEPLAYABILITIES_API FGameplayAbilityActorInfo
 	FName AffectedAnimInstanceTag; 
 	
 	/** Accessor to get the affected anim instance from the SkeletalMeshComponent */
-	UAnimInstance* GetAnimInstance() const
-	{ 
-		const USkeletalMeshComponent* SKMC = SkeletalMeshComponent.Get();
-
-		if (SKMC)
-		{
-			if (AffectedAnimInstanceTag != NAME_None)
-			{
-				if(UAnimInstance* Instance = SKMC->GetAnimInstance())
-				{
-					return Instance->GetLinkedAnimGraphInstanceByTag(AffectedAnimInstanceTag);
-				}
-			}
-
-			return SKMC->GetAnimInstance();
-		}
-
-		return nullptr;
-	}
-
+	UAnimInstance* GetAnimInstance() const;
+	
 	/** Returns true if this actor is locally controlled. Only true for players on the client that owns them */
 	bool IsLocallyControlled() const;
 
