Fix non unity build errors
#test compile cpp

diff --git a/Engine/Plugins/Runtime/GameplayAbilities/Source/GameplayAbilities/Private/Abilities/GameplayAbilityTargetActor_Radius.cpp b/Engine/Plugins/Runtime/GameplayAbilities/Source/GameplayAbilities/Private/Abilities/GameplayAbilityTargetActor_Radius.cpp
--- a/Engine/Plugins/Runtime/GameplayAbilities/Source/GameplayAbilities/Private/Abilities/GameplayAbilityTargetActor_Radius.cpp
+++ b/Engine/Plugins/Runtime/GameplayAbilities/Source/GameplayAbilities/Private/Abilities/GameplayAbilityTargetActor_Radius.cpp
@@ -4,6 +4,7 @@
 #include "GameFramework/Pawn.h"
 #include "WorldCollision.h"
 #include "Abilities/GameplayAbility.h"
+#include "Engine/World.h"
 
 // --------------------------------------------------------------------------------------------------------------------------------------------------------
 //
diff --git a/Engine/Plugins/Runtime/GameplayAbilities/Source/GameplayAbilities/Private/Abilities/Tasks/AbilityTask_Repeat.cpp b/Engine/Plugins/Runtime/GameplayAbilities/Source/GameplayAbilities/Private/Abilities/Tasks/AbilityTask_Repeat.cpp
--- a/Engine/Plugins/Runtime/GameplayAbilities/Source/GameplayAbilities/Private/Abilities/Tasks/AbilityTask_Repeat.cpp
+++ b/Engine/Plugins/Runtime/GameplayAbilities/Source/GameplayAbilities/Private/Abilities/Tasks/AbilityTask_Repeat.cpp
@@ -2,6 +2,7 @@
 
 #include "Abilities/Tasks/AbilityTask_Repeat.h"
 #include "TimerManager.h"
+#include "Engine/World.h"
 
 UAbilityTask_Repeat::UAbilityTask_Repeat(const FObjectInitializer& ObjectInitializer)
 : Super(ObjectInitializer)
diff --git a/Engine/Plugins/Runtime/GameplayAbilities/Source/GameplayAbilities/Private/Abilities/Tasks/AbilityTask_WaitAttributeChangeThreshold.cpp b/Engine/Plugins/Runtime/GameplayAbilities/Source/GameplayAbilities/Private/Abilities/Tasks/AbilityTask_WaitAttributeChangeThreshold.cpp
--- a/Engine/Plugins/Runtime/GameplayAbilities/Source/GameplayAbilities/Private/Abilities/Tasks/AbilityTask_WaitAttributeChangeThreshold.cpp
+++ b/Engine/Plugins/Runtime/GameplayAbilities/Source/GameplayAbilities/Private/Abilities/Tasks/AbilityTask_WaitAttributeChangeThreshold.cpp
@@ -4,7 +4,7 @@
 
 #include "AbilitySystemComponent.h"
 #include "AbilitySystemGlobals.h"
-
+#include "Engine/World.h"
 
 UAbilityTask_WaitAttributeChangeThreshold::UAbilityTask_WaitAttributeChangeThreshold(const FObjectInitializer& ObjectInitializer)
 	: Super(ObjectInitializer)
diff --git a/Engine/Plugins/Runtime/GameplayAbilities/Source/GameplayAbilities/Private/Abilities/Tasks/AbilityTask_WaitDelay.cpp b/Engine/Plugins/Runtime/GameplayAbilities/Source/GameplayAbilities/Private/Abilities/Tasks/AbilityTask_WaitDelay.cpp
--- a/Engine/Plugins/Runtime/GameplayAbilities/Source/GameplayAbilities/Private/Abilities/Tasks/AbilityTask_WaitDelay.cpp
+++ b/Engine/Plugins/Runtime/GameplayAbilities/Source/GameplayAbilities/Private/Abilities/Tasks/AbilityTask_WaitDelay.cpp
@@ -3,6 +3,7 @@
 #include "Abilities/Tasks/AbilityTask_WaitDelay.h"
 #include "TimerManager.h"
 #include "AbilitySystemGlobals.h"
+#include "Engine/World.h"
 
 UAbilityTask_WaitDelay::UAbilityTask_WaitDelay(const FObjectInitializer& ObjectInitializer)
 	: Super(ObjectInitializer)
diff --git a/Engine/Plugins/Runtime/GameplayAbilities/Source/GameplayAbilities/Private/Abilities/Tasks/AbilityTask_WaitInputPress.cpp b/Engine/Plugins/Runtime/GameplayAbilities/Source/GameplayAbilities/Private/Abilities/Tasks/AbilityTask_WaitInputPress.cpp
--- a/Engine/Plugins/Runtime/GameplayAbilities/Source/GameplayAbilities/Private/Abilities/Tasks/AbilityTask_WaitInputPress.cpp
+++ b/Engine/Plugins/Runtime/GameplayAbilities/Source/GameplayAbilities/Private/Abilities/Tasks/AbilityTask_WaitInputPress.cpp
@@ -2,6 +2,7 @@
 
 #include "Abilities/Tasks/AbilityTask_WaitInputPress.h"
 #include "AbilitySystemComponent.h"
+#include "Engine/World.h"
 
 UAbilityTask_WaitInputPress::UAbilityTask_WaitInputPress(const FObjectInitializer& ObjectInitializer)
 	: Super(ObjectInitializer)
diff --git a/Engine/Plugins/Runtime/GameplayAbilities/Source/GameplayAbilities/Private/Abilities/Tasks/AbilityTask_WaitInputRelease.cpp b/Engine/Plugins/Runtime/GameplayAbilities/Source/GameplayAbilities/Private/Abilities/Tasks/AbilityTask_WaitInputRelease.cpp
--- a/Engine/Plugins/Runtime/GameplayAbilities/Source/GameplayAbilities/Private/Abilities/Tasks/AbilityTask_WaitInputRelease.cpp
+++ b/Engine/Plugins/Runtime/GameplayAbilities/Source/GameplayAbilities/Private/Abilities/Tasks/AbilityTask_WaitInputRelease.cpp
@@ -2,6 +2,7 @@
 
 #include "Abilities/Tasks/AbilityTask_WaitInputRelease.h"
 #include "AbilitySystemComponent.h"
+#include "Engine/World.h"
 
 UAbilityTask_WaitInputRelease::UAbilityTask_WaitInputRelease(const FObjectInitializer& ObjectInitializer)
 	: Super(ObjectInitializer)
diff --git a/Engine/Plugins/Runtime/GameplayAbilities/Source/GameplayAbilities/Private/AbilitySystemBlueprintLibrary.cpp b/Engine/Plugins/Runtime/GameplayAbilities/Source/GameplayAbilities/Private/AbilitySystemBlueprintLibrary.cpp
--- a/Engine/Plugins/Runtime/GameplayAbilities/Source/GameplayAbilities/Private/AbilitySystemBlueprintLibrary.cpp
+++ b/Engine/Plugins/Runtime/GameplayAbilities/Source/GameplayAbilities/Private/AbilitySystemBlueprintLibrary.cpp
@@ -4,6 +4,7 @@
 #include "GameplayEffectAggregator.h"
 #include "AbilitySystemGlobals.h"
 #include "AbilitySystemComponent.h"
+#include "Engine/World.h"
 #include "GameplayEffect.h"
 #include "GameplayEffectUIData.h"
 #include "GameplayAbilitySpec.h"
diff --git a/Engine/Source/Editor/UnrealEd/Public/AnimationEditorUtils.h b/Engine/Source/Editor/UnrealEd/Public/AnimationEditorUtils.h
--- a/Engine/Source/Editor/UnrealEd/Public/AnimationEditorUtils.h
+++ b/Engine/Source/Editor/UnrealEd/Public/AnimationEditorUtils.h
@@ -16,6 +16,7 @@
 #include "Editor/ContentBrowser/Public/ContentBrowserModule.h"
 #include "Developer/AssetTools/Public/IAssetTools.h"
 #include "Developer/AssetTools/Public/AssetToolsModule.h"
+#include "Engine/SkeletalMesh.h"
 
 class FMenuBuilder;
 class UAnimBlueprint;
diff --git a/Engine/Source/Runtime/Engine/Private/Animation/AnimCompressionDerivedData.cpp b/Engine/Source/Runtime/Engine/Private/Animation/AnimCompressionDerivedData.cpp
--- a/Engine/Source/Runtime/Engine/Private/Animation/AnimCompressionDerivedData.cpp
+++ b/Engine/Source/Runtime/Engine/Private/Animation/AnimCompressionDerivedData.cpp
@@ -13,6 +13,7 @@
 #include "Animation/AnimCurveCompressionSettings.h"
 #include "AnimationCompression.h"
 #include "UObject/Package.h"
+#include "Misc/CoreDelegates.h"
 
 #if WITH_EDITOR
 
diff --git a/Engine/Source/Runtime/Engine/Private/CharacterMovementComponentAsync.cpp b/Engine/Source/Runtime/Engine/Private/CharacterMovementComponentAsync.cpp
--- a/Engine/Source/Runtime/Engine/Private/CharacterMovementComponentAsync.cpp
+++ b/Engine/Source/Runtime/Engine/Private/CharacterMovementComponentAsync.cpp
@@ -4,8 +4,10 @@
 #include "GameFramework/CharacterMovementComponent.h"
 #include "PhysicsProxy/SingleParticlePhysicsProxy.h"
 #include "Chaos/SpatialAccelerationCollection.h"
+#include "Components/PrimitiveComponent.h"
 #include "PBDRigidsSolver.h"
 #include "Engine/ScopedMovementUpdate.h"
+#include "Engine/World.h"
 
 void FCharacterMovementComponentAsyncInput::Simulate(const float DeltaSeconds, FCharacterMovementComponentAsyncOutput& Output) const
 {
diff --git a/Engine/Source/Runtime/Engine/Public/Animation/AnimCompressionDerivedDataPublic.h b/Engine/Source/Runtime/Engine/Public/Animation/AnimCompressionDerivedDataPublic.h
--- a/Engine/Source/Runtime/Engine/Public/Animation/AnimCompressionDerivedDataPublic.h
+++ b/Engine/Source/Runtime/Engine/Public/Animation/AnimCompressionDerivedDataPublic.h
@@ -7,6 +7,7 @@
 #include "Animation/AnimCompressionTypes.h"
 #include "TickableEditorObject.h"
 #include "Animation/AnimCompressionTypes.h"
+#include "UObject/GCObject.h"
 
 class UAnimSequence;
 class FDerivedDataAnimationCompression;
diff --git a/Engine/Source/Runtime/Engine/Public/Animation/AnimSync.h b/Engine/Source/Runtime/Engine/Public/Animation/AnimSync.h
--- a/Engine/Source/Runtime/Engine/Public/Animation/AnimSync.h
+++ b/Engine/Source/Runtime/Engine/Public/Animation/AnimSync.h
@@ -8,6 +8,7 @@
 
 struct FAnimationBaseContext;
 struct FAnimInstanceProxy;
+class UAnimInstance;
 class UMirrorDataTable; 
 
 namespace UE { namespace Anim {
diff --git a/Engine/Source/Runtime/Engine/Public/CharacterMovementComponentAsync.h b/Engine/Source/Runtime/Engine/Public/CharacterMovementComponentAsync.h
--- a/Engine/Source/Runtime/Engine/Public/CharacterMovementComponentAsync.h
+++ b/Engine/Source/Runtime/Engine/Public/CharacterMovementComponentAsync.h
@@ -3,6 +3,8 @@
 #pragma once
 #include "GameFramework/RootMotionSource.h"
 #include "Chaos/SimCallbackObject.h"
+#include "CollisionQueryParams.h"
+#include "CollisionShape.h"
 #include "Engine/OverlapInfo.h"
 #include "CharacterMovementComponentAsync.generated.h"
 
