- Remove WITH_NIAGARA define

#rnx

diff --git a/Engine/Plugins/Runtime/GameplayAbilities/Source/GameplayAbilities/GameplayAbilities.Build.cs b/Engine/Plugins/Runtime/GameplayAbilities/Source/GameplayAbilities/GameplayAbilities.Build.cs
--- a/Engine/Plugins/Runtime/GameplayAbilities/Source/GameplayAbilities/GameplayAbilities.Build.cs
+++ b/Engine/Plugins/Runtime/GameplayAbilities/Source/GameplayAbilities/GameplayAbilities.Build.cs
@@ -26,7 +26,6 @@ namespace UnrealBuildTool.Rules
 			// Niagara support for gameplay cue notifies.
 			{
 				PrivateDependencyModuleNames.Add("Niagara");
-				PublicDefinitions.Add("WITH_NIAGARA=1");
 			}
 
 			if (Target.bBuildEditor == true)
