Prevent a non-replicated ASC from registering replicated subobjects.
Renable ensure to detect this problem

diff --git a/Engine/Plugins/Runtime/GameplayAbilities/Source/GameplayAbilities/Private/AbilitySystemComponent.cpp b/Engine/Plugins/Runtime/GameplayAbilities/Source/GameplayAbilities/Private/AbilitySystemComponent.cpp
--- a/Engine/Plugins/Runtime/GameplayAbilities/Source/GameplayAbilities/Private/AbilitySystemComponent.cpp
+++ b/Engine/Plugins/Runtime/GameplayAbilities/Source/GameplayAbilities/Private/AbilitySystemComponent.cpp
@@ -241,6 +241,14 @@ void UAbilitySystemComponent::ReadyForReplication()
 	// Register the spawned attributes to the replicated sub object list.
 	if (IsUsingRegisteredSubObjectList())
 	{
+		for (UGameplayAbility* ReplicatedAbility : GetReplicatedInstancedAbilities_Mutable())
+		{
+			if (ReplicatedAbility)
+			{
+				AddReplicatedSubObject(ReplicatedAbility);
+			}
+		}
+
 		for (UAttributeSet* ReplicatedAttribute : SpawnedAttributes)
 		{
 			if (ReplicatedAttribute)
@@ -2951,7 +2959,7 @@ void UAbilitySystemComponent::AddReplicatedInstancedAbility(UGameplayAbility* Ga
 	{
 		ReplicatedAbilities.Add(GameplayAbility);
 		
-		if (IsUsingRegisteredSubObjectList())
+		if (IsUsingRegisteredSubObjectList() && IsReadyForReplication())
 		{
 			AddReplicatedSubObject(GameplayAbility);
 		}
@@ -2961,7 +2969,7 @@ void UAbilitySystemComponent::AddReplicatedInstancedAbility(UGameplayAbility* Ga
 void UAbilitySystemComponent::RemoveReplicatedInstancedAbility(UGameplayAbility* GameplayAbility)
 {
 	const bool bWasRemoved = GetReplicatedInstancedAbilities_Mutable().RemoveSingle(GameplayAbility) > 0;
-	if (bWasRemoved && IsUsingRegisteredSubObjectList())
+	if (bWasRemoved && IsUsingRegisteredSubObjectList() && IsReadyForReplication())
 	{
 		RemoveReplicatedSubObject(GameplayAbility);
 	}
@@ -2971,7 +2979,7 @@ void UAbilitySystemComponent::RemoveAllReplicatedInstancedAbilities()
 {
 	TArray<UGameplayAbility*>& ReplicatedAbilities = GetReplicatedInstancedAbilities_Mutable();
 
-	if (IsUsingRegisteredSubObjectList())
+	if (IsUsingRegisteredSubObjectList() && IsReadyForReplication())
 	{
 		for (UGameplayAbility* ReplicatedAbility : ReplicatedAbilities)
 		{
diff --git a/Engine/Source/Runtime/Engine/Private/ActorReplication.cpp b/Engine/Source/Runtime/Engine/Private/ActorReplication.cpp
--- a/Engine/Source/Runtime/Engine/Private/ActorReplication.cpp
+++ b/Engine/Source/Runtime/Engine/Private/ActorReplication.cpp
@@ -633,7 +633,7 @@ void AActor::AddActorComponentReplicatedSubObject(UActorComponent* OwnerComponen
 {
 	check(IsValid(OwnerComponent));
 	check(IsValid(SubObject));
-	//ensureMsgf(OwnerComponent->GetIsReplicated(), TEXT("Only components with replication enabled can register subobjects. %s::%s has replication disabled."), *GetName(), *SubObject->GetName());
+	ensureMsgf(OwnerComponent->GetIsReplicated(), TEXT("Only components with replication enabled can register subobjects. %s::%s has replication disabled."), *GetName(), *SubObject->GetName());
 	ensureMsgf(NetCondition != COND_Custom, TEXT("Custom netconditions do not work with SubObjects. %s::%s will not be replicated."), *GetName(), *SubObject->GetName());
 
 	FReplicatedComponentInfo* ComponentInfo = ReplicatedComponentsInfo.FindByKey(OwnerComponent);
