Fixed crash in UAbilityTask_PlayMontageAndWait
#tests Verified some abilities that use this task still function correctly

diff --git a/Engine/Plugins/Runtime/GameplayAbilities/Source/GameplayAbilities/Private/Abilities/Tasks/AbilityTask_PlayMontageAndWait.cpp b/Engine/Plugins/Runtime/GameplayAbilities/Source/GameplayAbilities/Private/Abilities/Tasks/AbilityTask_PlayMontageAndWait.cpp
--- a/Engine/Plugins/Runtime/GameplayAbilities/Source/GameplayAbilities/Private/Abilities/Tasks/AbilityTask_PlayMontageAndWait.cpp
+++ b/Engine/Plugins/Runtime/GameplayAbilities/Source/GameplayAbilities/Private/Abilities/Tasks/AbilityTask_PlayMontageAndWait.cpp
@@ -18,7 +18,10 @@ void UAbilityTask_PlayMontageAndWait::OnMontageBlendingOut(UAnimMontage* Montage
 	{
 		if (Montage == MontageToPlay)
 		{
-			AbilitySystemComponent->ClearAnimatingAbility(Ability);
+			if (UAbilitySystemComponent* ASC = AbilitySystemComponent.Get())
+			{
+				ASC->ClearAnimatingAbility(Ability);
+			}
 
 			// Reset AnimRootMotionTranslationScale
 			ACharacter* Character = Cast<ACharacter>(GetAvatarActor());
