[KillPendingKill] Fixing up missing Ability Task null pointer checks as it was assumed that the AbilitySystemCompnent member of AbilityTask was always valid, which is not the case anymore as it was converted to WeakPtr

diff --git a/Engine/Plugins/Runtime/GameplayAbilities/Source/GameplayAbilities/Private/Abilities/Tasks/AbilityTask_ApplyRootMotionConstantForce.cpp b/Engine/Plugins/Runtime/GameplayAbilities/Source/GameplayAbilities/Private/Abilities/Tasks/AbilityTask_ApplyRootMotionConstantForce.cpp
--- a/Engine/Plugins/Runtime/GameplayAbilities/Source/GameplayAbilities/Private/Abilities/Tasks/AbilityTask_ApplyRootMotionConstantForce.cpp
+++ b/Engine/Plugins/Runtime/GameplayAbilities/Source/GameplayAbilities/Private/Abilities/Tasks/AbilityTask_ApplyRootMotionConstantForce.cpp
@@ -50,9 +50,10 @@ UAbilityTask_ApplyRootMotionConstantForce* UAbilityTask_ApplyRootMotionConstantF
 
 void UAbilityTask_ApplyRootMotionConstantForce::SharedInitAndApply()
 {
-	if (AbilitySystemComponent->AbilityActorInfo->MovementComponent.IsValid())
+	UAbilitySystemComponent* ASC = AbilitySystemComponent.Get();
+	if (ASC && ASC->AbilityActorInfo->MovementComponent.IsValid())
 	{
-		MovementComponent = Cast<UCharacterMovementComponent>(AbilitySystemComponent->AbilityActorInfo->MovementComponent.Get());
+		MovementComponent = Cast<UCharacterMovementComponent>(ASC->AbilityActorInfo->MovementComponent.Get());
 		StartTime = GetWorld()->GetTimeSeconds();
 		EndTime = StartTime + Duration;
 
diff --git a/Engine/Plugins/Runtime/GameplayAbilities/Source/GameplayAbilities/Private/Abilities/Tasks/AbilityTask_ApplyRootMotionJumpForce.cpp b/Engine/Plugins/Runtime/GameplayAbilities/Source/GameplayAbilities/Private/Abilities/Tasks/AbilityTask_ApplyRootMotionJumpForce.cpp
--- a/Engine/Plugins/Runtime/GameplayAbilities/Source/GameplayAbilities/Private/Abilities/Tasks/AbilityTask_ApplyRootMotionJumpForce.cpp
+++ b/Engine/Plugins/Runtime/GameplayAbilities/Source/GameplayAbilities/Private/Abilities/Tasks/AbilityTask_ApplyRootMotionJumpForce.cpp
@@ -86,9 +86,10 @@ void UAbilityTask_ApplyRootMotionJumpForce::TriggerLanded()
 
 void UAbilityTask_ApplyRootMotionJumpForce::SharedInitAndApply()
 {
-	if (AbilitySystemComponent->AbilityActorInfo->MovementComponent.IsValid())
+	UAbilitySystemComponent* ASC = AbilitySystemComponent.Get();
+	if (ASC && ASC->AbilityActorInfo->MovementComponent.IsValid())
 	{
-		MovementComponent = Cast<UCharacterMovementComponent>(AbilitySystemComponent->AbilityActorInfo->MovementComponent.Get());
+		MovementComponent = Cast<UCharacterMovementComponent>(ASC->AbilityActorInfo->MovementComponent.Get());
 		StartTime = GetWorld()->GetTimeSeconds();
 		EndTime = StartTime + Duration;
 
diff --git a/Engine/Plugins/Runtime/GameplayAbilities/Source/GameplayAbilities/Private/Abilities/Tasks/AbilityTask_ApplyRootMotionMoveToActorForce.cpp b/Engine/Plugins/Runtime/GameplayAbilities/Source/GameplayAbilities/Private/Abilities/Tasks/AbilityTask_ApplyRootMotionMoveToActorForce.cpp
--- a/Engine/Plugins/Runtime/GameplayAbilities/Source/GameplayAbilities/Private/Abilities/Tasks/AbilityTask_ApplyRootMotionMoveToActorForce.cpp
+++ b/Engine/Plugins/Runtime/GameplayAbilities/Source/GameplayAbilities/Private/Abilities/Tasks/AbilityTask_ApplyRootMotionMoveToActorForce.cpp
@@ -101,9 +101,10 @@ void UAbilityTask_ApplyRootMotionMoveToActorForce::OnRep_TargetLocation()
 
 void UAbilityTask_ApplyRootMotionMoveToActorForce::SharedInitAndApply()
 {
-	if (AbilitySystemComponent->AbilityActorInfo->MovementComponent.IsValid())
+	UAbilitySystemComponent* ASC = AbilitySystemComponent.Get();
+	if (ASC && ASC->AbilityActorInfo->MovementComponent.IsValid())
 	{
-		MovementComponent = Cast<UCharacterMovementComponent>(AbilitySystemComponent->AbilityActorInfo->MovementComponent.Get());
+		MovementComponent = Cast<UCharacterMovementComponent>(ASC->AbilityActorInfo->MovementComponent.Get());
 		StartTime = GetWorld()->GetTimeSeconds();
 		EndTime = StartTime + Duration;
 
diff --git a/Engine/Plugins/Runtime/GameplayAbilities/Source/GameplayAbilities/Private/Abilities/Tasks/AbilityTask_ApplyRootMotionMoveToForce.cpp b/Engine/Plugins/Runtime/GameplayAbilities/Source/GameplayAbilities/Private/Abilities/Tasks/AbilityTask_ApplyRootMotionMoveToForce.cpp
--- a/Engine/Plugins/Runtime/GameplayAbilities/Source/GameplayAbilities/Private/Abilities/Tasks/AbilityTask_ApplyRootMotionMoveToForce.cpp
+++ b/Engine/Plugins/Runtime/GameplayAbilities/Source/GameplayAbilities/Private/Abilities/Tasks/AbilityTask_ApplyRootMotionMoveToForce.cpp
@@ -41,9 +41,10 @@ UAbilityTask_ApplyRootMotionMoveToForce* UAbilityTask_ApplyRootMotionMoveToForce
 
 void UAbilityTask_ApplyRootMotionMoveToForce::SharedInitAndApply()
 {
-	if (AbilitySystemComponent->AbilityActorInfo->MovementComponent.IsValid())
+	UAbilitySystemComponent* ASC = AbilitySystemComponent.Get();
+	if (ASC && ASC->AbilityActorInfo->MovementComponent.IsValid())
 	{
-		MovementComponent = Cast<UCharacterMovementComponent>(AbilitySystemComponent->AbilityActorInfo->MovementComponent.Get());
+		MovementComponent = Cast<UCharacterMovementComponent>(ASC->AbilityActorInfo->MovementComponent.Get());
 		StartTime = GetWorld()->GetTimeSeconds();
 		EndTime = StartTime + Duration;
 
diff --git a/Engine/Plugins/Runtime/GameplayAbilities/Source/GameplayAbilities/Private/Abilities/Tasks/AbilityTask_ApplyRootMotionRadialForce.cpp b/Engine/Plugins/Runtime/GameplayAbilities/Source/GameplayAbilities/Private/Abilities/Tasks/AbilityTask_ApplyRootMotionRadialForce.cpp
--- a/Engine/Plugins/Runtime/GameplayAbilities/Source/GameplayAbilities/Private/Abilities/Tasks/AbilityTask_ApplyRootMotionRadialForce.cpp
+++ b/Engine/Plugins/Runtime/GameplayAbilities/Source/GameplayAbilities/Private/Abilities/Tasks/AbilityTask_ApplyRootMotionRadialForce.cpp
@@ -45,9 +45,10 @@ UAbilityTask_ApplyRootMotionRadialForce* UAbilityTask_ApplyRootMotionRadialForce
 
 void UAbilityTask_ApplyRootMotionRadialForce::SharedInitAndApply()
 {
-	if (AbilitySystemComponent->AbilityActorInfo->MovementComponent.IsValid())
+	UAbilitySystemComponent* ASC = AbilitySystemComponent.Get();
+	if (ASC && ASC->AbilityActorInfo->MovementComponent.IsValid())
 	{
-		MovementComponent = Cast<UCharacterMovementComponent>(AbilitySystemComponent->AbilityActorInfo->MovementComponent.Get());
+		MovementComponent = Cast<UCharacterMovementComponent>(ASC->AbilityActorInfo->MovementComponent.Get());
 		StartTime = GetWorld()->GetTimeSeconds();
 		EndTime = StartTime + Duration;
 
diff --git a/Engine/Plugins/Runtime/GameplayAbilities/Source/GameplayAbilities/Private/Abilities/Tasks/AbilityTask_PlayMontageAndWait.cpp b/Engine/Plugins/Runtime/GameplayAbilities/Source/GameplayAbilities/Private/Abilities/Tasks/AbilityTask_PlayMontageAndWait.cpp
--- a/Engine/Plugins/Runtime/GameplayAbilities/Source/GameplayAbilities/Private/Abilities/Tasks/AbilityTask_PlayMontageAndWait.cpp
+++ b/Engine/Plugins/Runtime/GameplayAbilities/Source/GameplayAbilities/Private/Abilities/Tasks/AbilityTask_PlayMontageAndWait.cpp
@@ -101,13 +101,13 @@ void UAbilityTask_PlayMontageAndWait::Activate()
 
 	bool bPlayedMontage = false;
 
-	if (AbilitySystemComponent.IsValid())
+	if (UAbilitySystemComponent* ASC = AbilitySystemComponent.Get())
 	{
 		const FGameplayAbilityActorInfo* ActorInfo = Ability->GetCurrentActorInfo();
 		UAnimInstance* AnimInstance = ActorInfo->GetAnimInstance();
 		if (AnimInstance != nullptr)
 		{
-			if (AbilitySystemComponent->PlayMontage(Ability, Ability->GetCurrentActivationInfo(), MontageToPlay, Rate, StartSection, StartTimeSeconds) > 0.f)
+			if (ASC->PlayMontage(Ability, Ability->GetCurrentActivationInfo(), MontageToPlay, Rate, StartSection, StartTimeSeconds) > 0.f)
 			{
 				// Playing a montage could potentially fire off a callback into game code which could kill this ability! Early out if we are  pending kill.
 				if (ShouldBroadcastAbilityTaskDelegates() == false)
@@ -157,8 +157,6 @@ void UAbilityTask_PlayMontageAndWait::Activate()
 
 void UAbilityTask_PlayMontageAndWait::ExternalCancel()
 {
-	check(AbilitySystemComponent.IsValid());
-
 	if (ShouldBroadcastAbilityTaskDelegates())
 	{
 		OnCancelled.Broadcast();
@@ -206,10 +204,11 @@ bool UAbilityTask_PlayMontageAndWait::StopPlayingMontage()
 
 	// Check if the montage is still playing
 	// The ability would have been interrupted, in which case we should automatically stop the montage
-	if (AbilitySystemComponent.IsValid() && Ability)
+	UAbilitySystemComponent* ASC = AbilitySystemComponent.Get();
+	if (ASC && Ability)
 	{
-		if (AbilitySystemComponent->GetAnimatingAbility() == Ability
-			&& AbilitySystemComponent->GetCurrentMontage() == MontageToPlay)
+		if (ASC->GetAnimatingAbility() == Ability
+			&& ASC->GetCurrentMontage() == MontageToPlay)
 		{
 			// Unbind delegates so they don't get called as well
 			FAnimMontageInstance* MontageInstance = AnimInstance->GetActiveInstanceForMontage(MontageToPlay);
@@ -219,7 +218,7 @@ bool UAbilityTask_PlayMontageAndWait::StopPlayingMontage()
 				MontageInstance->OnMontageEnded.Unbind();
 			}
 
-			AbilitySystemComponent->CurrentMontageStop();
+			ASC->CurrentMontageStop();
 			return true;
 		}
 	}
diff --git a/Engine/Plugins/Runtime/GameplayAbilities/Source/GameplayAbilities/Private/Abilities/Tasks/AbilityTask_SpawnActor.cpp b/Engine/Plugins/Runtime/GameplayAbilities/Source/GameplayAbilities/Private/Abilities/Tasks/AbilityTask_SpawnActor.cpp
--- a/Engine/Plugins/Runtime/GameplayAbilities/Source/GameplayAbilities/Private/Abilities/Tasks/AbilityTask_SpawnActor.cpp
+++ b/Engine/Plugins/Runtime/GameplayAbilities/Source/GameplayAbilities/Private/Abilities/Tasks/AbilityTask_SpawnActor.cpp
@@ -66,7 +66,10 @@ void UAbilityTask_SpawnActor::FinishSpawningActor(UGameplayAbility* OwningAbilit
 			}
 		if (!bTransformSet)
 		{
-			SpawnTransform = AbilitySystemComponent->GetOwner()->GetTransform();
+			if (UAbilitySystemComponent* ASC = AbilitySystemComponent.Get())
+			{
+				SpawnTransform = ASC->GetOwner()->GetTransform();
+			}
 		}
 
 		SpawnedActor->FinishSpawning(SpawnTransform);
diff --git a/Engine/Plugins/Runtime/GameplayAbilities/Source/GameplayAbilities/Private/Abilities/Tasks/AbilityTask_VisualizeTargeting.cpp b/Engine/Plugins/Runtime/GameplayAbilities/Source/GameplayAbilities/Private/Abilities/Tasks/AbilityTask_VisualizeTargeting.cpp
--- a/Engine/Plugins/Runtime/GameplayAbilities/Source/GameplayAbilities/Private/Abilities/Tasks/AbilityTask_VisualizeTargeting.cpp
+++ b/Engine/Plugins/Runtime/GameplayAbilities/Source/GameplayAbilities/Private/Abilities/Tasks/AbilityTask_VisualizeTargeting.cpp
@@ -93,11 +93,12 @@ bool UAbilityTask_VisualizeTargeting::BeginSpawningActor(UGameplayAbility* Ownin
 
 void UAbilityTask_VisualizeTargeting::FinishSpawningActor(UGameplayAbility* OwningAbility, AGameplayAbilityTargetActor* SpawnedActor)
 {
-	if (SpawnedActor)
+	UAbilitySystemComponent* ASC = AbilitySystemComponent.Get();
+	if (SpawnedActor && ASC)
 	{
 		check(TargetActor == SpawnedActor);
 
-		const FTransform SpawnTransform = AbilitySystemComponent->GetOwner()->GetTransform();
+		const FTransform SpawnTransform = ASC->GetOwner()->GetTransform();
 
 		SpawnedActor->FinishSpawning(SpawnTransform);
 
@@ -142,7 +143,10 @@ void UAbilityTask_VisualizeTargeting::FinalizeTargetActor(AGameplayAbilityTarget
 	check(SpawnedActor);
 	check(Ability);
 
-	AbilitySystemComponent->SpawnedTargetActors.Push(SpawnedActor);
+	if (UAbilitySystemComponent* ASC = AbilitySystemComponent.Get())
+	{
+		ASC->SpawnedTargetActors.Push(SpawnedActor);
+	}
 
 	SpawnedActor->StartTargeting(Ability);
 }
diff --git a/Engine/Plugins/Runtime/GameplayAbilities/Source/GameplayAbilities/Private/Abilities/Tasks/AbilityTask_WaitInputPress.cpp b/Engine/Plugins/Runtime/GameplayAbilities/Source/GameplayAbilities/Private/Abilities/Tasks/AbilityTask_WaitInputPress.cpp
--- a/Engine/Plugins/Runtime/GameplayAbilities/Source/GameplayAbilities/Private/Abilities/Tasks/AbilityTask_WaitInputPress.cpp
+++ b/Engine/Plugins/Runtime/GameplayAbilities/Source/GameplayAbilities/Private/Abilities/Tasks/AbilityTask_WaitInputPress.cpp
@@ -15,23 +15,24 @@ void UAbilityTask_WaitInputPress::OnPressCallback()
 {
 	float ElapsedTime = GetWorld()->GetTimeSeconds() - StartTime;
 
-	if (!Ability || !AbilitySystemComponent.IsValid())
+	UAbilitySystemComponent* ASC = AbilitySystemComponent.Get();
+	if (!Ability || !ASC)
 	{
 		return;
 	}
 
-	AbilitySystemComponent->AbilityReplicatedEventDelegate(EAbilityGenericReplicatedEvent::InputPressed, GetAbilitySpecHandle(), GetActivationPredictionKey()).Remove(DelegateHandle);
+	ASC->AbilityReplicatedEventDelegate(EAbilityGenericReplicatedEvent::InputPressed, GetAbilitySpecHandle(), GetActivationPredictionKey()).Remove(DelegateHandle);
 
-	FScopedPredictionWindow ScopedPrediction(AbilitySystemComponent.Get(), IsPredictingClient());
+	FScopedPredictionWindow ScopedPrediction(ASC, IsPredictingClient());
 
 	if (IsPredictingClient())
 	{
 		// Tell the server about this
-		AbilitySystemComponent->ServerSetReplicatedEvent(EAbilityGenericReplicatedEvent::InputPressed, GetAbilitySpecHandle(), GetActivationPredictionKey(), AbilitySystemComponent->ScopedPredictionKey);
+		ASC->ServerSetReplicatedEvent(EAbilityGenericReplicatedEvent::InputPressed, GetAbilitySpecHandle(), GetActivationPredictionKey(), ASC->ScopedPredictionKey);
 	}
 	else
 	{
-		AbilitySystemComponent->ConsumeGenericReplicatedEvent(EAbilityGenericReplicatedEvent::InputPressed, GetAbilitySpecHandle(), GetActivationPredictionKey());
+		ASC->ConsumeGenericReplicatedEvent(EAbilityGenericReplicatedEvent::InputPressed, GetAbilitySpecHandle(), GetActivationPredictionKey());
 	}
 
 	// We are done. Kill us so we don't keep getting broadcast messages
@@ -52,7 +53,9 @@ UAbilityTask_WaitInputPress* UAbilityTask_WaitInputPress::WaitInputPress(class U
 void UAbilityTask_WaitInputPress::Activate()
 {
 	StartTime = GetWorld()->GetTimeSeconds();
-	if (Ability)
+
+	UAbilitySystemComponent* ASC = AbilitySystemComponent.Get();
+	if (ASC && Ability)
 	{
 		if (bTestInitialState && IsLocallyControlled())
 		{
@@ -64,10 +67,10 @@ void UAbilityTask_WaitInputPress::Activate()
 			}
 		}
 
-		DelegateHandle = AbilitySystemComponent->AbilityReplicatedEventDelegate(EAbilityGenericReplicatedEvent::InputPressed, GetAbilitySpecHandle(), GetActivationPredictionKey()).AddUObject(this, &UAbilityTask_WaitInputPress::OnPressCallback);
+		DelegateHandle = ASC->AbilityReplicatedEventDelegate(EAbilityGenericReplicatedEvent::InputPressed, GetAbilitySpecHandle(), GetActivationPredictionKey()).AddUObject(this, &UAbilityTask_WaitInputPress::OnPressCallback);
 		if (IsForRemoteClient())
 		{
-			if (!AbilitySystemComponent->CallReplicatedEventDelegateIfSet(EAbilityGenericReplicatedEvent::InputPressed, GetAbilitySpecHandle(), GetActivationPredictionKey()))
+			if (!ASC->CallReplicatedEventDelegateIfSet(EAbilityGenericReplicatedEvent::InputPressed, GetAbilitySpecHandle(), GetActivationPredictionKey()))
 			{
 				SetWaitingOnRemotePlayerData();
 			}
diff --git a/Engine/Plugins/Runtime/GameplayAbilities/Source/GameplayAbilities/Private/Abilities/Tasks/AbilityTask_WaitInputRelease.cpp b/Engine/Plugins/Runtime/GameplayAbilities/Source/GameplayAbilities/Private/Abilities/Tasks/AbilityTask_WaitInputRelease.cpp
--- a/Engine/Plugins/Runtime/GameplayAbilities/Source/GameplayAbilities/Private/Abilities/Tasks/AbilityTask_WaitInputRelease.cpp
+++ b/Engine/Plugins/Runtime/GameplayAbilities/Source/GameplayAbilities/Private/Abilities/Tasks/AbilityTask_WaitInputRelease.cpp
@@ -15,23 +15,24 @@ void UAbilityTask_WaitInputRelease::OnReleaseCallback()
 {
 	float ElapsedTime = GetWorld()->GetTimeSeconds() - StartTime;
 
-	if (!Ability || !AbilitySystemComponent.IsValid())
+	UAbilitySystemComponent* ASC = AbilitySystemComponent.Get();
+	if (!Ability || !ASC)
 	{
 		return;
 	}
 
-	AbilitySystemComponent->AbilityReplicatedEventDelegate(EAbilityGenericReplicatedEvent::InputReleased, GetAbilitySpecHandle(), GetActivationPredictionKey()).Remove(DelegateHandle);
+	ASC->AbilityReplicatedEventDelegate(EAbilityGenericReplicatedEvent::InputReleased, GetAbilitySpecHandle(), GetActivationPredictionKey()).Remove(DelegateHandle);
 
-	FScopedPredictionWindow ScopedPrediction(AbilitySystemComponent.Get(), IsPredictingClient());
+	FScopedPredictionWindow ScopedPrediction(ASC, IsPredictingClient());
 
 	if (IsPredictingClient())
 	{
 		// Tell the server about this
-		AbilitySystemComponent->ServerSetReplicatedEvent(EAbilityGenericReplicatedEvent::InputReleased, GetAbilitySpecHandle(), GetActivationPredictionKey(), AbilitySystemComponent->ScopedPredictionKey);
+		ASC->ServerSetReplicatedEvent(EAbilityGenericReplicatedEvent::InputReleased, GetAbilitySpecHandle(), GetActivationPredictionKey(), ASC->ScopedPredictionKey);
 	}
 	else
 	{
-		AbilitySystemComponent->ConsumeGenericReplicatedEvent(EAbilityGenericReplicatedEvent::InputReleased, GetAbilitySpecHandle(), GetActivationPredictionKey());
+		ASC->ConsumeGenericReplicatedEvent(EAbilityGenericReplicatedEvent::InputReleased, GetAbilitySpecHandle(), GetActivationPredictionKey());
 	}
 
 	// We are done. Kill us so we don't keep getting broadcast messages
@@ -52,7 +53,9 @@ UAbilityTask_WaitInputRelease* UAbilityTask_WaitInputRelease::WaitInputRelease(c
 void UAbilityTask_WaitInputRelease::Activate()
 {
 	StartTime = GetWorld()->GetTimeSeconds();
-	if (Ability)
+
+	UAbilitySystemComponent* ASC = AbilitySystemComponent.Get();
+	if (ASC && Ability)
 	{
 		if (bTestInitialState && IsLocallyControlled())
 		{
@@ -64,10 +67,10 @@ void UAbilityTask_WaitInputRelease::Activate()
 			}
 		}
 
-		DelegateHandle = AbilitySystemComponent->AbilityReplicatedEventDelegate(EAbilityGenericReplicatedEvent::InputReleased, GetAbilitySpecHandle(), GetActivationPredictionKey()).AddUObject(this, &UAbilityTask_WaitInputRelease::OnReleaseCallback);
+		DelegateHandle = ASC->AbilityReplicatedEventDelegate(EAbilityGenericReplicatedEvent::InputReleased, GetAbilitySpecHandle(), GetActivationPredictionKey()).AddUObject(this, &UAbilityTask_WaitInputRelease::OnReleaseCallback);
 		if (IsForRemoteClient())
 		{
-			if (!AbilitySystemComponent->CallReplicatedEventDelegateIfSet(EAbilityGenericReplicatedEvent::InputReleased, GetAbilitySpecHandle(), GetActivationPredictionKey()))
+			if (!ASC->CallReplicatedEventDelegateIfSet(EAbilityGenericReplicatedEvent::InputReleased, GetAbilitySpecHandle(), GetActivationPredictionKey()))
 			{
 				SetWaitingOnRemotePlayerData();
 			}
diff --git a/Engine/Plugins/Runtime/GameplayAbilities/Source/GameplayAbilities/Private/Abilities/Tasks/AbilityTask_WaitTargetData.cpp b/Engine/Plugins/Runtime/GameplayAbilities/Source/GameplayAbilities/Private/Abilities/Tasks/AbilityTask_WaitTargetData.cpp
--- a/Engine/Plugins/Runtime/GameplayAbilities/Source/GameplayAbilities/Private/Abilities/Tasks/AbilityTask_WaitTargetData.cpp
+++ b/Engine/Plugins/Runtime/GameplayAbilities/Source/GameplayAbilities/Private/Abilities/Tasks/AbilityTask_WaitTargetData.cpp
@@ -102,11 +102,12 @@ bool UAbilityTask_WaitTargetData::BeginSpawningActor(UGameplayAbility* OwningAbi
 
 void UAbilityTask_WaitTargetData::FinishSpawningActor(UGameplayAbility* OwningAbility, AGameplayAbilityTargetActor* SpawnedActor)
 {
-	if (IsValid(SpawnedActor))
+	UAbilitySystemComponent* ASC = AbilitySystemComponent.Get();
+	if (ASC && IsValid(SpawnedActor))
 	{
 		check(TargetActor == SpawnedActor);
 
-		const FTransform SpawnTransform = AbilitySystemComponent->GetOwner()->GetTransform();
+		const FTransform SpawnTransform = ASC->GetOwner()->GetTransform();
 
 		SpawnedActor->FinishSpawning(SpawnTransform);
 
@@ -148,8 +149,11 @@ void UAbilityTask_WaitTargetData::FinalizeTargetActor(AGameplayAbilityTargetActo
 	check(SpawnedActor);
 	check(Ability);
 
-	// User ability activation is inhibited while this is active
-	AbilitySystemComponent->SpawnedTargetActors.Push(SpawnedActor);
+	if (UAbilitySystemComponent* ASC = AbilitySystemComponent.Get())
+	{
+		// User ability activation is inhibited while this is active
+		ASC->SpawnedTargetActors.Push(SpawnedActor);
+	}
 
 	SpawnedActor->StartTargeting(Ability);
 
@@ -177,6 +181,12 @@ void UAbilityTask_WaitTargetData::RegisterTargetDataCallbacks()
 		return;
 	}
 
+	UAbilitySystemComponent* ASC = AbilitySystemComponent.Get();
+	if (!ASC)
+	{
+		return;
+	}
+
 	check(TargetClass);
 	check(Ability);
 
@@ -196,10 +206,10 @@ void UAbilityTask_WaitTargetData::RegisterTargetDataCallbacks()
 			FPredictionKey ActivationPredictionKey = GetActivationPredictionKey();
 
 			//Since multifire is supported, we still need to hook up the callbacks
-			AbilitySystemComponent->AbilityTargetDataSetDelegate(SpecHandle, ActivationPredictionKey ).AddUObject(this, &UAbilityTask_WaitTargetData::OnTargetDataReplicatedCallback);
-			AbilitySystemComponent->AbilityTargetDataCancelledDelegate(SpecHandle, ActivationPredictionKey ).AddUObject(this, &UAbilityTask_WaitTargetData::OnTargetDataReplicatedCancelledCallback);
+			ASC->AbilityTargetDataSetDelegate(SpecHandle, ActivationPredictionKey ).AddUObject(this, &UAbilityTask_WaitTargetData::OnTargetDataReplicatedCallback);
+			ASC->AbilityTargetDataCancelledDelegate(SpecHandle, ActivationPredictionKey ).AddUObject(this, &UAbilityTask_WaitTargetData::OnTargetDataReplicatedCancelledCallback);
 
-			AbilitySystemComponent->CallReplicatedTargetDataDelegatesIfSet(SpecHandle, ActivationPredictionKey );
+			ASC->CallReplicatedTargetDataDelegatesIfSet(SpecHandle, ActivationPredictionKey );
 
 			SetWaitingOnRemotePlayerData();
 		}
@@ -209,10 +219,12 @@ void UAbilityTask_WaitTargetData::RegisterTargetDataCallbacks()
 /** Valid TargetData was replicated to use (we are server, was sent from client) */
 void UAbilityTask_WaitTargetData::OnTargetDataReplicatedCallback(const FGameplayAbilityTargetDataHandle& Data, FGameplayTag ActivationTag)
 {
-	check(AbilitySystemComponent.IsValid());
-
 	FGameplayAbilityTargetDataHandle MutableData = Data;
-	AbilitySystemComponent->ConsumeClientReplicatedTargetData(GetAbilitySpecHandle(), GetActivationPredictionKey());
+
+	if (UAbilitySystemComponent* ASC = AbilitySystemComponent.Get())
+	{
+		ASC->ConsumeClientReplicatedTargetData(GetAbilitySpecHandle(), GetActivationPredictionKey());
+	}
 
 	/** 
 	 *  Call into the TargetActor to sanitize/verify the data. If this returns false, we are rejecting
@@ -247,7 +259,6 @@ void UAbilityTask_WaitTargetData::OnTargetDataReplicatedCallback(const FGameplay
 /** Client canceled this Targeting Task (we are the server) */
 void UAbilityTask_WaitTargetData::OnTargetDataReplicatedCancelledCallback()
 {
-	check(AbilitySystemComponent.IsValid());
 	if (ShouldBroadcastAbilityTaskDelegates())
 	{
 		Cancelled.Broadcast(FGameplayAbilityTargetDataHandle());
@@ -258,13 +269,13 @@ void UAbilityTask_WaitTargetData::OnTargetDataReplicatedCancelledCallback()
 /** The TargetActor we spawned locally has called back with valid target data */
 void UAbilityTask_WaitTargetData::OnTargetDataReadyCallback(const FGameplayAbilityTargetDataHandle& Data)
 {
-	check(AbilitySystemComponent.IsValid());
-	if (!Ability)
+	UAbilitySystemComponent* ASC = AbilitySystemComponent.Get();
+	if (!Ability || !ASC)
 	{
 		return;
 	}
 
-	FScopedPredictionWindow	ScopedPrediction(AbilitySystemComponent.Get(), ShouldReplicateDataToServer());
+	FScopedPredictionWindow	ScopedPrediction(ASC, ShouldReplicateDataToServer());
 	
 	const FGameplayAbilityActorInfo* Info = Ability->GetCurrentActorInfo();
 	if (IsPredictingClient())
@@ -272,12 +283,12 @@ void UAbilityTask_WaitTargetData::OnTargetDataReadyCallback(const FGameplayAbili
 		if (!TargetActor->ShouldProduceTargetDataOnServer)
 		{
 			FGameplayTag ApplicationTag; // Fixme: where would this be useful?
-			AbilitySystemComponent->CallServerSetReplicatedTargetData(GetAbilitySpecHandle(), GetActivationPredictionKey(), Data, ApplicationTag, AbilitySystemComponent->ScopedPredictionKey);
+			ASC->CallServerSetReplicatedTargetData(GetAbilitySpecHandle(), GetActivationPredictionKey(), Data, ApplicationTag, ASC->ScopedPredictionKey);
 		}
 		else if (ConfirmationType == EGameplayTargetingConfirmation::UserConfirmed)
 		{
 			// We aren't going to send the target data, but we will send a generic confirmed message.
-			AbilitySystemComponent->ServerSetReplicatedEvent(EAbilityGenericReplicatedEvent::GenericConfirm, GetAbilitySpecHandle(), GetActivationPredictionKey(), AbilitySystemComponent->ScopedPredictionKey);
+			ASC->ServerSetReplicatedEvent(EAbilityGenericReplicatedEvent::GenericConfirm, GetAbilitySpecHandle(), GetActivationPredictionKey(), ASC->ScopedPredictionKey);
 		}
 	}
 
@@ -295,20 +306,24 @@ void UAbilityTask_WaitTargetData::OnTargetDataReadyCallback(const FGameplayAbili
 /** The TargetActor we spawned locally has called back with a cancel event (they still include the 'last/best' targetdata but the consumer of this may want to discard it) */
 void UAbilityTask_WaitTargetData::OnTargetDataCancelledCallback(const FGameplayAbilityTargetDataHandle& Data)
 {
-	check(AbilitySystemComponent.IsValid());
+	UAbilitySystemComponent* ASC = AbilitySystemComponent.Get();
+	if (!ASC)
+	{
+		return;
+	}
 
-	FScopedPredictionWindow ScopedPrediction(AbilitySystemComponent.Get(), IsPredictingClient());
+	FScopedPredictionWindow ScopedPrediction(ASC, IsPredictingClient());
 
 	if (IsPredictingClient())
 	{
 		if (!TargetActor->ShouldProduceTargetDataOnServer)
 		{
-			AbilitySystemComponent->ServerSetReplicatedTargetDataCancelled(GetAbilitySpecHandle(), GetActivationPredictionKey(), AbilitySystemComponent->ScopedPredictionKey );
+			ASC->ServerSetReplicatedTargetDataCancelled(GetAbilitySpecHandle(), GetActivationPredictionKey(), ASC->ScopedPredictionKey );
 		}
 		else
 		{
 			// We aren't going to send the target data, but we will send a generic confirmed message.
-			AbilitySystemComponent->ServerSetReplicatedEvent(EAbilityGenericReplicatedEvent::GenericCancel, GetAbilitySpecHandle(), GetActivationPredictionKey(), AbilitySystemComponent->ScopedPredictionKey);
+			ASC->ServerSetReplicatedEvent(EAbilityGenericReplicatedEvent::GenericCancel, GetAbilitySpecHandle(), GetActivationPredictionKey(), ASC->ScopedPredictionKey);
 		}
 	}
 	Cancelled.Broadcast(Data);
@@ -318,7 +333,6 @@ void UAbilityTask_WaitTargetData::OnTargetDataCancelledCallback(const FGameplayA
 /** Called when the ability is asked to confirm from an outside node. What this means depends on the individual task. By default, this does nothing other than ending if bEndTask is true. */
 void UAbilityTask_WaitTargetData::ExternalConfirm(bool bEndTask)
 {
-	check(AbilitySystemComponent.IsValid());
 	if (TargetActor)
 	{
 		if (TargetActor->ShouldProduceTargetData())
@@ -332,7 +346,6 @@ void UAbilityTask_WaitTargetData::ExternalConfirm(bool bEndTask)
 /** Called when the ability is asked to confirm from an outside node. What this means depends on the individual task. By default, this does nothing other than ending if bEndTask is true. */
 void UAbilityTask_WaitTargetData::ExternalCancel()
 {
-	check(AbilitySystemComponent.IsValid());
 	if (ShouldBroadcastAbilityTaskDelegates())
 	{
 		Cancelled.Broadcast(FGameplayAbilityTargetDataHandle());
