Track and restore previous custom movement mode properly in AbilityTask_ApplyRootMotionMoveToActorForce (Previously it was only tracking the movement mode, and if movement happened to be a custom mode, the type of custom mode would be lost)

#tests equipped Lead Sled (Weapon with Leap attack), slid towards enemy and leap attacked enemy successfully. Verified that movement is normal after attack.

diff --git a/Engine/Plugins/Runtime/GameplayAbilities/Source/GameplayAbilities/Private/Abilities/Tasks/AbilityTask_ApplyRootMotionMoveToActorForce.cpp b/Engine/Plugins/Runtime/GameplayAbilities/Source/GameplayAbilities/Private/Abilities/Tasks/AbilityTask_ApplyRootMotionMoveToActorForce.cpp
--- a/Engine/Plugins/Runtime/GameplayAbilities/Source/GameplayAbilities/Private/Abilities/Tasks/AbilityTask_ApplyRootMotionMoveToActorForce.cpp
+++ b/Engine/Plugins/Runtime/GameplayAbilities/Source/GameplayAbilities/Private/Abilities/Tasks/AbilityTask_ApplyRootMotionMoveToActorForce.cpp
@@ -30,6 +30,7 @@ UAbilityTask_ApplyRootMotionMoveToActorForce::UAbilityTask_ApplyRootMotionMoveTo
 	bSetNewMovementMode = false;
 	NewMovementMode = EMovementMode::MOVE_Walking;
 	PreviousMovementMode = EMovementMode::MOVE_None;
+	PreviousCustomMode = 0;
 	TargetLocationOffset = FVector::ZeroVector;
 	OffsetAlignment = ERootMotionMoveToActorTargetOffsetType::AlignFromTargetToSource;
 	bRestrictSpeedToExpected = false;
@@ -113,6 +114,7 @@ void UAbilityTask_ApplyRootMotionMoveToActorForce::SharedInitAndApply()
 			if (bSetNewMovementMode)
 			{
 				PreviousMovementMode = MovementComponent->MovementMode;
+				PreviousCustomMode = MovementComponent->CustomMovementMode;
 				MovementComponent->SetMovementMode(NewMovementMode);
 			}
 
@@ -352,7 +354,7 @@ void UAbilityTask_ApplyRootMotionMoveToActorForce::OnDestroy(bool AbilityIsEndin
 
 		if (bSetNewMovementMode)
 		{
-			MovementComponent->SetMovementMode(PreviousMovementMode);
+			MovementComponent->SetMovementMode(PreviousMovementMode, PreviousCustomMode);
 		}
 	}
 
diff --git a/Engine/Plugins/Runtime/GameplayAbilities/Source/GameplayAbilities/Public/Abilities/Tasks/AbilityTask_ApplyRootMotionMoveToActorForce.h b/Engine/Plugins/Runtime/GameplayAbilities/Source/GameplayAbilities/Public/Abilities/Tasks/AbilityTask_ApplyRootMotionMoveToActorForce.h
--- a/Engine/Plugins/Runtime/GameplayAbilities/Source/GameplayAbilities/Public/Abilities/Tasks/AbilityTask_ApplyRootMotionMoveToActorForce.h
+++ b/Engine/Plugins/Runtime/GameplayAbilities/Source/GameplayAbilities/Public/Abilities/Tasks/AbilityTask_ApplyRootMotionMoveToActorForce.h
@@ -130,4 +130,5 @@ protected:
 	UCurveFloat* TargetLerpSpeedVerticalCurve;
 
 	EMovementMode PreviousMovementMode;
+	uint8 PreviousCustomMode;
 };
