Manual merge of 21006616 - Converting FGameplayAbilitySpec/Def::SourceObject to be a weak reference.
* Includes update to lyra SourceObject references.

diff --git a/Engine/Plugins/Runtime/GameplayAbilities/Source/GameplayAbilities/Private/Abilities/GameplayAbility.cpp b/Engine/Plugins/Runtime/GameplayAbilities/Source/GameplayAbilities/Private/Abilities/GameplayAbility.cpp
--- a/Engine/Plugins/Runtime/GameplayAbilities/Source/GameplayAbilities/Private/Abilities/GameplayAbility.cpp
+++ b/Engine/Plugins/Runtime/GameplayAbilities/Source/GameplayAbilities/Private/Abilities/GameplayAbility.cpp
@@ -1576,7 +1576,7 @@ UObject* UGameplayAbility::GetSourceObject(FGameplayAbilitySpecHandle Handle, co
 			FGameplayAbilitySpec* AbilitySpec = AbilitySystemComponent->FindAbilitySpecFromHandle(Handle);
 			if (AbilitySpec)
 			{
-				return AbilitySpec->SourceObject;
+				return AbilitySpec->SourceObject.Get();
 			}
 		}
 	}
@@ -1593,7 +1593,7 @@ UObject* UGameplayAbility::GetCurrentSourceObject() const
 	FGameplayAbilitySpec* AbilitySpec = GetCurrentAbilitySpec();
 	if (AbilitySpec)
 	{
-		return AbilitySpec->SourceObject;
+		return AbilitySpec->SourceObject.Get();
 	}
 	return nullptr;
 }
@@ -1611,7 +1611,7 @@ FGameplayEffectContextHandle UGameplayAbility::MakeEffectContext(const FGameplay
 	// Pass along the source object to the effect
 	if (FGameplayAbilitySpec* AbilitySpec = ActorInfo->AbilitySystemComponent->FindAbilitySpecFromHandle(Handle))
 	{
-		Context.AddSourceObject(AbilitySpec->SourceObject);
+		Context.AddSourceObject(AbilitySpec->SourceObject.Get());
 	}
 
 	return Context;
diff --git a/Engine/Plugins/Runtime/GameplayAbilities/Source/GameplayAbilities/Private/GameplayAbilityTypes.cpp b/Engine/Plugins/Runtime/GameplayAbilities/Source/GameplayAbilities/Private/GameplayAbilityTypes.cpp
--- a/Engine/Plugins/Runtime/GameplayAbilities/Source/GameplayAbilities/Private/GameplayAbilityTypes.cpp
+++ b/Engine/Plugins/Runtime/GameplayAbilities/Source/GameplayAbilities/Private/GameplayAbilityTypes.cpp
@@ -255,7 +255,7 @@ FGameplayAbilitySpec::FGameplayAbilitySpec(TSubclassOf<UGameplayAbility> InAbili
 FGameplayAbilitySpec::FGameplayAbilitySpec(FGameplayAbilitySpecDef& InDef, int32 InGameplayEffectLevel, FActiveGameplayEffectHandle InGameplayEffectHandle)
 	: Ability(InDef.Ability ? InDef.Ability->GetDefaultObject<UGameplayAbility>() : nullptr)
 	, InputID(InDef.InputID)
-	, SourceObject(InDef.SourceObject)
+	, SourceObject(InDef.SourceObject.Get())
 	, ActiveCount(0)
 	, InputPressed(false)
 	, RemoveAfterActivation(false)
@@ -270,7 +270,7 @@ FGameplayAbilitySpec::FGameplayAbilitySpec(FGameplayAbilitySpecDef& InDef, int32
 
 	FString ContextString = FString::Printf(TEXT("FGameplayAbilitySpec::FGameplayAbilitySpec for %s from %s"), 
 		(InDef.Ability ? *InDef.Ability->GetName() : TEXT("INVALID ABILITY")), 
-		(InDef.SourceObject ? *InDef.SourceObject->GetName() : TEXT("INVALID ABILITY")));
+		(InDef.SourceObject.IsValid() ? *InDef.SourceObject->GetName() : TEXT("INVALID ABILITY")));
 	Level = InDef.LevelScalableFloat.GetValueAtLevel(InGameplayEffectLevel, &ContextString);
 }
 
diff --git a/Engine/Plugins/Runtime/GameplayAbilities/Source/GameplayAbilities/Private/GameplayDebuggerCategory_Abilities.cpp b/Engine/Plugins/Runtime/GameplayAbilities/Source/GameplayAbilities/Private/GameplayDebuggerCategory_Abilities.cpp
--- a/Engine/Plugins/Runtime/GameplayAbilities/Source/GameplayAbilities/Private/GameplayDebuggerCategory_Abilities.cpp
+++ b/Engine/Plugins/Runtime/GameplayAbilities/Source/GameplayAbilities/Private/GameplayDebuggerCategory_Abilities.cpp
@@ -99,7 +99,7 @@ void FGameplayDebuggerCategory_Abilities::CollectData(APlayerController* OwnerPC
 			ItemData.Ability.RemoveFromStart(DEFAULT_OBJECT_PREFIX);
 			ItemData.Ability.RemoveFromEnd(TEXT("_C"));
 
-			ItemData.Source = GetNameSafe(AbilitySpec.SourceObject);
+			ItemData.Source = GetNameSafe(AbilitySpec.SourceObject.Get());
 			ItemData.Source.RemoveFromStart(DEFAULT_OBJECT_PREFIX);
 
 			ItemData.Level = AbilitySpec.Level;
diff --git a/Engine/Plugins/Runtime/GameplayAbilities/Source/GameplayAbilities/Public/GameplayAbilitySpec.h b/Engine/Plugins/Runtime/GameplayAbilities/Source/GameplayAbilities/Public/GameplayAbilitySpec.h
--- a/Engine/Plugins/Runtime/GameplayAbilities/Source/GameplayAbilities/Public/GameplayAbilitySpec.h
+++ b/Engine/Plugins/Runtime/GameplayAbilities/Source/GameplayAbilities/Public/GameplayAbilitySpec.h
@@ -133,7 +133,7 @@ struct FGameplayAbilitySpecDef
 
 	/** What granted this spec, not replicated or settable in editor */
 	UPROPERTY(NotReplicated)
-	UObject* SourceObject;
+	TWeakObjectPtr<UObject> SourceObject;
 
 	/** SetbyCaller Magnitudes that were passed in to this ability by a GE (GE's that grant abilities) */
 	TMap<FGameplayTag, float>	SetByCallerTagMagnitudes;
@@ -248,7 +248,7 @@ struct GAMEPLAYABILITIES_API FGameplayAbilitySpec : public FFastArraySerializerI
 
 	/** Object this ability was created from, can be an actor or static object. Useful to bind an ability to a gameplay object */
 	UPROPERTY()
-	UObject* SourceObject;
+	TWeakObjectPtr<UObject> SourceObject;
 
 	/** A count of the number of times this ability has been activated minus the number of times it has been ended. For instanced abilities this will be the number of currently active instances. Can't replicate until prediction accurately handles this.*/
 	UPROPERTY(NotReplicated)
