Fix a printf safety issue reported by licensee, implict int conversions do not work well with variadic macros
This was not a problem with the current struct because it's first member is an int32, but this avoids any future issues

diff --git a/Engine/Plugins/Runtime/GameplayAbilities/Source/GameplayAbilities/Private/GameplayCueTranslator.cpp b/Engine/Plugins/Runtime/GameplayAbilities/Source/GameplayAbilities/Private/GameplayCueTranslator.cpp
--- a/Engine/Plugins/Runtime/GameplayAbilities/Source/GameplayAbilities/Private/GameplayCueTranslator.cpp
+++ b/Engine/Plugins/Runtime/GameplayAbilities/Source/GameplayAbilities/Private/GameplayCueTranslator.cpp
@@ -524,7 +524,7 @@ bool FGameplayCueTranslationManager::TranslateTag_Internal(FGameplayCueTranslato
 			{
 				if (TranslationLUT.IsValidIndex(NodeIndex) == false)
 				{
-					UE_LOG(LogGameplayCueTranslator, Error, TEXT("FGameplayCueTranslationManager::TranslateTag_Internal %s invalid index %d was returned from NodeLookup. NodeLookup.Num=%d. Tag %s"), *GetNameSafe(Link.RulesCDO), NodeIndex, TranslationLUT.Num(), *TagName.ToString());
+					UE_LOG(LogGameplayCueTranslator, Error, TEXT("FGameplayCueTranslationManager::TranslateTag_Internal %s invalid index %d was returned from NodeLookup. NodeLookup.Num=%d. Tag %s"), *GetNameSafe(Link.RulesCDO), NodeIndex.Index, TranslationLUT.Num(), *TagName.ToString());
 					continue;
 				}
 
