Removing uses of master/slave found in GameplayAbilities with base folder Source

diff --git a/Engine/Config/BaseEngine.ini b/Engine/Config/BaseEngine.ini
--- a/Engine/Config/BaseEngine.ini
+++ b/Engine/Config/BaseEngine.ini
@@ -1113,6 +1113,9 @@ DefaultWorldPartitionHLODLayer=/Engine/Maps/Templates/HLODs/HLODLayer_Instanced.
 ; Skeletal Mesh pointer deprecation
 +FunctionRedirects=(OldName="SkinnedMeshComponent.SetSkeletalMesh", NewName="SetSkinnedAssetAndUpdate")
 
++PropertyRedirects=(OldName="AGameplayAbilityWorldReticle.MasterPC", NewName="PrimaryPC")
++PropertyRedirects=(OldName="AGameplayAbilityTargetActor.MasterPC", NewName="PrimaryPC")
+
 [CoreUObject.Metadata]
 ; Note: UnrealHeaderTool should be rerun after making changes to MetadataRedirects to catch any keys specified in class headers
 +MetadataRedirects=(OldKey="K2Protected", NewKey="BlueprintProtected")
diff --git a/Engine/Plugins/Runtime/GameplayAbilities/Source/GameplayAbilities/Private/Abilities/GameplayAbilityTargetActor.cpp b/Engine/Plugins/Runtime/GameplayAbilities/Source/GameplayAbilities/Private/Abilities/GameplayAbilityTargetActor.cpp
--- a/Engine/Plugins/Runtime/GameplayAbilities/Source/GameplayAbilities/Private/Abilities/GameplayAbilityTargetActor.cpp
+++ b/Engine/Plugins/Runtime/GameplayAbilities/Source/GameplayAbilities/Private/Abilities/GameplayAbilityTargetActor.cpp
@@ -121,7 +121,7 @@ void AGameplayAbilityTargetActor::CancelTargeting()
 bool AGameplayAbilityTargetActor::IsNetRelevantFor(const AActor* RealViewer, const AActor* ViewTarget, const FVector& SrcLocation) const
 {
 	//The player who created the ability doesn't need to be updated about it - there should be local prediction in place.
-	if (RealViewer == MasterPC)
+	if (RealViewer == PrimaryPC)
 	{
 		return false;
 	}
@@ -145,7 +145,7 @@ bool AGameplayAbilityTargetActor::OnReplicatedTargetDataReceived(FGameplayAbilit
 bool AGameplayAbilityTargetActor::ShouldProduceTargetData() const
 {
 	// return true if we are locally owned, or (we are the server and this is a gameplaytarget ability that can produce target data server side)
-	return (MasterPC && MasterPC->IsLocalController()) || ShouldProduceTargetDataOnServer;
+	return (PrimaryPC && PrimaryPC->IsLocalController()) || ShouldProduceTargetDataOnServer;
 }
 
 void AGameplayAbilityTargetActor::BindToConfirmCancelInputs()
diff --git a/Engine/Plugins/Runtime/GameplayAbilities/Source/GameplayAbilities/Private/Abilities/GameplayAbilityTargetActor_Trace.cpp b/Engine/Plugins/Runtime/GameplayAbilities/Source/GameplayAbilities/Private/Abilities/GameplayAbilityTargetActor_Trace.cpp
--- a/Engine/Plugins/Runtime/GameplayAbilities/Source/GameplayAbilities/Private/Abilities/GameplayAbilityTargetActor_Trace.cpp
+++ b/Engine/Plugins/Runtime/GameplayAbilities/Source/GameplayAbilities/Private/Abilities/GameplayAbilityTargetActor_Trace.cpp
@@ -158,7 +158,7 @@ void AGameplayAbilityTargetActor_Trace::StartTargeting(UGameplayAbility* InAbili
 		AGameplayAbilityWorldReticle* SpawnedReticleActor = GetWorld()->SpawnActor<AGameplayAbilityWorldReticle>(ReticleClass, GetActorLocation(), GetActorRotation());
 		if (SpawnedReticleActor)
 		{
-			SpawnedReticleActor->InitializeReticle(this, MasterPC, ReticleParams);
+			SpawnedReticleActor->InitializeReticle(this, PrimaryPC, ReticleParams);
 			ReticleActor = SpawnedReticleActor;
 
 			// This is to catch cases of playing on a listen server where we are using a replicated reticle actor.
diff --git a/Engine/Plugins/Runtime/GameplayAbilities/Source/GameplayAbilities/Private/Abilities/GameplayAbilityWorldReticle.cpp b/Engine/Plugins/Runtime/GameplayAbilities/Source/GameplayAbilities/Private/Abilities/GameplayAbilityWorldReticle.cpp
--- a/Engine/Plugins/Runtime/GameplayAbilities/Source/GameplayAbilities/Private/Abilities/GameplayAbilityWorldReticle.cpp
+++ b/Engine/Plugins/Runtime/GameplayAbilities/Source/GameplayAbilities/Private/Abilities/GameplayAbilityWorldReticle.cpp
@@ -31,7 +31,7 @@ void AGameplayAbilityWorldReticle::InitializeReticle(AActor* InTargetingActor, A
 {
 	check(InTargetingActor);
 	TargetingActor = InTargetingActor;
-	MasterPC = PlayerController;
+	PrimaryPC = PlayerController;
 	AddTickPrerequisiteActor(TargetingActor);		//We want the reticle to tick after the targeting actor so that designers have the final say on the position
 	Parameters = InParameters;
 	OnParametersInitialized();
@@ -40,7 +40,7 @@ void AGameplayAbilityWorldReticle::InitializeReticle(AActor* InTargetingActor, A
 bool AGameplayAbilityWorldReticle::IsNetRelevantFor(const AActor* RealViewer, const AActor* ViewTarget, const FVector& SrcLocation) const
 {
 	//The player who created the ability doesn't need to be updated about it - there should be local prediction in place.
-	if (RealViewer == MasterPC)
+	if (RealViewer == PrimaryPC)
 	{
 		return false;
 	}
diff --git a/Engine/Plugins/Runtime/GameplayAbilities/Source/GameplayAbilities/Private/Abilities/Tasks/AbilityTask_VisualizeTargeting.cpp b/Engine/Plugins/Runtime/GameplayAbilities/Source/GameplayAbilities/Private/Abilities/Tasks/AbilityTask_VisualizeTargeting.cpp
--- a/Engine/Plugins/Runtime/GameplayAbilities/Source/GameplayAbilities/Private/Abilities/Tasks/AbilityTask_VisualizeTargeting.cpp
+++ b/Engine/Plugins/Runtime/GameplayAbilities/Source/GameplayAbilities/Private/Abilities/Tasks/AbilityTask_VisualizeTargeting.cpp
@@ -135,7 +135,7 @@ void UAbilityTask_VisualizeTargeting::InitializeTargetActor(AGameplayAbilityTarg
 	check(SpawnedActor);
 	check(Ability);
 
-	SpawnedActor->MasterPC = Ability->GetCurrentActorInfo()->PlayerController.Get();
+	SpawnedActor->PrimaryPC = Ability->GetCurrentActorInfo()->PlayerController.Get();
 }
 
 void UAbilityTask_VisualizeTargeting::FinalizeTargetActor(AGameplayAbilityTargetActor* SpawnedActor) const
diff --git a/Engine/Plugins/Runtime/GameplayAbilities/Source/GameplayAbilities/Private/Abilities/Tasks/AbilityTask_WaitTargetData.cpp b/Engine/Plugins/Runtime/GameplayAbilities/Source/GameplayAbilities/Private/Abilities/Tasks/AbilityTask_WaitTargetData.cpp
--- a/Engine/Plugins/Runtime/GameplayAbilities/Source/GameplayAbilities/Private/Abilities/Tasks/AbilityTask_WaitTargetData.cpp
+++ b/Engine/Plugins/Runtime/GameplayAbilities/Source/GameplayAbilities/Private/Abilities/Tasks/AbilityTask_WaitTargetData.cpp
@@ -137,7 +137,7 @@ void UAbilityTask_WaitTargetData::InitializeTargetActor(AGameplayAbilityTargetAc
 	check(SpawnedActor);
 	check(Ability);
 
-	SpawnedActor->MasterPC = Ability->GetCurrentActorInfo()->PlayerController.Get();
+	SpawnedActor->PrimaryPC = Ability->GetCurrentActorInfo()->PlayerController.Get();
 
 	// If we spawned the target actor, always register the callbacks for when the data is ready.
 	SpawnedActor->TargetDataReadyDelegate.AddUObject(const_cast<UAbilityTask_WaitTargetData*>(this), &UAbilityTask_WaitTargetData::OnTargetDataReadyCallback);
diff --git a/Engine/Plugins/Runtime/GameplayAbilities/Source/GameplayAbilities/Private/GameplayCueTranslator.cpp b/Engine/Plugins/Runtime/GameplayAbilities/Source/GameplayAbilities/Private/GameplayCueTranslator.cpp
--- a/Engine/Plugins/Runtime/GameplayAbilities/Source/GameplayAbilities/Private/GameplayCueTranslator.cpp
+++ b/Engine/Plugins/Runtime/GameplayAbilities/Source/GameplayAbilities/Private/GameplayCueTranslator.cpp
@@ -279,7 +279,7 @@ bool FGameplayCueTranslationManager::BuildTagTranslationTable_r(const FName& Tag
 							SwappedNames = SplitNames;
 
 							// Remove the "To Names" from the SwappedNames array, replace with the single "From Name"
-							// E.g, GC.{Steel.Master} -> GC.{Hero}
+							// E.g, GC.{Steel.Primary} -> GC.{Hero}
 							int32 NumRemoves = SwapRule.ToNames.Num(); // We are going to remove as many tags 
 							int32 RemoveAtIdx = TagIdx - (SwapRule.ToNames.Num() - 1);
 							check(SwappedNames.IsValidIndex(RemoveAtIdx));
diff --git a/Engine/Plugins/Runtime/GameplayAbilities/Source/GameplayAbilities/Public/Abilities/GameplayAbilityTargetActor.h b/Engine/Plugins/Runtime/GameplayAbilities/Source/GameplayAbilities/Public/Abilities/GameplayAbilityTargetActor.h
--- a/Engine/Plugins/Runtime/GameplayAbilities/Source/GameplayAbilities/Public/Abilities/GameplayAbilityTargetActor.h
+++ b/Engine/Plugins/Runtime/GameplayAbilities/Source/GameplayAbilities/Public/Abilities/GameplayAbilityTargetActor.h
@@ -69,8 +69,13 @@ public:
 
 	virtual bool IsNetRelevantFor(const AActor* RealViewer, const AActor* ViewTarget, const FVector& SrcLocation) const override;
 
-	UPROPERTY(BlueprintReadOnly, Category = "Targeting")
+#if WITH_EDITOR
+	UE_DEPRECATED(5.1, "This property is deprecated. Please use PrimaryPC instead.")
 	APlayerController* MasterPC;
+#endif // WITH_EDITOR
+	
+	UPROPERTY(BlueprintReadOnly, Category = "Targeting")
+	APlayerController* PrimaryPC;
 
 	UPROPERTY()
 	UGameplayAbility* OwningAbility;
diff --git a/Engine/Plugins/Runtime/GameplayAbilities/Source/GameplayAbilities/Public/Abilities/GameplayAbilityWorldReticle.h b/Engine/Plugins/Runtime/GameplayAbilities/Source/GameplayAbilities/Public/Abilities/GameplayAbilityWorldReticle.h
--- a/Engine/Plugins/Runtime/GameplayAbilities/Source/GameplayAbilities/Public/Abilities/GameplayAbilityWorldReticle.h
+++ b/Engine/Plugins/Runtime/GameplayAbilities/Source/GameplayAbilities/Public/Abilities/GameplayAbilityWorldReticle.h
@@ -82,9 +82,15 @@ protected:
 	UPROPERTY(BlueprintReadOnly, Category = "Network")
 	bool bIsTargetAnActor;
 
+#if WITH_EDITOR
 	/** This is used in the process of determining whether we should replicate to a specific client. */
-	UPROPERTY(BlueprintReadOnly, Category = "Network")
+	UE_DEPRECATED(5.1, "This property is deprecated. Please use PrimaryPC instead")
 	APlayerController* MasterPC;
+#endif // WITH_EDITOR
+	
+	/** This is used in the process of determining whether we should replicate to a specific client. */
+	UPROPERTY(BlueprintReadOnly, Category = "Network")
+	APlayerController* PrimaryPC;
 
 	/** In the future, we may want to grab things like sockets off of this. */
 	UPROPERTY(BlueprintReadOnly, Category = "Network")
