SubObject
* Downgrade check that detected GameplayAbilities sending RPCS before registering into a warning

diff --git a/Engine/Plugins/Runtime/GameplayAbilities/Source/GameplayAbilities/Private/Abilities/GameplayAbility.cpp b/Engine/Plugins/Runtime/GameplayAbilities/Source/GameplayAbilities/Private/Abilities/GameplayAbility.cpp
--- a/Engine/Plugins/Runtime/GameplayAbilities/Source/GameplayAbilities/Private/Abilities/GameplayAbility.cpp
+++ b/Engine/Plugins/Runtime/GameplayAbilities/Source/GameplayAbilities/Private/Abilities/GameplayAbility.cpp
@@ -122,7 +122,7 @@ bool UGameplayAbility::CallRemoteFunction(UFunction* Function, void* Parameters,
 		{
 			// If the component is using the subobject list, detect subobjects that send RPC's before they are registered.
 			const bool bIsSubObjectRegistered = !AbilitySystemComponent->IsUsingRegisteredSubObjectList() || Owner->IsActorComponentReplicatedSubObjectRegistered(AbilitySystemComponent, this);
-			checkf(bIsSubObjectRegistered, TEXT("%s sent an RPC before it was registered as a subobject on %s::%s"), *GetName(), *Owner->GetName(), *AbilitySystemComponent->GetName());
+			UE_CLOG(!bIsSubObjectRegistered, LogAbilitySystem, Warning, TEXT("%s sent an RPC before it was registered as a subobject on %s::%s"), *GetName(), *Owner->GetName(), *AbilitySystemComponent->GetName());
 		}
 #endif
 		
