Add AbiliitySystem debug commands to UE

diff --git a/Engine/Plugins/Runtime/GameplayAbilities/Source/GameplayAbilities/Private/AbilitySystemComponent.cpp b/Engine/Plugins/Runtime/GameplayAbilities/Source/GameplayAbilities/Private/AbilitySystemComponent.cpp
--- a/Engine/Plugins/Runtime/GameplayAbilities/Source/GameplayAbilities/Private/AbilitySystemComponent.cpp
+++ b/Engine/Plugins/Runtime/GameplayAbilities/Source/GameplayAbilities/Private/AbilitySystemComponent.cpp
@@ -270,6 +270,11 @@ const FActiveGameplayEffect* UAbilitySystemComponent::GetActiveGameplayEffect(co
 	return ActiveGameplayEffects.GetActiveGameplayEffect(Handle);
 }
 
+const FActiveGameplayEffectsContainer& UAbilitySystemComponent::GetActiveGameplayEffects() const
+{
+	return ActiveGameplayEffects;
+}
+
 const UGameplayEffect* UAbilitySystemComponent::GetGameplayEffectCDO(const FActiveGameplayEffectHandle Handle) const
 {
 	// get the active gameplay effect struct
diff --git a/Engine/Plugins/Runtime/GameplayAbilities/Source/GameplayAbilities/Private/AbilitySystemDebugHUD.cpp b/Engine/Plugins/Runtime/GameplayAbilities/Source/GameplayAbilities/Private/AbilitySystemDebugHUD.cpp
--- a/Engine/Plugins/Runtime/GameplayAbilities/Source/GameplayAbilities/Private/AbilitySystemDebugHUD.cpp
+++ b/Engine/Plugins/Runtime/GameplayAbilities/Source/GameplayAbilities/Private/AbilitySystemDebugHUD.cpp
@@ -12,11 +12,237 @@
 #include "AbilitySystemComponent.h"
 #include "AbilitySystemGlobals.h"
 
+namespace
+{
+	static float DebugDrawMaxDistance = 2048.f;
+	static FAutoConsoleVariableRef CVarDebugDrawMaxDistance(
+		TEXT("AbilitySystem.DebugDrawMaxDistance"),
+		DebugDrawMaxDistance,
+		TEXT("Set the maximum camera distance allowed for Debug Drawing by the Ability System.")
+	);
+}
+
+bool UAbilitySystemDebugHUDExtension_Tags::bEnabled = false;
+
+TSet<FString> UAbilitySystemDebugHUDExtension_Attributes::AttributesToDisplay;
+
+bool UAbilitySystemDebugHUDExtension_Attributes::bIncludeModifiers = false;
+
+bool UAbilitySystemDebugHUDExtension_BlockedAbilityTags::bEnabled = false;
+
+bool AAbilitySystemDebugHUD::bEnableBasicHUD = false;
+
+TSubclassOf<AAbilitySystemDebugHUD> AAbilitySystemDebugHUD::HUDClass = AAbilitySystemDebugHUD::StaticClass();
+
+void UAbilitySystemDebugHUDExtension_Tags::ToggleExtension(const TArray<FString>& Args, UWorld* World)
+{
+	bEnabled = !bEnabled;
+
+	AAbilitySystemDebugHUD::NotifyExtensionEnableChanged(World);
+}
+
+bool UAbilitySystemDebugHUDExtension_Tags::IsEnabled() const
+{
+	return bEnabled;
+}
+
+void UAbilitySystemDebugHUDExtension_Tags::GetDebugStrings(const AActor* Actor, const UAbilitySystemComponent* Comp, OUT TArray<FString>& OutDebugStrings) const
+{
+	if (!Comp)
+	{
+		return;
+	}
+
+	FGameplayTagContainer Container;
+	Comp->GetOwnedGameplayTags(Container);
+
+	FGameplayTagContainer LeafOnlyContainer;
+	for (auto ContainerIt = Container.CreateConstIterator(); ContainerIt; ++ContainerIt)
+	{
+		LeafOnlyContainer.AddLeafTag(*ContainerIt);
+	}
+
+	for (auto ContainerIt = LeafOnlyContainer.CreateConstIterator(); ContainerIt; ++ContainerIt)
+	{
+		OutDebugStrings.Add(ContainerIt->ToString());
+	}
+}
+
+void UAbilitySystemDebugHUDExtension_Attributes::ToggleExtension(const TArray<FString>& Args, UWorld* World)
+{
+	for (const FString& Arg : Args)
+	{
+		FSetElementId FoundId = AttributesToDisplay.FindId(Arg);
+		if (FoundId.IsValidId())
+		{
+			AttributesToDisplay.Remove(FoundId);
+		}
+		else
+		{
+			AttributesToDisplay.Add(Arg);
+		}
+	}
+
+	AAbilitySystemDebugHUD::NotifyExtensionEnableChanged(World);
+}
+
+void UAbilitySystemDebugHUDExtension_Attributes::ToggleIncludeModifiers()
+{
+	bIncludeModifiers = !bIncludeModifiers;
+}
+
+void UAbilitySystemDebugHUDExtension_Attributes::ClearDisplayedAttributes(const TArray<FString>& Args, UWorld* World)
+{
+	AttributesToDisplay.Empty();
+
+	AAbilitySystemDebugHUD::NotifyExtensionEnableChanged(World);
+}
+
+bool UAbilitySystemDebugHUDExtension_Attributes::IsEnabled() const
+{
+	return AttributesToDisplay.Num() > 0;
+}
+
+void UAbilitySystemDebugHUDExtension_Attributes::GetDebugStrings(const AActor* Actor, const UAbilitySystemComponent* Comp, OUT TArray<FString>& OutDebugStrings) const
+{
+	if (!Comp)
+	{
+		return;
+	}
+
+	const FActiveGameplayEffectsContainer& ActiveEffects = Comp->GetActiveGameplayEffects();
+
+	TArray<UObject*> Objects;
+	GetObjectsWithOuter(Comp->GetOwnerActor(), Objects);
+
+	TMultiMap<FGameplayAttribute, FActiveGameplayEffectsContainer::DebugExecutedGameplayEffectData> EffectMap;
+
+	ActiveEffects.GetActiveGameplayEffectDataByAttribute(EffectMap);
+
+	for (UObject* Object : Objects)
+	{
+		UAttributeSet* AttributeSet = Cast<UAttributeSet>(Object);
+		if (AttributeSet)
+		{
+			UClass* Class = AttributeSet->GetClass();
+
+			for (TFieldIterator<FProperty> PropertyIt(Class, EFieldIteratorFlags::IncludeSuper); PropertyIt; ++PropertyIt)
+			{
+				FProperty* Property = *PropertyIt;
+
+				if (AttributesToDisplay.Contains(Property->GetName()))
+				{
+					FGameplayAttribute Attribute(Property);
+					OutDebugStrings.Add(FString::Printf(TEXT("%s: %f (Base: %f)"), *Property->GetName(), Comp->GetNumericAttribute(Attribute), Comp->GetNumericAttributeBase(Attribute)));
+
+					// if we're also showing modifiers, grab them and display them here
+					if (bIncludeModifiers)
+					{
+						TArray<FActiveGameplayEffectsContainer::DebugExecutedGameplayEffectData> AttributeEffects;
+						EffectMap.MultiFind(Attribute, AttributeEffects);
+
+						for (const FActiveGameplayEffectsContainer::DebugExecutedGameplayEffectData& DebugData : AttributeEffects)
+						{
+							OutDebugStrings.Add(FString::Printf(TEXT("%s %s"), *DebugData.GameplayEffectName, *DebugData.ActivationState));
+							OutDebugStrings.Add(FString::Printf(TEXT("%s %f"), *EGameplayModOpToString(DebugData.ModifierOp), DebugData.Magnitude));
+							OutDebugStrings.Add(FString::Printf(TEXT("Stack Count: %i"), DebugData.StackCount));
+						}
+					}
+				}
+			}
+		}
+	}
+}
+
+void UAbilitySystemDebugHUDExtension_BlockedAbilityTags::ToggleExtension(const TArray<FString>& Args, UWorld* World)
+{
+	bEnabled = !bEnabled;
+
+	AAbilitySystemDebugHUD::NotifyExtensionEnableChanged(World);
+}
+
+bool UAbilitySystemDebugHUDExtension_BlockedAbilityTags::IsEnabled() const
+{
+	return bEnabled;
+}
+
+void UAbilitySystemDebugHUDExtension_BlockedAbilityTags::GetDebugStrings(const AActor* Actor, const UAbilitySystemComponent* Comp, OUT TArray<FString>& OutDebugStrings) const
+{
+	if (!Comp)
+	{
+		return;
+	}
+
+	FGameplayTagContainer Container;
+	Comp->GetBlockedAbilityTags(Container);
+
+	FGameplayTagContainer LeafOnlyContainer;
+	for (auto ContainerIt = Container.CreateConstIterator(); ContainerIt; ++ContainerIt)
+	{
+		LeafOnlyContainer.AddLeafTag(*ContainerIt);
+	}
+
+	for (auto ContainerIt = LeafOnlyContainer.CreateConstIterator(); ContainerIt; ++ContainerIt)
+	{
+		OutDebugStrings.Add(ContainerIt->ToString());
+	}
+}
+
 AAbilitySystemDebugHUD::AAbilitySystemDebugHUD(const FObjectInitializer& ObjectInitializer) : Super(ObjectInitializer)
 {
 
 }
 
+void AAbilitySystemDebugHUD::ToggleBasicHUD(const TArray<FString>& Args, UWorld* World)
+{
+	bEnableBasicHUD = !bEnableBasicHUD;
+
+	NotifyExtensionEnableChanged(World);
+}
+
+void AAbilitySystemDebugHUD::NotifyExtensionEnableChanged(UWorld* InWorld)
+{
+	if (!InWorld)
+	{
+		return;
+	}
+
+	AAbilitySystemDebugHUD* HUD = NULL;
+	for (TActorIterator<AAbilitySystemDebugHUD> It(InWorld); It; ++It)
+	{
+		HUD = *It;
+		break;
+	}
+
+	bool bAnyExtensionEnabled = bEnableBasicHUD;
+	for (TObjectIterator<UAbilitySystemDebugHUDExtension> ExtIt(RF_NoFlags); ExtIt; ++ExtIt)
+	{
+		bAnyExtensionEnabled |= ExtIt->IsEnabled();
+	}
+
+	static FDelegateHandle DrawDebugDelegateHandle;
+
+	if (!HUD && bAnyExtensionEnabled)
+	{
+		UClass* HUDClassToSpawn = HUDClass.Get();
+		if (!HUDClassToSpawn)
+		{
+			HUDClassToSpawn = AAbilitySystemDebugHUD::StaticClass();
+		}
+
+		HUD = InWorld->SpawnActor<AAbilitySystemDebugHUD>();
+
+		FDebugDrawDelegate DrawDebugDelegate = FDebugDrawDelegate::CreateUObject(HUD, &AAbilitySystemDebugHUD::DrawDebugHUD);
+		DrawDebugDelegateHandle = UDebugDrawService::Register(TEXT("GameplayDebug"), DrawDebugDelegate);
+	}
+	else if (HUD && !bAnyExtensionEnabled)
+	{
+		FDebugDrawDelegate DrawDebugDelegate = FDebugDrawDelegate::CreateUObject(HUD, &AAbilitySystemDebugHUD::DrawDebugHUD);
+		UDebugDrawService::Unregister(DrawDebugDelegateHandle);
+		HUD->Destroy();
+	}
+}
+
 void AAbilitySystemDebugHUD::DrawWithBackground(UFont* InFont, const FString& Text, const FColor& TextColor, EAlignHorizontal::Type HAlign, float& OffsetX, EAlignVertical::Type VAlign, float& OffsetY, float Alpha)
 {
 	float SizeX, SizeY;
@@ -46,7 +272,7 @@ void AAbilitySystemDebugHUD::DrawWithBackground(UFont* InFont, const FString& Te
 	OffsetY += 25;
 }
 
-void AAbilitySystemDebugHUD::DrawDebugHUD(UCanvas* InCanvas, APlayerController*)
+void AAbilitySystemDebugHUD::DrawDebugHUD(UCanvas* InCanvas, APlayerController* Cont)
 {
 	Canvas = InCanvas;
 	if (!Canvas)
@@ -54,31 +280,117 @@ void AAbilitySystemDebugHUD::DrawDebugHUD(UCanvas* InCanvas, APlayerController*)
 		return;
 	}
 
-	UPlayer * LocalPlayer = GetWorld()->GetFirstLocalPlayerFromController();
-	if (!LocalPlayer)
+	APlayerController *PC = GetWorld()->GetFirstPlayerController();
+	if (!PC)
 	{
 		return;
 	}
 
-	APlayerController *PC = LocalPlayer->PlayerController;
-	if (!PC)
+	if (bEnableBasicHUD)
 	{
-		return;
+		if (UAbilitySystemComponent* AbilitySystemComponent = UAbilitySystemGlobals::GetAbilitySystemComponentFromActor(PC->GetPawn()))
+		{
+			DrawDebugAbilitySystemComponent(AbilitySystemComponent);
+		}
 	}
 
-	if (UAbilitySystemComponent* AbilitySystemComponent = UAbilitySystemGlobals::GetAbilitySystemComponentFromActor(PC->GetPawn()))
+	DrawAbilityDebugInfo(Canvas, PC);
+}
+
+bool AAbilitySystemDebugHUD::ShouldDisplayDebugForActor(UCanvas* InCanvas, const AActor* Actor, const FVector& CameraPosition, const FVector& CameraDir) const
+{
+	if (!Actor)
+	{
+		return false;
+	}
+
+	if (!Actor->WasRecentlyRendered(0.1f))
+	{
+		return false;
+	}
+
+	const FVector ThisLocation = Actor->GetActorLocation();
+	const FVector CameraToActorVector = (ThisLocation - CameraPosition);
+	const bool bIsBehindCamera = ((CameraDir | CameraToActorVector) < 0.f);
+	if (bIsBehindCamera)
+	{
+		return false;
+	}
+
+	const bool bBeyondMaxDistance = CameraToActorVector.SizeSquared2D() > FMath::Square(DebugDrawMaxDistance);
+	if (bBeyondMaxDistance)
+	{
+		return false;
+	}
+
+	FVector ThisOrigin;
+	FVector BoxExtent;
+	Actor->GetActorBounds(true, ThisOrigin, BoxExtent);
+
+	const FVector ScreenLoc = InCanvas->Project(ThisOrigin);
+	if ((ScreenLoc.X < 0) || (ScreenLoc.X >= InCanvas->ClipX) || (ScreenLoc.Y < 0) || (ScreenLoc.Y >= InCanvas->ClipY))
+	{
+		return false;
+	}
+
+	return true;
+}
+
+void AAbilitySystemDebugHUD::DisplayDebugStrings(UCanvas* InCanvas, const AActor* Actor, const TArray<FString>& DebugStrings, const FVector& CameraPosition, const FVector& CameraDir, float& VerticalOffset) const
+{
+	if (DebugStrings.Num() == 0)
 	{
-		DrawDebugAbilitySystemComponent(AbilitySystemComponent);
 		return;
 	}
 
-	if (PC->GetPawn())
+	struct FStringAndLineHeight
 	{
-		UAbilitySystemComponent * AbilitySystemComponent = PC->GetPawn()->FindComponentByClass<UAbilitySystemComponent>();
-		if (AbilitySystemComponent)
-		{
-			DrawDebugAbilitySystemComponent(AbilitySystemComponent);
-		}
+		FString String;
+		float LineHeight;
+
+		FStringAndLineHeight(const FString& InString, float InLineHeight)
+			: String(InString), LineHeight(InLineHeight)
+		{}
+	};
+
+	UFont* Font = GEngine->GetTinyFont();
+	TArray<FStringAndLineHeight> StringsToDraw;
+	float MaxWidth = 0.f;
+	float TotalHeight = 0.f;
+	for (auto ContainerIt = DebugStrings.CreateConstIterator(); ContainerIt; ++ContainerIt)
+	{
+		const FString DrawString = *ContainerIt;
+		float TextWidth = 0.f;
+		float TextHeight = 0.f;
+		InCanvas->TextSize(Font, DrawString, TextWidth, TextHeight);
+		MaxWidth = FMath::Max(MaxWidth, TextWidth);
+		TotalHeight += TextHeight;
+		new(StringsToDraw)FStringAndLineHeight(DrawString, TextHeight);
+	}
+
+	FVector ThisOrigin;
+	FVector BoxExtent;
+	Actor->GetActorBounds(true, ThisOrigin, BoxExtent);
+	const FVector ScreenLoc = InCanvas->Project(ThisOrigin);
+	const float TextLeft = ScreenLoc.X - MaxWidth * 0.5f;
+	const float TextTop = ScreenLoc.Y - TotalHeight * 0.5f + VerticalOffset;
+
+	// Draw a background
+	const FLinearColor BackgroundColor(0.106, 0.039f, 0.039f, 0.3f);
+	const float BoxPadding = 2.f;
+	const FVector2D BoxLocation(TextLeft - BoxPadding, TextTop - BoxPadding);
+	const FVector2D BoxSize(MaxWidth + BoxPadding * 2, TotalHeight + BoxPadding * 2);
+	VerticalOffset += BoxSize.Y + BoxPadding;
+	FCanvasTileItem Tile(BoxLocation, BoxSize, BackgroundColor);
+	InCanvas->DrawItem(Tile);
+
+	// Debug Text
+	InCanvas->SetDrawColor(FColor(242, 242, 242));
+	float DrawTop = TextTop;
+	for (const FStringAndLineHeight& StringAndLineHeight : StringsToDraw)
+	{
+		InCanvas->DrawText(Font, StringAndLineHeight.String, TextLeft, DrawTop);
+		DrawTop += StringAndLineHeight.LineHeight;
 	}
 }
 
@@ -139,41 +451,91 @@ void AAbilitySystemDebugHUD::DrawDebugAbilitySystemComponent(UAbilitySystemCompo
 	
 }
 
-#if !UE_BUILD_SHIPPING
-static void	ToggleDebugHUD(const TArray<FString>& Args, UWorld* InWorld)
+void AAbilitySystemDebugHUD::DrawAbilityDebugInfo(UCanvas* InCanvas, APlayerController* PC) const
 {
-
-	if (!InWorld)
+	if (!PC)
+	{
+		// Need a PC to determine the camera location
 		return;
+	}
 
-	AAbilitySystemDebugHUD *HUD = NULL;
+	UWorld* PCWorld = PC->GetWorld();
 
-	for (TActorIterator<AAbilitySystemDebugHUD> It(InWorld); It; ++It)
+	if (!PCWorld)
 	{
-		HUD = *It;
-		break;
+		// Need a world to determine which actors to draw
+		return;
 	}
 
-	static FDelegateHandle DrawDebugDelegateHandle;
+	FVector CameraPosition;
+	FRotator CameraRotation;
+	PC->GetPlayerViewPoint(CameraPosition, CameraRotation);
+	const FVector CameraDir = CameraRotation.Vector();
 
-	if (!HUD)
-	{
-		HUD = InWorld->SpawnActor<AAbilitySystemDebugHUD>();
-		
-		FDebugDrawDelegate DrawDebugDelegate = FDebugDrawDelegate::CreateUObject(HUD, &AAbilitySystemDebugHUD::DrawDebugHUD);
-		DrawDebugDelegateHandle = UDebugDrawService::Register(TEXT("GameplayDebug"), DrawDebugDelegate);
-	}
-	else
+	for (TObjectIterator<UAbilitySystemComponent> ASCIt; ASCIt; ++ASCIt)
 	{
-		FDebugDrawDelegate DrawDebugDelegate = FDebugDrawDelegate::CreateUObject(HUD, &AAbilitySystemDebugHUD::DrawDebugHUD);
-		UDebugDrawService::Unregister(DrawDebugDelegateHandle);
-		HUD->Destroy();
+		if (UAbilitySystemComponent* ASC = *ASCIt)
+		{
+			if (ASC->GetWorld() == PCWorld)
+			{
+				const AActor* AvatarActor = ASC->GetAvatarActor();
+				if (!ShouldDisplayDebugForActor(InCanvas, AvatarActor, CameraPosition, CameraDir))
+				{
+					continue;
+				}
+
+				float VerticalOffset = 0.f;
+				for (TObjectIterator<UAbilitySystemDebugHUDExtension> ExtIt(EObjectFlags::RF_NoFlags); ExtIt; ++ExtIt)
+				{
+					if (ExtIt->IsEnabled())
+					{
+						TArray<FString> DebugStrings;
+						ExtIt->GetDebugStrings(AvatarActor, ASC, DebugStrings);
+						DisplayDebugStrings(InCanvas, AvatarActor, DebugStrings, CameraPosition, CameraDir, VerticalOffset);
+					}
+				}
+			}
+		}
 	}
 }
 
+
+#if !UE_BUILD_SHIPPING
+
 FAutoConsoleCommandWithWorldAndArgs AbilitySystemToggleDebugHUDCommand(
-	TEXT("AbilitySystem.ToggleDebugHUD"),
-	TEXT("ToggleDebugHUD Drawing"),
-	FConsoleCommandWithWorldAndArgsDelegate::CreateStatic(ToggleDebugHUD)
-	);
+	TEXT("AbilitySystem.DebugBasicHUD"),
+	TEXT("Toggles Drawing a basic debug HUD for the local player's ability system component"),
+	FConsoleCommandWithWorldAndArgsDelegate::CreateStatic(AAbilitySystemDebugHUD::ToggleBasicHUD)
+);
+
+FAutoConsoleCommandWithWorldAndArgs AbilitySystemToggleDebugTagsCommand(
+	TEXT("AbilitySystem.DebugAbilityTags"),
+	TEXT("Toggles Drawing Ability Tags on Actors with AbilitySystemComponents"),
+	FConsoleCommandWithWorldAndArgsDelegate::CreateStatic(UAbilitySystemDebugHUDExtension_Tags::ToggleExtension)
+);
+
+FAutoConsoleCommandWithWorldAndArgs AbilitySystemToggleDebugAttributesCommand(
+	TEXT("AbilitySystem.DebugAttribute"),
+	TEXT("Usage: AbilitySystem.DebugAttribute [AttributeName] [AttributeName]...\nToggles Drawing the given attributes on Actors with AbilitySystemComponents"),
+	FConsoleCommandWithWorldAndArgsDelegate::CreateStatic(UAbilitySystemDebugHUDExtension_Attributes::ToggleExtension)
+);
+
+FAutoConsoleCommandWithWorldAndArgs AbilitySystemClearDebugAttributesCommand(
+	TEXT("AbilitySystem.ClearDebugAttributes"),
+	TEXT("Usage: AbilitySystem.ClearDebugAttributes...\nStops drawing all attributes on Actors with AbilitySystemComponents"),
+	FConsoleCommandWithWorldAndArgsDelegate::CreateStatic(UAbilitySystemDebugHUDExtension_Attributes::ClearDisplayedAttributes)
+);
+
+FAutoConsoleCommand AbilitySystemToggleIncludeModifiersCommand(
+	TEXT("AbilitySystem.DebugIncludeModifiers"),
+	TEXT("Usage: AbilitySystem.DebugIncludeModifiers...\nToggles whether or not modifiers are displayed alongside debugged attributes"),
+	FConsoleCommandDelegate::CreateStatic(UAbilitySystemDebugHUDExtension_Attributes::ToggleIncludeModifiers)
+);
+
+FAutoConsoleCommandWithWorldAndArgs AbilitySystemToggleDebugBlockedTagsCommand(
+	TEXT("AbilitySystem.DebugBlockedAbilityTags"),
+	TEXT("Toggles Drawing Blocked Ability Tags on Actors with AbilitySystemComponents"),
+	FConsoleCommandWithWorldAndArgsDelegate::CreateStatic(UAbilitySystemDebugHUDExtension_BlockedAbilityTags::ToggleExtension)
+);
+
 #endif // !UE_BUILD_SHIPPING
diff --git a/Engine/Plugins/Runtime/GameplayAbilities/Source/GameplayAbilities/Public/AbilitySystemComponent.h b/Engine/Plugins/Runtime/GameplayAbilities/Source/GameplayAbilities/Public/AbilitySystemComponent.h
--- a/Engine/Plugins/Runtime/GameplayAbilities/Source/GameplayAbilities/Public/AbilitySystemComponent.h
+++ b/Engine/Plugins/Runtime/GameplayAbilities/Source/GameplayAbilities/Public/AbilitySystemComponent.h
@@ -418,6 +418,9 @@ class GAMEPLAYABILITIES_API UAbilitySystemComponent : public UGameplayTasksCompo
 	/** Returns const pointer to the actual active gameplay effect structure */
 	const FActiveGameplayEffect* GetActiveGameplayEffect(const FActiveGameplayEffectHandle Handle) const;
 
+	/** Returns all active gameplay effects on this ASC */
+	const FActiveGameplayEffectsContainer& GetActiveGameplayEffects() const;
+
 	/** Returns a const pointer to the gameplay effect CDO associated with an active handle. */
 	const UGameplayEffect* GetGameplayEffectCDO(const FActiveGameplayEffectHandle Handle) const;
 
diff --git a/Engine/Plugins/Runtime/GameplayAbilities/Source/GameplayAbilities/Public/AbilitySystemDebugHUD.h b/Engine/Plugins/Runtime/GameplayAbilities/Source/GameplayAbilities/Public/AbilitySystemDebugHUD.h
--- a/Engine/Plugins/Runtime/GameplayAbilities/Source/GameplayAbilities/Public/AbilitySystemDebugHUD.h
+++ b/Engine/Plugins/Runtime/GameplayAbilities/Source/GameplayAbilities/Public/AbilitySystemDebugHUD.h
@@ -5,6 +5,7 @@
 #include "CoreMinimal.h"
 #include "UObject/ObjectMacros.h"
 #include "GameFramework/HUD.h"
+#include "GameplayTagContainer.h"
 #include "AbilitySystemDebugHUD.generated.h"
 
 class APlayerController;
@@ -32,18 +33,140 @@ namespace EAlignVertical
 	};
 }
 
+/**
+ * An extension for the Ability System debug HUD
+ * Extensions should be used for displaying information on screen for each Ability System Actor
+ */
+UCLASS(Abstract)
+class GAMEPLAYABILITIES_API UAbilitySystemDebugHUDExtension : public UObject
+{
+	GENERATED_BODY()
+
+public:
+
+	/** Returns whether this extension should be rendered for each actor */
+	virtual bool IsEnabled() const { return false; }
+
+	/**
+	 * Builds an array of strings to render for Ability System actor
+	 * 
+	 * @param Actor	The avatar actor for the ability system component
+	 * @param Comp	The Ability System Component being debugged
+	 * @param OutDebugStrings	Array to fill with debug strings
+	 */
+	virtual void GetDebugStrings(const AActor* Actor, const UAbilitySystemComponent* Comp, OUT TArray<FString>& OutDebugStrings) const { }
+};
+
+/**
+ * Ability System Debug HUD extension for drawing all owned tags by an actor
+ */
+UCLASS()
+class UAbilitySystemDebugHUDExtension_Tags final : public UAbilitySystemDebugHUDExtension
+{
+	GENERATED_BODY()
+
+	static bool bEnabled;
+
+public:
+
+	static void ToggleExtension(const TArray<FString>& Args, UWorld* World);
+
+	//~ Begin UAbilitySystemDebugHUDExtension interface
+	bool IsEnabled() const override;
+	void GetDebugStrings(const AActor* Actor, const UAbilitySystemComponent* Comp, OUT TArray<FString>& OutDebugStrings) const override;
+	//~ End UAbilitySystemDebugHUDExtension interface
+
+};
+
+/**
+ * Ability System Debug HUD extension for drawing attributes of an actor
+ */
+UCLASS()
+class UAbilitySystemDebugHUDExtension_Attributes final : public UAbilitySystemDebugHUDExtension
+{
+	GENERATED_BODY()
+
+	static TSet<FString> AttributesToDisplay;
+
+	static bool bIncludeModifiers;
+public:
+	static void ToggleExtension(const TArray<FString>& Args, UWorld* World);
+	
+	static void ToggleIncludeModifiers();
+
+	static void ClearDisplayedAttributes(const TArray<FString>& Args, UWorld* World);
+
+	//~ Begin UAbilitySystemDebugHUDExtension interface
+	bool IsEnabled() const override;
+	void GetDebugStrings(const AActor* Actor, const UAbilitySystemComponent* Comp, OUT TArray<FString>& OutDebugStrings) const override;
+	//~ End UAbilitySystemDebugHUDExtension interface
+
+};
+
+/**
+ * Ability System Debug HUD extension for drawing all blocked ability tags on an actor
+ */
+UCLASS()
+class UAbilitySystemDebugHUDExtension_BlockedAbilityTags final : public UAbilitySystemDebugHUDExtension
+{
+	GENERATED_BODY()
+
+	static bool bEnabled;
+
+public:
+
+	static void ToggleExtension(const TArray<FString>& Args, UWorld* World);
+
+	//~ Begin UAbilitySystemDebugHUDExtension interface
+	bool IsEnabled() const override;
+	void GetDebugStrings(const AActor* Actor, const UAbilitySystemComponent* Comp, OUT TArray<FString>& OutDebugStrings) const override;
+	//~ End UAbilitySystemDebugHUDExtension interface
+
+};
+
 UCLASS()
 class AAbilitySystemDebugHUD : public AHUD
 {
 	GENERATED_UCLASS_BODY()
 
-	/** main HUD update loop */
-	void DrawWithBackground(UFont* InFont, const FString& Text, const FColor& TextColor, EAlignHorizontal::Type HAlign, float& OffsetX, EAlignVertical::Type VAlign, float& OffsetY, float Alpha=1.f);
+public:
+
+	static void ToggleBasicHUD(const TArray<FString>& Args, UWorld* World);
 
+	/**
+	 * Notifies the AbilitySystemDebugHUD that an extension has been enabled/disabled
+	 * This will create the HUD if it did not exist and there are any extensions enabled
+	 * or destroy the HUD if it did exist and there are no extensions enabled
+	 */
+	static GAMEPLAYABILITIES_API void NotifyExtensionEnableChanged(UWorld* InWorld);
+
+	/** 
+	 * Replaces the AbilitySystemDebugHUD with a user class
+	 * NOTE: Only one AbilitySystemDebugHUD class can be active at a time.
+	 */
+	static GAMEPLAYABILITIES_API void RegisterHUDClass(TSubclassOf<AAbilitySystemDebugHUD> InHUDClass);
+	/** main HUD update loop */
 	void DrawDebugHUD(UCanvas* Canvas, APlayerController* PC);
 
+protected:
+
+	/** Returns whether debug information should be drawn for a given actor */
+	virtual bool ShouldDisplayDebugForActor(UCanvas* InCanvas, const AActor* Actor, const FVector& CameraPosition, const FVector& CameraDir) const;
+
 private:
 
-	void DrawDebugAbilitySystemComponent(UAbilitySystemComponent *Component);
+	void DrawWithBackground(UFont* InFont, const FString& Text, const FColor& TextColor, EAlignHorizontal::Type HAlign, float& OffsetX, EAlignVertical::Type VAlign, float& OffsetY, float Alpha = 1.f);
+
+	void DisplayDebugStrings(UCanvas* InCanvas, const AActor* Actor, const TArray<FString>& DebugStrings, const FVector& CameraPosition, const FVector& CameraDir, float& VerticalOffset) const;
+
+	void DrawDebugAbilitySystemComponent(UAbilitySystemComponent* Component);
+
+	void DrawAbilityDebugInfo(UCanvas* Canvas, APlayerController* PC) const;
+
+private:
+
+	static bool bEnableBasicHUD;
+
+	static TSubclassOf<AAbilitySystemDebugHUD> HUDClass;
 
 };
