Fix -> operator implementation for FActiveGameplayEffectIterator

#rnx

diff --git a/Engine/Plugins/Runtime/GameplayAbilities/Source/GameplayAbilities/Public/ActiveGameplayEffectIterator.h b/Engine/Plugins/Runtime/GameplayAbilities/Source/GameplayAbilities/Public/ActiveGameplayEffectIterator.h
--- a/Engine/Plugins/Runtime/GameplayAbilities/Source/GameplayAbilities/Public/ActiveGameplayEffectIterator.h
+++ b/Engine/Plugins/Runtime/GameplayAbilities/Source/GameplayAbilities/Public/ActiveGameplayEffectIterator.h
@@ -54,10 +54,10 @@ public:
 		return *Current;
 	}
 
-	ElementType& operator->() const
+	ElementType* operator->() const
 	{
 		check(Current);
-		return *Current;
+		return Current;
 	}
 
 	explicit operator bool() const
