Add all tags and source counts to the DebugAbilityTags HUD

diff --git a/Engine/Plugins/Runtime/GameplayAbilities/Source/GameplayAbilities/Private/AbilitySystemDebugHUD.cpp b/Engine/Plugins/Runtime/GameplayAbilities/Source/GameplayAbilities/Private/AbilitySystemDebugHUD.cpp
--- a/Engine/Plugins/Runtime/GameplayAbilities/Source/GameplayAbilities/Private/AbilitySystemDebugHUD.cpp
+++ b/Engine/Plugins/Runtime/GameplayAbilities/Source/GameplayAbilities/Private/AbilitySystemDebugHUD.cpp
@@ -11,6 +11,7 @@
 #include "EngineUtils.h"
 #include "AbilitySystemComponent.h"
 #include "AbilitySystemGlobals.h"
+#include "GameplayEffect.h"
 
 namespace
 {
@@ -53,18 +54,31 @@ void UAbilitySystemDebugHUDExtension_Tags::GetDebugStrings(const AActor* Actor,
 		return;
 	}
 
-	FGameplayTagContainer Container;
-	Comp->GetOwnedGameplayTags(Container);
-
-	FGameplayTagContainer LeafOnlyContainer;
-	for (auto ContainerIt = Container.CreateConstIterator(); ContainerIt; ++ContainerIt)
+	// Count up how many GameplayEffects granted the ASC each tag
+	TMap<FGameplayTag, uint32> TagCounts;
+	for (FActiveGameplayEffectsContainer::ConstIterator ActiveEffectIt = Comp->GetActiveGameplayEffects().CreateConstIterator(); ActiveEffectIt; ++ActiveEffectIt)
 	{
-		LeafOnlyContainer.AddLeafTag(*ContainerIt);
+		FGameplayTagContainer EffectTagContainer;
+		ActiveEffectIt->Spec.GetAllGrantedTags(EffectTagContainer);
+
+		for (auto EffectTagIt = EffectTagContainer.CreateConstIterator(); EffectTagIt; ++EffectTagIt)
+		{
+			++TagCounts.FindOrAdd(*EffectTagIt, 0);
+		}
 	}
 
-	for (auto ContainerIt = LeafOnlyContainer.CreateConstIterator(); ContainerIt; ++ContainerIt)
+	FGameplayTagContainer Container;
+	Comp->GetOwnedGameplayTags(Container);
+	for (auto ContainerIt = Container.CreateConstIterator(); ContainerIt; ++ContainerIt)
 	{
-		OutDebugStrings.Add(ContainerIt->ToString());
+		if (const uint32* TagCount = TagCounts.Find(*ContainerIt))
+		{
+			OutDebugStrings.Add(FString::Printf(TEXT("%s x%d"), *ContainerIt->ToString(), *TagCount));
+		}
+		else
+		{
+			OutDebugStrings.Add(ContainerIt->ToString());
+		}
 	}
 }
 
@@ -173,18 +187,31 @@ void UAbilitySystemDebugHUDExtension_BlockedAbilityTags::GetDebugStrings(const A
 		return;
 	}
 
-	FGameplayTagContainer Container;
-	Comp->GetBlockedAbilityTags(Container);
-
-	FGameplayTagContainer LeafOnlyContainer;
-	for (auto ContainerIt = Container.CreateConstIterator(); ContainerIt; ++ContainerIt)
+	// Count up how many GameplayEffects blocked each tag
+	TMap<FGameplayTag, uint32> TagCounts;
+	for (FActiveGameplayEffectsContainer::ConstIterator ActiveEffectIt = Comp->GetActiveGameplayEffects().CreateConstIterator(); ActiveEffectIt; ++ActiveEffectIt)
 	{
-		LeafOnlyContainer.AddLeafTag(*ContainerIt);
+		FGameplayTagContainer EffectTagContainer;
+		ActiveEffectIt->Spec.GetAllBlockedAbilityTags(EffectTagContainer);
+
+		for (auto EffectTagIt = EffectTagContainer.CreateConstIterator(); EffectTagIt; ++EffectTagIt)
+		{
+			++TagCounts.FindOrAdd(*EffectTagIt, 0);
+		}
 	}
 
-	for (auto ContainerIt = LeafOnlyContainer.CreateConstIterator(); ContainerIt; ++ContainerIt)
+	FGameplayTagContainer Container;
+	Comp->GetBlockedAbilityTags(Container);
+	for (auto ContainerIt = Container.CreateConstIterator(); ContainerIt; ++ContainerIt)
 	{
-		OutDebugStrings.Add(ContainerIt->ToString());
+		if (const uint32* TagCount = TagCounts.Find(*ContainerIt))
+		{
+			OutDebugStrings.Add(FString::Printf(TEXT("%s x%d"), *ContainerIt->ToString(), *TagCount));
+		}
+		else
+		{
+			OutDebugStrings.Add(ContainerIt->ToString());
+		}
 	}
 }
 
