ObjectPtr upgrade for engine plugins used from UnrealEditor -allmodules.  Automated upgrade process.

diff --git a/Engine/Plugins/Runtime/DataRegistry/Source/DataRegistry/Public/DataRegistrySource.h b/Engine/Plugins/Runtime/DataRegistry/Source/DataRegistry/Public/DataRegistrySource.h
--- a/Engine/Plugins/Runtime/DataRegistry/Source/DataRegistry/Public/DataRegistrySource.h
+++ b/Engine/Plugins/Runtime/DataRegistry/Source/DataRegistry/Public/DataRegistrySource.h
@@ -92,7 +92,7 @@ protected:
 
 	/** What data source we were created from, if this is a transient source */
 	UPROPERTY()
-	UDataRegistrySource* ParentSource = nullptr;
+	TObjectPtr<UDataRegistrySource> ParentSource = nullptr;
 };
 
 
@@ -159,7 +159,7 @@ protected:
 
 	/** Map from source identifier such as package name to registered child */
 	UPROPERTY(Transient)
-	TMap<FName, UDataRegistrySource*> RuntimeChildren;
+	TMap<FName, TObjectPtr<UDataRegistrySource>> RuntimeChildren;
 
 	/** List of desired source ids, in order */
 	TArray<FName> RuntimeNames;
diff --git a/Engine/Plugins/Runtime/DataRegistry/Source/DataRegistry/Public/DataRegistrySource_CurveTable.h b/Engine/Plugins/Runtime/DataRegistry/Source/DataRegistry/Public/DataRegistrySource_CurveTable.h
--- a/Engine/Plugins/Runtime/DataRegistry/Source/DataRegistry/Public/DataRegistrySource_CurveTable.h
+++ b/Engine/Plugins/Runtime/DataRegistry/Source/DataRegistry/Public/DataRegistrySource_CurveTable.h
@@ -29,11 +29,11 @@ public:
 protected:
 	/** Hard ref to loaded table */
 	UPROPERTY(Transient)
-	UCurveTable* CachedTable;
+	TObjectPtr<UCurveTable> CachedTable;
 
 	/** Preload table ref, will be set if this is a hard source */
 	UPROPERTY()
-	UCurveTable* PreloadTable;
+	TObjectPtr<UCurveTable> PreloadTable;
 
 	/** Last time this was accessed */
 	mutable float LastAccessTime;
diff --git a/Engine/Plugins/Runtime/DataRegistry/Source/DataRegistry/Public/DataRegistrySource_DataTable.h b/Engine/Plugins/Runtime/DataRegistry/Source/DataRegistry/Public/DataRegistrySource_DataTable.h
--- a/Engine/Plugins/Runtime/DataRegistry/Source/DataRegistry/Public/DataRegistrySource_DataTable.h
+++ b/Engine/Plugins/Runtime/DataRegistry/Source/DataRegistry/Public/DataRegistrySource_DataTable.h
@@ -43,11 +43,11 @@ public:
 protected:
 	/** Hard ref to loaded table */
 	UPROPERTY(Transient)
-	UDataTable* CachedTable;
+	TObjectPtr<UDataTable> CachedTable;
 
 	/** Preload table ref, will be set if this is a hard source */
 	UPROPERTY()
-	UDataTable* PreloadTable;
+	TObjectPtr<UDataTable> PreloadTable;
 
 	/** Last time this was accessed */
 	mutable float LastAccessTime;
diff --git a/Engine/Plugins/Runtime/GameplayAbilities/Source/GameplayAbilities/Public/Abilities/GameplayAbility.h b/Engine/Plugins/Runtime/GameplayAbilities/Source/GameplayAbilities/Public/Abilities/GameplayAbility.h
--- a/Engine/Plugins/Runtime/GameplayAbilities/Source/GameplayAbilities/Public/Abilities/GameplayAbility.h
+++ b/Engine/Plugins/Runtime/GameplayAbilities/Source/GameplayAbilities/Public/Abilities/GameplayAbility.h
@@ -798,7 +798,7 @@ protected:
 
 	/** List of currently active tasks, do not modify directly */
 	UPROPERTY()
-	TArray<UGameplayTask*>	ActiveTasks;
+	TArray<TObjectPtr<UGameplayTask>>	ActiveTasks;
 
 	/** Tasks can emit debug messages throughout their life for debugging purposes. Saved on the ability so that they persist after task is finished */
 	TArray<FAbilityTaskDebugMessage> TaskDebugMessages;
@@ -825,7 +825,7 @@ protected:
 
 	/** Active montage being played by this ability */
 	UPROPERTY()
-	class UAnimMontage* CurrentMontage;
+	TObjectPtr<class UAnimMontage> CurrentMontage;
 
 	// ----------------------------------------------------------------------------------------------------------------
 	//	Target Data
diff --git a/Engine/Plugins/Runtime/GameplayAbilities/Source/GameplayAbilities/Public/Abilities/GameplayAbilityTargetActor.h b/Engine/Plugins/Runtime/GameplayAbilities/Source/GameplayAbilities/Public/Abilities/GameplayAbilityTargetActor.h
--- a/Engine/Plugins/Runtime/GameplayAbilities/Source/GameplayAbilities/Public/Abilities/GameplayAbilityTargetActor.h
+++ b/Engine/Plugins/Runtime/GameplayAbilities/Source/GameplayAbilities/Public/Abilities/GameplayAbilityTargetActor.h
@@ -75,16 +75,16 @@ public:
 #endif // WITH_EDITOR
 	
 	UPROPERTY(BlueprintReadOnly, Category = "Targeting")
-	APlayerController* PrimaryPC;
+	TObjectPtr<APlayerController> PrimaryPC;
 
 	UPROPERTY()
-	UGameplayAbility* OwningAbility;
+	TObjectPtr<UGameplayAbility> OwningAbility;
 
 	UPROPERTY(BlueprintReadOnly, Replicated, Category = Targeting)
 	bool bDestroyOnConfirmation;
 
 	UPROPERTY(BlueprintReadOnly, Replicated, Category = Targeting)
-	AActor* SourceActor;
+	TObjectPtr<AActor> SourceActor;
 
 	/** Parameters for world reticle. Usage of these parameters is dependent on the reticle. */
 	UPROPERTY(BlueprintReadWrite, meta = (ExposeOnSpawn = true), Category = Targeting)
@@ -105,5 +105,5 @@ public:
 	FDelegateHandle GenericCancelHandle;
 
 	UPROPERTY()
-	UAbilitySystemComponent* GenericDelegateBoundASC;
+	TObjectPtr<UAbilitySystemComponent> GenericDelegateBoundASC;
 };
diff --git a/Engine/Plugins/Runtime/GameplayAbilities/Source/GameplayAbilities/Public/Abilities/GameplayAbilityTargetActor_ActorPlacement.h b/Engine/Plugins/Runtime/GameplayAbilities/Source/GameplayAbilities/Public/Abilities/GameplayAbilityTargetActor_ActorPlacement.h
--- a/Engine/Plugins/Runtime/GameplayAbilities/Source/GameplayAbilities/Public/Abilities/GameplayAbilityTargetActor_ActorPlacement.h
+++ b/Engine/Plugins/Runtime/GameplayAbilities/Source/GameplayAbilities/Public/Abilities/GameplayAbilityTargetActor_ActorPlacement.h
@@ -21,11 +21,11 @@ class AGameplayAbilityTargetActor_ActorPlacement : public AGameplayAbilityTarget
 
 	/** Actor we intend to place. */
 	UPROPERTY(BlueprintReadWrite, EditAnywhere, meta = (ExposeOnSpawn = true), Category = Targeting)
-	UClass* PlacedActorClass;		//Using a special class for replication purposes. (Not implemented yet)
+	TObjectPtr<UClass> PlacedActorClass;		//Using a special class for replication purposes. (Not implemented yet)
 	
 	/** Override Material 0 on our placed actor's meshes with this material for visualization. */
 	UPROPERTY(BlueprintReadWrite, EditAnywhere, meta = (ExposeOnSpawn = true), Category = Projectile)
-	UMaterialInterface *PlacedActorMaterial;
+	TObjectPtr<UMaterialInterface> PlacedActorMaterial;
 
 	/** Visualization for the intended location of the placed actor. */
 	TWeakObjectPtr<AGameplayAbilityWorldReticle_ActorVisualization> ActorVisualizationReticle;
diff --git a/Engine/Plugins/Runtime/GameplayAbilities/Source/GameplayAbilities/Public/Abilities/GameplayAbilityTargetDataFilter.h b/Engine/Plugins/Runtime/GameplayAbilities/Source/GameplayAbilities/Public/Abilities/GameplayAbilityTargetDataFilter.h
--- a/Engine/Plugins/Runtime/GameplayAbilities/Source/GameplayAbilities/Public/Abilities/GameplayAbilityTargetDataFilter.h
+++ b/Engine/Plugins/Runtime/GameplayAbilities/Source/GameplayAbilities/Public/Abilities/GameplayAbilityTargetDataFilter.h
@@ -65,7 +65,7 @@ struct GAMEPLAYABILITIES_API FGameplayTargetDataFilter
 
 	/** Actor we're comparing against. */
 	UPROPERTY()
-	AActor* SelfActor = nullptr;
+	TObjectPtr<AActor> SelfActor = nullptr;
 
 	/** Subclass actors must be to pass the filter. */
 	UPROPERTY(BlueprintReadWrite, EditAnywhere, meta = (ExposeOnSpawn = true), Category = Filter)
diff --git a/Engine/Plugins/Runtime/GameplayAbilities/Source/GameplayAbilities/Public/Abilities/GameplayAbilityTargetTypes.h b/Engine/Plugins/Runtime/GameplayAbilities/Source/GameplayAbilities/Public/Abilities/GameplayAbilityTargetTypes.h
--- a/Engine/Plugins/Runtime/GameplayAbilities/Source/GameplayAbilities/Public/Abilities/GameplayAbilityTargetTypes.h
+++ b/Engine/Plugins/Runtime/GameplayAbilities/Source/GameplayAbilities/Public/Abilities/GameplayAbilityTargetTypes.h
@@ -341,15 +341,15 @@ struct GAMEPLAYABILITIES_API FGameplayAbilityTargetingLocationInfo
 
 	/** A source actor is needed for Actor-based targeting, but not for Socket-based targeting. */
 	UPROPERTY(BlueprintReadWrite, meta = (ExposeOnSpawn = true), Category = Targeting)
-	AActor* SourceActor;
+	TObjectPtr<AActor> SourceActor;
 
 	/** Socket-based targeting requires a skeletal mesh component to check for the named socket. */
 	UPROPERTY(BlueprintReadWrite, meta = (ExposeOnSpawn = true), Category = Targeting)
-	UMeshComponent* SourceComponent;
+	TObjectPtr<UMeshComponent> SourceComponent;
 
 	/** Ability that will be using the targeting data */
 	UPROPERTY(BlueprintReadWrite, meta = (ExposeOnSpawn = true), Category = Targeting)
-	UGameplayAbility* SourceAbility;
+	TObjectPtr<UGameplayAbility> SourceAbility;
 
 	/** If SourceComponent is valid, this is the name of the socket transform that will be used. If no Socket is provided, SourceComponent's transform will be used. */
 	UPROPERTY(BlueprintReadWrite, meta = (ExposeOnSpawn = true), Category = Targeting)
diff --git a/Engine/Plugins/Runtime/GameplayAbilities/Source/GameplayAbilities/Public/Abilities/GameplayAbilityTypes.h b/Engine/Plugins/Runtime/GameplayAbilities/Source/GameplayAbilities/Public/Abilities/GameplayAbilityTypes.h
--- a/Engine/Plugins/Runtime/GameplayAbilities/Source/GameplayAbilities/Public/Abilities/GameplayAbilityTypes.h
+++ b/Engine/Plugins/Runtime/GameplayAbilities/Source/GameplayAbilities/Public/Abilities/GameplayAbilityTypes.h
@@ -212,7 +212,7 @@ struct GAMEPLAYABILITIES_API FGameplayAbilityRepAnimMontage
 
 	/** AnimMontage ref */
 	UPROPERTY()
-	UAnimMontage* AnimMontage;
+	TObjectPtr<UAnimMontage> AnimMontage;
 
 	/** Play Rate */
 	UPROPERTY()
@@ -300,7 +300,7 @@ struct GAMEPLAYABILITIES_API FGameplayAbilityLocalAnimMontage
 
 	/** What montage is being played */
 	UPROPERTY()
-	UAnimMontage* AnimMontage;
+	TObjectPtr<UAnimMontage> AnimMontage;
 
 	/** ID tied to a particular play of a montage, used to trigger replication when the same montage is played multiple times. This ID wraps around when it reaches its max value.  */
 	UPROPERTY()
@@ -312,7 +312,7 @@ struct GAMEPLAYABILITIES_API FGameplayAbilityLocalAnimMontage
 
 	/** The ability, if any, that instigated this montage */
 	UPROPERTY()
-	UGameplayAbility* AnimatingAbility;
+	TObjectPtr<UGameplayAbility> AnimatingAbility;
 };
 
 
@@ -337,19 +337,19 @@ struct GAMEPLAYABILITIES_API FGameplayEventData
 
 	/** The instigator of the event */
 	UPROPERTY(EditAnywhere, BlueprintReadWrite, Category = GameplayAbilityTriggerPayload)
-	const AActor* Instigator;
+	TObjectPtr<const AActor> Instigator;
 
 	/** The target of the event */
 	UPROPERTY(EditAnywhere, BlueprintReadWrite, Category = GameplayAbilityTriggerPayload)
-	const AActor* Target;
+	TObjectPtr<const AActor> Target;
 
 	/** An optional ability-specific object to be passed though the event */
 	UPROPERTY(EditAnywhere, BlueprintReadWrite, Category = GameplayAbilityTriggerPayload)
-	const UObject* OptionalObject;
+	TObjectPtr<const UObject> OptionalObject;
 
 	/** A second optional ability-specific object to be passed though the event */
 	UPROPERTY(EditAnywhere, BlueprintReadWrite, Category = GameplayAbilityTriggerPayload)
-	const UObject* OptionalObject2;
+	TObjectPtr<const UObject> OptionalObject2;
 
 	/** Polymorphic context information */
 	UPROPERTY(EditAnywhere, BlueprintReadWrite, Category = GameplayAbilityTriggerPayload)
@@ -402,7 +402,7 @@ struct FAbilityEndedData
 
 	/** Ability that ended, normally instance but could be CDO */
 	UPROPERTY()
-	UGameplayAbility* AbilityThatEnded;
+	TObjectPtr<UGameplayAbility> AbilityThatEnded;
 
 	/** Specific ability spec that ended */
 	UPROPERTY()
@@ -507,7 +507,7 @@ struct GAMEPLAYABILITIES_API FAttributeDefaults
 	TSubclassOf<UAttributeSet> Attributes;
 
 	UPROPERTY(EditAnywhere, Category = "AttributeTest")
-	UDataTable* DefaultStartingTable;
+	TObjectPtr<UDataTable> DefaultStartingTable;
 };
 
 
@@ -522,7 +522,7 @@ struct GAMEPLAYABILITIES_API FAbilityTaskDebugMessage
 	{ }
 
 	UPROPERTY()
-	UGameplayTask*	FromTask;
+	TObjectPtr<UGameplayTask>	FromTask;
 
 	UPROPERTY()
 	FString Message;
diff --git a/Engine/Plugins/Runtime/GameplayAbilities/Source/GameplayAbilities/Public/Abilities/GameplayAbilityWorldReticle.h b/Engine/Plugins/Runtime/GameplayAbilities/Source/GameplayAbilities/Public/Abilities/GameplayAbilityWorldReticle.h
--- a/Engine/Plugins/Runtime/GameplayAbilities/Source/GameplayAbilities/Public/Abilities/GameplayAbilityWorldReticle.h
+++ b/Engine/Plugins/Runtime/GameplayAbilities/Source/GameplayAbilities/Public/Abilities/GameplayAbilityWorldReticle.h
@@ -90,9 +90,9 @@ protected:
 	
 	/** This is used in the process of determining whether we should replicate to a specific client. */
 	UPROPERTY(BlueprintReadOnly, Category = "Network")
-	APlayerController* PrimaryPC;
+	TObjectPtr<APlayerController> PrimaryPC;
 
 	/** In the future, we may want to grab things like sockets off of this. */
 	UPROPERTY(BlueprintReadOnly, Category = "Network")
-	AActor* TargetingActor;
+	TObjectPtr<AActor> TargetingActor;
 };
diff --git a/Engine/Plugins/Runtime/GameplayAbilities/Source/GameplayAbilities/Public/Abilities/GameplayAbilityWorldReticle_ActorVisualization.h b/Engine/Plugins/Runtime/GameplayAbilities/Source/GameplayAbilities/Public/Abilities/GameplayAbilityWorldReticle_ActorVisualization.h
--- a/Engine/Plugins/Runtime/GameplayAbilities/Source/GameplayAbilities/Public/Abilities/GameplayAbilityWorldReticle_ActorVisualization.h
+++ b/Engine/Plugins/Runtime/GameplayAbilities/Source/GameplayAbilities/Public/Abilities/GameplayAbilityWorldReticle_ActorVisualization.h
@@ -22,11 +22,11 @@ public:
 private:
 	/** Hardcoded collision component, so other objects don't think they can collide with the visualization actor */
 	UPROPERTY()
-	class UCapsuleComponent* CollisionComponent;
+	TObjectPtr<class UCapsuleComponent> CollisionComponent;
 public:
 
 	UPROPERTY()
-	TArray<UActorComponent*> VisualizationComponents;
+	TArray<TObjectPtr<UActorComponent>> VisualizationComponents;
 
 	/** Overridable function called whenever this actor is being removed from a level */
 	virtual void EndPlay(const EEndPlayReason::Type EndPlayReason) override;
diff --git a/Engine/Plugins/Runtime/GameplayAbilities/Source/GameplayAbilities/Public/Abilities/GameplayAbility_Montage.h b/Engine/Plugins/Runtime/GameplayAbilities/Source/GameplayAbilities/Public/Abilities/GameplayAbility_Montage.h
--- a/Engine/Plugins/Runtime/GameplayAbilities/Source/GameplayAbilities/Public/Abilities/GameplayAbility_Montage.h
+++ b/Engine/Plugins/Runtime/GameplayAbilities/Source/GameplayAbilities/Public/Abilities/GameplayAbility_Montage.h
@@ -26,7 +26,7 @@ public:
 	virtual void ActivateAbility(const FGameplayAbilitySpecHandle Handle, const FGameplayAbilityActorInfo* OwnerInfo, const FGameplayAbilityActivationInfo ActivationInfo, const FGameplayEventData* TriggerEventData) override;
 
 	UPROPERTY(EditDefaultsOnly, Category = MontageAbility)
-	UAnimMontage *	MontageToPlay;
+	TObjectPtr<UAnimMontage> 	MontageToPlay;
 
 	UPROPERTY(EditDefaultsOnly, Category = MontageAbility)
 	float	PlayRate;
@@ -40,7 +40,7 @@ public:
 
 	/** Deprecated. Use GameplayEffectClassesWhileAnimating instead. */
 	UPROPERTY(VisibleDefaultsOnly, Category = Deprecated)
-	TArray<const UGameplayEffect*>	GameplayEffectsWhileAnimating;
+	TArray<TObjectPtr<const UGameplayEffect>>	GameplayEffectsWhileAnimating;
 
 	void OnMontageEnded(UAnimMontage* Montage, bool bInterrupted, TWeakObjectPtr<UAbilitySystemComponent> AbilitySystemComponent, TArray<struct FActiveGameplayEffectHandle>	AppliedEffects);
 
diff --git a/Engine/Plugins/Runtime/GameplayAbilities/Source/GameplayAbilities/Public/Abilities/Tasks/AbilityTask.h b/Engine/Plugins/Runtime/GameplayAbilities/Source/GameplayAbilities/Public/Abilities/Tasks/AbilityTask.h
--- a/Engine/Plugins/Runtime/GameplayAbilities/Source/GameplayAbilities/Public/Abilities/Tasks/AbilityTask.h
+++ b/Engine/Plugins/Runtime/GameplayAbilities/Source/GameplayAbilities/Public/Abilities/Tasks/AbilityTask.h
@@ -99,7 +99,7 @@ class GAMEPLAYABILITIES_API UAbilityTask : public UGameplayTask
 
 	/** GameplayAbility that created us */
 	UPROPERTY()
-	UGameplayAbility* Ability;
+	TObjectPtr<UGameplayAbility> Ability;
 
 	UPROPERTY()
 	TWeakObjectPtr<UAbilitySystemComponent> AbilitySystemComponent;
diff --git a/Engine/Plugins/Runtime/GameplayAbilities/Source/GameplayAbilities/Public/Abilities/Tasks/AbilityTask_ApplyRootMotionConstantForce.h b/Engine/Plugins/Runtime/GameplayAbilities/Source/GameplayAbilities/Public/Abilities/Tasks/AbilityTask_ApplyRootMotionConstantForce.h
--- a/Engine/Plugins/Runtime/GameplayAbilities/Source/GameplayAbilities/Public/Abilities/Tasks/AbilityTask_ApplyRootMotionConstantForce.h
+++ b/Engine/Plugins/Runtime/GameplayAbilities/Source/GameplayAbilities/Public/Abilities/Tasks/AbilityTask_ApplyRootMotionConstantForce.h
@@ -74,7 +74,7 @@ protected:
 	 *          is in units of seconds if this force has unlimited duration (Duration < 0)
 	 */
 	UPROPERTY(Replicated)
-	UCurveFloat* StrengthOverTime;
+	TObjectPtr<UCurveFloat> StrengthOverTime;
 
 	UPROPERTY(Replicated)
 	bool bEnableGravity;
diff --git a/Engine/Plugins/Runtime/GameplayAbilities/Source/GameplayAbilities/Public/Abilities/Tasks/AbilityTask_ApplyRootMotionJumpForce.h b/Engine/Plugins/Runtime/GameplayAbilities/Source/GameplayAbilities/Public/Abilities/Tasks/AbilityTask_ApplyRootMotionJumpForce.h
--- a/Engine/Plugins/Runtime/GameplayAbilities/Source/GameplayAbilities/Public/Abilities/Tasks/AbilityTask_ApplyRootMotionJumpForce.h
+++ b/Engine/Plugins/Runtime/GameplayAbilities/Source/GameplayAbilities/Public/Abilities/Tasks/AbilityTask_ApplyRootMotionJumpForce.h
@@ -82,7 +82,7 @@ protected:
 	bool bFinishOnLanded;
 
 	UPROPERTY(Replicated)
-	UCurveVector* PathOffsetCurve;
+	TObjectPtr<UCurveVector> PathOffsetCurve;
 
 	/** 
 	 *  Maps real time to movement fraction curve to affect the speed of the
@@ -92,7 +92,7 @@ protected:
 	 *  Default if unset is a 1:1 correspondence
 	 */
 	UPROPERTY(Replicated)
-	UCurveFloat* TimeMappingCurve;
+	TObjectPtr<UCurveFloat> TimeMappingCurve;
 
 	bool bHasLanded;
 };
diff --git a/Engine/Plugins/Runtime/GameplayAbilities/Source/GameplayAbilities/Public/Abilities/Tasks/AbilityTask_ApplyRootMotionMoveToActorForce.h b/Engine/Plugins/Runtime/GameplayAbilities/Source/GameplayAbilities/Public/Abilities/Tasks/AbilityTask_ApplyRootMotionMoveToActorForce.h
--- a/Engine/Plugins/Runtime/GameplayAbilities/Source/GameplayAbilities/Public/Abilities/Tasks/AbilityTask_ApplyRootMotionMoveToActorForce.h
+++ b/Engine/Plugins/Runtime/GameplayAbilities/Source/GameplayAbilities/Public/Abilities/Tasks/AbilityTask_ApplyRootMotionMoveToActorForce.h
@@ -82,7 +82,7 @@ protected:
 	FVector TargetLocation;
 
 	UPROPERTY(Replicated)
-	AActor* TargetActor;
+	TObjectPtr<AActor> TargetActor;
 
 	UPROPERTY(Replicated)
 	FVector TargetLocationOffset;
@@ -111,7 +111,7 @@ protected:
 	bool bRestrictSpeedToExpected;
 
 	UPROPERTY(Replicated)
-	UCurveVector* PathOffsetCurve;
+	TObjectPtr<UCurveVector> PathOffsetCurve;
 
 	/** 
 	 *  Maps real time to movement fraction curve to affect the speed of the
@@ -121,13 +121,13 @@ protected:
 	 *  Default if unset is a 1:1 correspondence
 	 */
 	UPROPERTY(Replicated)
-	UCurveFloat* TimeMappingCurve;
+	TObjectPtr<UCurveFloat> TimeMappingCurve;
 
 	UPROPERTY(Replicated)
-	UCurveFloat* TargetLerpSpeedHorizontalCurve;
+	TObjectPtr<UCurveFloat> TargetLerpSpeedHorizontalCurve;
 
 	UPROPERTY(Replicated)
-	UCurveFloat* TargetLerpSpeedVerticalCurve;
+	TObjectPtr<UCurveFloat> TargetLerpSpeedVerticalCurve;
 
 	EMovementMode PreviousMovementMode;
 	uint8 PreviousCustomMode;
diff --git a/Engine/Plugins/Runtime/GameplayAbilities/Source/GameplayAbilities/Public/Abilities/Tasks/AbilityTask_ApplyRootMotionMoveToForce.h b/Engine/Plugins/Runtime/GameplayAbilities/Source/GameplayAbilities/Public/Abilities/Tasks/AbilityTask_ApplyRootMotionMoveToForce.h
--- a/Engine/Plugins/Runtime/GameplayAbilities/Source/GameplayAbilities/Public/Abilities/Tasks/AbilityTask_ApplyRootMotionMoveToForce.h
+++ b/Engine/Plugins/Runtime/GameplayAbilities/Source/GameplayAbilities/Public/Abilities/Tasks/AbilityTask_ApplyRootMotionMoveToForce.h
@@ -68,7 +68,7 @@ protected:
 	bool bRestrictSpeedToExpected = false;
 
 	UPROPERTY(Replicated)
-	UCurveVector* PathOffsetCurve = nullptr;
+	TObjectPtr<UCurveVector> PathOffsetCurve = nullptr;
 
 	EMovementMode PreviousMovementMode = EMovementMode::MOVE_None;
 	uint8 PreviousCustomMovementMode = 0;
diff --git a/Engine/Plugins/Runtime/GameplayAbilities/Source/GameplayAbilities/Public/Abilities/Tasks/AbilityTask_ApplyRootMotionRadialForce.h b/Engine/Plugins/Runtime/GameplayAbilities/Source/GameplayAbilities/Public/Abilities/Tasks/AbilityTask_ApplyRootMotionRadialForce.h
--- a/Engine/Plugins/Runtime/GameplayAbilities/Source/GameplayAbilities/Public/Abilities/Tasks/AbilityTask_ApplyRootMotionRadialForce.h
+++ b/Engine/Plugins/Runtime/GameplayAbilities/Source/GameplayAbilities/Public/Abilities/Tasks/AbilityTask_ApplyRootMotionRadialForce.h
@@ -47,7 +47,7 @@ protected:
 	FVector Location;
 
 	UPROPERTY(Replicated)
-	AActor* LocationActor;
+	TObjectPtr<AActor> LocationActor;
 
 	UPROPERTY(Replicated)
 	float Strength;
@@ -73,7 +73,7 @@ protected:
 	 *  Curve X is 0 to 1 normalized distance (0 = 0cm, 1 = what Strength % at Radius units out)
 	 */
 	UPROPERTY(Replicated)
-	UCurveFloat* StrengthDistanceFalloff;
+	TObjectPtr<UCurveFloat> StrengthDistanceFalloff;
 
 	/** 
 	 *  Strength of the force over time
@@ -82,7 +82,7 @@ protected:
 	 *          is in units of seconds if this force has unlimited duration (Duration < 0)
 	 */
 	UPROPERTY(Replicated)
-	UCurveFloat* StrengthOverTime;
+	TObjectPtr<UCurveFloat> StrengthOverTime;
 
 	UPROPERTY(Replicated)
 	bool bUseFixedWorldDirection;
diff --git a/Engine/Plugins/Runtime/GameplayAbilities/Source/GameplayAbilities/Public/Abilities/Tasks/AbilityTask_ApplyRootMotion_Base.h b/Engine/Plugins/Runtime/GameplayAbilities/Source/GameplayAbilities/Public/Abilities/Tasks/AbilityTask_ApplyRootMotion_Base.h
--- a/Engine/Plugins/Runtime/GameplayAbilities/Source/GameplayAbilities/Public/Abilities/Tasks/AbilityTask_ApplyRootMotion_Base.h
+++ b/Engine/Plugins/Runtime/GameplayAbilities/Source/GameplayAbilities/Public/Abilities/Tasks/AbilityTask_ApplyRootMotion_Base.h
@@ -43,7 +43,7 @@ protected:
 	float FinishClampVelocity;
 
 	UPROPERTY()
-	UCharacterMovementComponent* MovementComponent; 
+	TObjectPtr<UCharacterMovementComponent> MovementComponent; 
 	
 	uint16 RootMotionSourceID;
 
diff --git a/Engine/Plugins/Runtime/GameplayAbilities/Source/GameplayAbilities/Public/Abilities/Tasks/AbilityTask_MoveToLocation.h b/Engine/Plugins/Runtime/GameplayAbilities/Source/GameplayAbilities/Public/Abilities/Tasks/AbilityTask_MoveToLocation.h
--- a/Engine/Plugins/Runtime/GameplayAbilities/Source/GameplayAbilities/Public/Abilities/Tasks/AbilityTask_MoveToLocation.h
+++ b/Engine/Plugins/Runtime/GameplayAbilities/Source/GameplayAbilities/Public/Abilities/Tasks/AbilityTask_MoveToLocation.h
@@ -69,10 +69,10 @@ protected:
 	float TimeMoveWillEnd;
 
 	UPROPERTY(Replicated)
-	UCurveFloat* LerpCurve;
+	TObjectPtr<UCurveFloat> LerpCurve;
 
 	UPROPERTY(Replicated)
-	UCurveVector* LerpCurveVector;
+	TObjectPtr<UCurveVector> LerpCurveVector;
 
 	
 };
diff --git a/Engine/Plugins/Runtime/GameplayAbilities/Source/GameplayAbilities/Public/Abilities/Tasks/AbilityTask_PlayMontageAndWait.h b/Engine/Plugins/Runtime/GameplayAbilities/Source/GameplayAbilities/Public/Abilities/Tasks/AbilityTask_PlayMontageAndWait.h
--- a/Engine/Plugins/Runtime/GameplayAbilities/Source/GameplayAbilities/Public/Abilities/Tasks/AbilityTask_PlayMontageAndWait.h
+++ b/Engine/Plugins/Runtime/GameplayAbilities/Source/GameplayAbilities/Public/Abilities/Tasks/AbilityTask_PlayMontageAndWait.h
@@ -74,7 +74,7 @@ protected:
 	FDelegateHandle InterruptedHandle;
 
 	UPROPERTY()
-	UAnimMontage* MontageToPlay;
+	TObjectPtr<UAnimMontage> MontageToPlay;
 
 	UPROPERTY()
 	float Rate;
diff --git a/Engine/Plugins/Runtime/GameplayAbilities/Source/GameplayAbilities/Public/Abilities/Tasks/AbilityTask_WaitAttributeChange.h b/Engine/Plugins/Runtime/GameplayAbilities/Source/GameplayAbilities/Public/Abilities/Tasks/AbilityTask_WaitAttributeChange.h
--- a/Engine/Plugins/Runtime/GameplayAbilities/Source/GameplayAbilities/Public/Abilities/Tasks/AbilityTask_WaitAttributeChange.h
+++ b/Engine/Plugins/Runtime/GameplayAbilities/Source/GameplayAbilities/Public/Abilities/Tasks/AbilityTask_WaitAttributeChange.h
@@ -62,7 +62,7 @@ class GAMEPLAYABILITIES_API UAbilityTask_WaitAttributeChange : public UAbilityTa
 protected:
 
 	UPROPERTY()
-	UAbilitySystemComponent* ExternalOwner;
+	TObjectPtr<UAbilitySystemComponent> ExternalOwner;
 
 	UAbilitySystemComponent* GetFocusedASC();
 
diff --git a/Engine/Plugins/Runtime/GameplayAbilities/Source/GameplayAbilities/Public/Abilities/Tasks/AbilityTask_WaitAttributeChangeRatioThreshold.h b/Engine/Plugins/Runtime/GameplayAbilities/Source/GameplayAbilities/Public/Abilities/Tasks/AbilityTask_WaitAttributeChangeRatioThreshold.h
--- a/Engine/Plugins/Runtime/GameplayAbilities/Source/GameplayAbilities/Public/Abilities/Tasks/AbilityTask_WaitAttributeChangeRatioThreshold.h
+++ b/Engine/Plugins/Runtime/GameplayAbilities/Source/GameplayAbilities/Public/Abilities/Tasks/AbilityTask_WaitAttributeChangeRatioThreshold.h
@@ -51,7 +51,7 @@ protected:
 	FTimerHandle CheckAttributeTimer;
 
 	UPROPERTY()
-	UAbilitySystemComponent* ExternalOwner;
+	TObjectPtr<UAbilitySystemComponent> ExternalOwner;
 
 	UAbilitySystemComponent* GetFocusedASC();
 
diff --git a/Engine/Plugins/Runtime/GameplayAbilities/Source/GameplayAbilities/Public/Abilities/Tasks/AbilityTask_WaitGameplayEvent.h b/Engine/Plugins/Runtime/GameplayAbilities/Source/GameplayAbilities/Public/Abilities/Tasks/AbilityTask_WaitGameplayEvent.h
--- a/Engine/Plugins/Runtime/GameplayAbilities/Source/GameplayAbilities/Public/Abilities/Tasks/AbilityTask_WaitGameplayEvent.h
+++ b/Engine/Plugins/Runtime/GameplayAbilities/Source/GameplayAbilities/Public/Abilities/Tasks/AbilityTask_WaitGameplayEvent.h
@@ -42,7 +42,7 @@ class GAMEPLAYABILITIES_API UAbilityTask_WaitGameplayEvent : public UAbilityTask
 	FGameplayTag Tag;
 
 	UPROPERTY()
-	UAbilitySystemComponent* OptionalExternalTarget;
+	TObjectPtr<UAbilitySystemComponent> OptionalExternalTarget;
 
 	bool UseExternalTarget;	
 	bool OnlyTriggerOnce;
diff --git a/Engine/Plugins/Runtime/GameplayAbilities/Source/GameplayAbilities/Public/Abilities/Tasks/AbilityTask_WaitGameplayTagBase.h b/Engine/Plugins/Runtime/GameplayAbilities/Source/GameplayAbilities/Public/Abilities/Tasks/AbilityTask_WaitGameplayTagBase.h
--- a/Engine/Plugins/Runtime/GameplayAbilities/Source/GameplayAbilities/Public/Abilities/Tasks/AbilityTask_WaitGameplayTagBase.h
+++ b/Engine/Plugins/Runtime/GameplayAbilities/Source/GameplayAbilities/Public/Abilities/Tasks/AbilityTask_WaitGameplayTagBase.h
@@ -32,7 +32,7 @@ protected:
 	bool RegisteredCallback;
 	
 	UPROPERTY()
-	UAbilitySystemComponent* OptionalExternalTarget;
+	TObjectPtr<UAbilitySystemComponent> OptionalExternalTarget;
 
 	bool UseExternalTarget;	
 	bool OnlyTriggerOnce;
diff --git a/Engine/Plugins/Runtime/GameplayAbilities/Source/GameplayAbilities/Public/Abilities/Tasks/AbilityTask_WaitGameplayTagQuery.h b/Engine/Plugins/Runtime/GameplayAbilities/Source/GameplayAbilities/Public/Abilities/Tasks/AbilityTask_WaitGameplayTagQuery.h
--- a/Engine/Plugins/Runtime/GameplayAbilities/Source/GameplayAbilities/Public/Abilities/Tasks/AbilityTask_WaitGameplayTagQuery.h
+++ b/Engine/Plugins/Runtime/GameplayAbilities/Source/GameplayAbilities/Public/Abilities/Tasks/AbilityTask_WaitGameplayTagQuery.h
@@ -78,7 +78,7 @@ protected:
 	
 	/** This is the optional external target to use when getting the ASC to get tags from. */
 	UPROPERTY()
-	UAbilitySystemComponent* OptionalExternalTarget = nullptr;
+	TObjectPtr<UAbilitySystemComponent> OptionalExternalTarget = nullptr;
 
 	/** This indicates when to Trigger the Triggered output for this node. */
 	EWaitGameplayTagQueryTriggerCondition TriggerCondition = EWaitGameplayTagQueryTriggerCondition::WhenTrue;
diff --git a/Engine/Plugins/Runtime/GameplayAbilities/Source/GameplayAbilities/Public/Abilities/Tasks/AbilityTask_WaitTargetData.h b/Engine/Plugins/Runtime/GameplayAbilities/Source/GameplayAbilities/Public/Abilities/Tasks/AbilityTask_WaitTargetData.h
--- a/Engine/Plugins/Runtime/GameplayAbilities/Source/GameplayAbilities/Public/Abilities/Tasks/AbilityTask_WaitTargetData.h
+++ b/Engine/Plugins/Runtime/GameplayAbilities/Source/GameplayAbilities/Public/Abilities/Tasks/AbilityTask_WaitTargetData.h
@@ -82,7 +82,7 @@ protected:
 
 	/** The TargetActor that we spawned */
 	UPROPERTY()
-	AGameplayAbilityTargetActor* TargetActor;
+	TObjectPtr<AGameplayAbilityTargetActor> TargetActor;
 
 	TEnumAsByte<EGameplayTargetingConfirmation::Type> ConfirmationType;
 
diff --git a/Engine/Plugins/Runtime/GameplayAbilities/Source/GameplayAbilities/Public/Abilities/Tasks/AbilityTask_WaitVelocityChange.h b/Engine/Plugins/Runtime/GameplayAbilities/Source/GameplayAbilities/Public/Abilities/Tasks/AbilityTask_WaitVelocityChange.h
--- a/Engine/Plugins/Runtime/GameplayAbilities/Source/GameplayAbilities/Public/Abilities/Tasks/AbilityTask_WaitVelocityChange.h
+++ b/Engine/Plugins/Runtime/GameplayAbilities/Source/GameplayAbilities/Public/Abilities/Tasks/AbilityTask_WaitVelocityChange.h
@@ -30,7 +30,7 @@ class GAMEPLAYABILITIES_API UAbilityTask_WaitVelocityChange: public UAbilityTask
 protected:
 
 	UPROPERTY()
-	UMovementComponent*	CachedMovementComponent;
+	TObjectPtr<UMovementComponent>	CachedMovementComponent;
 
 	float	MinimumMagnitude;
 	FVector Direction;
diff --git a/Engine/Plugins/Runtime/GameplayAbilities/Source/GameplayAbilities/Public/AbilitySystemComponent.h b/Engine/Plugins/Runtime/GameplayAbilities/Source/GameplayAbilities/Public/AbilitySystemComponent.h
--- a/Engine/Plugins/Runtime/GameplayAbilities/Source/GameplayAbilities/Public/AbilitySystemComponent.h
+++ b/Engine/Plugins/Runtime/GameplayAbilities/Source/GameplayAbilities/Public/AbilitySystemComponent.h
@@ -1315,7 +1315,7 @@ class GAMEPLAYABILITIES_API UAbilitySystemComponent : public UGameplayTasksCompo
 
 	/** List of currently active targeting actors */
 	UPROPERTY()
-	TArray<AGameplayAbilityTargetActor*> SpawnedTargetActors;
+	TArray<TObjectPtr<AGameplayAbilityTargetActor>> SpawnedTargetActors;
 
 	/** Bind to an input component with some default action names */
 	virtual void BindToInputComponent(UInputComponent* InputComponent);
@@ -1442,11 +1442,11 @@ private:
 
 	/** The actor that owns this component logically */
 	UPROPERTY(ReplicatedUsing = OnRep_OwningActor)
-	AActor* OwnerActor;
+	TObjectPtr<AActor> OwnerActor;
 
 	/** The actor that is the physical representation used for abilities. Can be NULL */
 	UPROPERTY(ReplicatedUsing = OnRep_OwningActor)
-	AActor* AvatarActor;
+	TObjectPtr<AActor> AvatarActor;
 
 public:
 
@@ -1627,7 +1627,7 @@ protected:
 	/** Full list of all instance-per-execution gameplay abilities associated with this component */
 	UE_DEPRECATED(5.1, "This array will be made private. Use GetReplicatedInstancedAbilities, AddReplicatedInstancedAbility or RemoveReplicatedInstancedAbility instead.")
 	UPROPERTY()
-	TArray<UGameplayAbility*>	AllReplicatedInstancedAbilities;
+	TArray<TObjectPtr<UGameplayAbility>>	AllReplicatedInstancedAbilities;
 
 	/** Full list of all instance-per-execution gameplay abilities associated with this component */
 	PRAGMA_DISABLE_DEPRECATION_WARNINGS
@@ -1869,7 +1869,7 @@ private:
 
 	/** List of attribute sets */
 	UPROPERTY(Replicated, ReplicatedUsing=OnRep_SpawnedAttributes)
-	TArray<UAttributeSet*>	SpawnedAttributes;
+	TArray<TObjectPtr<UAttributeSet>>	SpawnedAttributes;
 
 	UFUNCTION()
 	void OnRep_SpawnedAttributes(const TArray<UAttributeSet*>& PreviousSpawnedAttributes);
diff --git a/Engine/Plugins/Runtime/GameplayAbilities/Source/GameplayAbilities/Public/AbilitySystemGlobals.h b/Engine/Plugins/Runtime/GameplayAbilities/Source/GameplayAbilities/Public/AbilitySystemGlobals.h
--- a/Engine/Plugins/Runtime/GameplayAbilities/Source/GameplayAbilities/Public/AbilitySystemGlobals.h
+++ b/Engine/Plugins/Runtime/GameplayAbilities/Source/GameplayAbilities/Public/AbilitySystemGlobals.h
@@ -34,7 +34,7 @@ struct FNetSerializeScriptStructCache
 	bool NetSerialize(FArchive& Ar, UScriptStruct*& Struct);
 
 	UPROPERTY()
-	TArray<UScriptStruct*> ScriptStructs;
+	TArray<TObjectPtr<UScriptStruct>> ScriptStructs;
 };
 
 
@@ -344,14 +344,14 @@ protected:
 	FSoftObjectPath GlobalCurveTableName;
 
 	UPROPERTY()
-	UCurveTable* GlobalCurveTable;
+	TObjectPtr<UCurveTable> GlobalCurveTable;
 
 	/** Holds information about the valid attributes' min and max values and stacking rules */
 	UPROPERTY(config)
 	FSoftObjectPath GlobalAttributeMetaDataTableName;
 
 	UPROPERTY()
-	UDataTable* GlobalAttributeMetaDataTable;
+	TObjectPtr<UDataTable> GlobalAttributeMetaDataTable;
 
 	/** Holds default values for attribute sets, keyed off of Name/Levels. NOTE: Preserved for backwards compatibility, should use the array version below now */
 	UPROPERTY(config)
@@ -363,7 +363,7 @@ protected:
 
 	/** Curve tables containing default values for attribute sets, keyed off of Name/Levels */
 	UPROPERTY()
-	TArray<UCurveTable*> GlobalAttributeDefaultsTables;
+	TArray<TObjectPtr<UCurveTable>> GlobalAttributeDefaultsTables;
 
 	/** Class reference to gameplay cue manager. Use this if you want to just instantiate a class for your gameplay cue manager without having to create an asset. */
 	UPROPERTY(config)
@@ -382,7 +382,7 @@ protected:
 	FSoftObjectPath GameplayTagResponseTableName;
 
 	UPROPERTY()
-	UGameplayTagReponseTable* GameplayTagResponseTable;
+	TObjectPtr<UGameplayTagReponseTable> GameplayTagResponseTable;
 
 	/** Set to true if you want clients to try to predict gameplay effects done to targets. If false it will only predict self effects */
 	UPROPERTY(config)
@@ -397,7 +397,7 @@ protected:
 
 	/** Manager for all gameplay cues */
 	UPROPERTY()
-	UGameplayCueManager* GlobalGameplayCueManager;
+	TObjectPtr<UGameplayCueManager> GlobalGameplayCueManager;
 
 	/** Used to initialize attribute sets */
 	TSharedPtr<FAttributeSetInitter> GlobalAttributeSetInitter;
diff --git a/Engine/Plugins/Runtime/GameplayAbilities/Source/GameplayAbilities/Public/AbilitySystemTestPawn.h b/Engine/Plugins/Runtime/GameplayAbilities/Source/GameplayAbilities/Public/AbilitySystemTestPawn.h
--- a/Engine/Plugins/Runtime/GameplayAbilities/Source/GameplayAbilities/Public/AbilitySystemTestPawn.h
+++ b/Engine/Plugins/Runtime/GameplayAbilities/Source/GameplayAbilities/Public/AbilitySystemTestPawn.h
@@ -23,7 +23,7 @@ class GAMEPLAYABILITIES_API AAbilitySystemTestPawn : public ADefaultPawn, public
 private:
 	/** DefaultPawn collision component */
 	UPROPERTY(Category = AbilitySystem, VisibleAnywhere, BlueprintReadOnly, meta = (AllowPrivateAccess = "true"))
-	class UAbilitySystemComponent* AbilitySystemComponent;
+	TObjectPtr<class UAbilitySystemComponent> AbilitySystemComponent;
 public:
 
 	//UPROPERTY(EditDefaultsOnly, Category=GameplayEffects)
diff --git a/Engine/Plugins/Runtime/GameplayAbilities/Source/GameplayAbilities/Public/AttributeSet.h b/Engine/Plugins/Runtime/GameplayAbilities/Source/GameplayAbilities/Public/AttributeSet.h
--- a/Engine/Plugins/Runtime/GameplayAbilities/Source/GameplayAbilities/Public/AttributeSet.h
+++ b/Engine/Plugins/Runtime/GameplayAbilities/Source/GameplayAbilities/Public/AttributeSet.h
@@ -157,7 +157,7 @@ private:
 	//FProperty*	Attribute;
 
 	UPROPERTY(Category = GameplayAttribute, VisibleAnywhere)
-	UStruct* AttributeOwner;
+	TObjectPtr<UStruct> AttributeOwner;
 };
 
 #if WITH_EDITORONLY_DATA
diff --git a/Engine/Plugins/Runtime/GameplayAbilities/Source/GameplayAbilities/Public/GameplayAbilitySpec.h b/Engine/Plugins/Runtime/GameplayAbilities/Source/GameplayAbilities/Public/GameplayAbilitySpec.h
--- a/Engine/Plugins/Runtime/GameplayAbilities/Source/GameplayAbilities/Public/GameplayAbilitySpec.h
+++ b/Engine/Plugins/Runtime/GameplayAbilities/Source/GameplayAbilities/Public/GameplayAbilitySpec.h
@@ -236,7 +236,7 @@ struct GAMEPLAYABILITIES_API FGameplayAbilitySpec : public FFastArraySerializerI
 	
 	/** Ability of the spec (Always the CDO. This should be const but too many things modify it currently) */
 	UPROPERTY()
-	UGameplayAbility* Ability;
+	TObjectPtr<UGameplayAbility> Ability;
 	
 	/** Level of Ability */
 	UPROPERTY()
@@ -280,11 +280,11 @@ struct GAMEPLAYABILITIES_API FGameplayAbilitySpec : public FFastArraySerializerI
 
 	/** Non replicating instances of this ability. */
 	UPROPERTY(NotReplicated)
-	TArray<UGameplayAbility*> NonReplicatedInstances;
+	TArray<TObjectPtr<UGameplayAbility>> NonReplicatedInstances;
 
 	/** Replicated instances of this ability.. */
 	UPROPERTY()
-	TArray<UGameplayAbility*> ReplicatedInstances;
+	TArray<TObjectPtr<UGameplayAbility>> ReplicatedInstances;
 
 	/** Handle to GE that granted us (usually invalid) */
 	UPROPERTY(NotReplicated)
@@ -334,7 +334,7 @@ struct GAMEPLAYABILITIES_API FGameplayAbilitySpecContainer : public FFastArraySe
 
 	/** Component that owns this list */
 	UPROPERTY(NotReplicated)
-	UAbilitySystemComponent* Owner;
+	TObjectPtr<UAbilitySystemComponent> Owner;
 
 	/** Initializes Owner variable */
 	void RegisterWithOwner(UAbilitySystemComponent* Owner);
diff --git a/Engine/Plugins/Runtime/GameplayAbilities/Source/GameplayAbilities/Public/GameplayCueInterface.h b/Engine/Plugins/Runtime/GameplayAbilities/Source/GameplayAbilities/Public/GameplayCueInterface.h
--- a/Engine/Plugins/Runtime/GameplayAbilities/Source/GameplayAbilities/Public/GameplayCueInterface.h
+++ b/Engine/Plugins/Runtime/GameplayAbilities/Source/GameplayAbilities/Public/GameplayCueInterface.h
@@ -153,7 +153,7 @@ private:
 	int32 GetGameStateTime(const UWorld* World) const;
 
 	UPROPERTY(NotReplicated)
-	class UAbilitySystemComponent*	Owner = nullptr;
+	TObjectPtr<class UAbilitySystemComponent>	Owner = nullptr;
 	
 	friend struct FActiveGameplayCue;
 };
@@ -229,7 +229,7 @@ private:
 	TBitArray< TInlineAllocator<NumInlineTags> >			LocalBitMask;
 
 	UPROPERTY()
-	UAbilitySystemComponent* Owner = nullptr;
+	TObjectPtr<UAbilitySystemComponent> Owner = nullptr;
 
 	int32 LastSourceArrayReplicationKey = -1;
 
diff --git a/Engine/Plugins/Runtime/GameplayAbilities/Source/GameplayAbilities/Public/GameplayCueManager.h b/Engine/Plugins/Runtime/GameplayAbilities/Source/GameplayAbilities/Public/GameplayCueManager.h
--- a/Engine/Plugins/Runtime/GameplayAbilities/Source/GameplayAbilities/Public/GameplayCueManager.h
+++ b/Engine/Plugins/Runtime/GameplayAbilities/Source/GameplayAbilities/Public/GameplayCueManager.h
@@ -73,15 +73,15 @@ struct FGameplayCueObjectLibrary
 
 	/** Object library for actor based notifies */
 	UPROPERTY()
-	UObjectLibrary* ActorObjectLibrary;
+	TObjectPtr<UObjectLibrary> ActorObjectLibrary;
 
 	/** Object library for object based notifies */
 	UPROPERTY()
-	UObjectLibrary* StaticObjectLibrary;
+	TObjectPtr<UObjectLibrary> StaticObjectLibrary;
 
 	/** Set to put the loaded asset data into. If null we will use the global set (RuntimeGameplayCueObjectLibrary.CueSet) */
 	UPROPERTY()
-	UGameplayCueSet* CueSet;
+	TObjectPtr<UGameplayCueSet> CueSet;
 
 	/** Priority to use if async loading */
 	TAsyncLoadPriority AsyncPriority;
@@ -380,7 +380,7 @@ protected:
 
 	/** Hardref to the gameplaycue notify classes we have async loaded*/
 	UPROPERTY(transient)
-	TArray<UClass*> LoadedGameplayCueNotifyClasses;
+	TArray<TObjectPtr<UClass>> LoadedGameplayCueNotifyClasses;
 
 	/** Classes that we need to preallocate instances for */
 	UPROPERTY(transient)
diff --git a/Engine/Plugins/Runtime/GameplayAbilities/Source/GameplayAbilities/Public/GameplayCueNotifyTypes.h b/Engine/Plugins/Runtime/GameplayAbilities/Source/GameplayAbilities/Public/GameplayCueNotifyTypes.h
--- a/Engine/Plugins/Runtime/GameplayAbilities/Source/GameplayAbilities/Public/GameplayCueNotifyTypes.h
+++ b/Engine/Plugins/Runtime/GameplayAbilities/Source/GameplayAbilities/Public/GameplayCueNotifyTypes.h
@@ -272,15 +272,15 @@ public:
 
 	// List of FX components spawned.  There may be null pointers here as it matches the defined order.
 	UPROPERTY(BlueprintReadOnly, Transient, Category = GameplayCueNotify)
-	TArray<UFXSystemComponent*> FxSystemComponents;
+	TArray<TObjectPtr<UFXSystemComponent>> FxSystemComponents;
 
 	// List of audio components spawned.  There may be null pointers here as it matches the defined order.
 	UPROPERTY(BlueprintReadOnly, Transient, Category = GameplayCueNotify)
-	TArray<UAudioComponent*> AudioComponents;
+	TArray<TObjectPtr<UAudioComponent>> AudioComponents;
 
 	// List of camera shakes played.  There will be one camera shake per local player controller if shake is played in world.
 	UPROPERTY(BlueprintReadOnly, Transient, Category = GameplayCueNotify)
-	TArray<UCameraShakeBase*> CameraShakes;
+	TArray<TObjectPtr<UCameraShakeBase>> CameraShakes;
 
 	// List of camera len effects spawned.  There will be one camera lens effect per local player controller if the effect is played in world.
 	UPROPERTY(BlueprintReadOnly, Transient, Category = GameplayCueNotify)
@@ -288,15 +288,15 @@ public:
 
 	// Force feedback component that was spawned.  This is only valid when force feedback is set to play in world.
 	UPROPERTY(BlueprintReadOnly, Transient, Category = GameplayCueNotify)
-	UForceFeedbackComponent* ForceFeedbackComponent;
+	TObjectPtr<UForceFeedbackComponent> ForceFeedbackComponent;
 
 	// Player controller used to play the force feedback effect.  Used to stop the effect later.
 	UPROPERTY(Transient)
-	APlayerController* ForceFeedbackTargetPC;
+	TObjectPtr<APlayerController> ForceFeedbackTargetPC;
 
 	// Spawned decal component.  This may be null.
 	UPROPERTY(BlueprintReadOnly, Transient, Category = GameplayCueNotify)
-	UDecalComponent* DecalComponent;
+	TObjectPtr<UDecalComponent> DecalComponent;
 };
 
 
@@ -330,7 +330,7 @@ public:
 
 	// Niagara FX system to spawn.
 	UPROPERTY(EditDefaultsOnly, BlueprintReadOnly, Category = GameplayCueNotify)
-	UNiagaraSystem* NiagaraSystem;
+	TObjectPtr<UNiagaraSystem> NiagaraSystem;
 
 	// If enabled, use the spawn condition override and not the default one.
 	UPROPERTY(EditDefaultsOnly, BlueprintReadOnly, Category = GameplayCueNotify, Meta = (InlineEditConditionToggle))
@@ -394,10 +394,10 @@ public:
 
 	// Sound to play.
 	UPROPERTY(EditDefaultsOnly, BlueprintReadOnly, Category = GameplayCueNotify)
-	USoundBase* Sound;
+	TObjectPtr<USoundBase> Sound;
 
 	UPROPERTY(meta = (DeprecatedProperty, DeprecationMessage = "5.0 - SoundCue is deprecated. Instead use the Sound property. The type is USoundBase not USoundCue."))
-	USoundBase* SoundCue;
+	TObjectPtr<USoundBase> SoundCue;
 
 	// How long it should take to fade out.  Only used on looping gameplay cues.
 	UPROPERTY(EditDefaultsOnly, BlueprintReadOnly, Category = GameplayCueNotify)
@@ -578,7 +578,7 @@ public:
 
 	// Force feedback effect to play.
 	UPROPERTY(EditDefaultsOnly, BlueprintReadOnly, Category = GameplayCueNotify)
-	UForceFeedbackEffect* ForceFeedbackEffect;
+	TObjectPtr<UForceFeedbackEffect> ForceFeedbackEffect;
 
 	// Tag used to identify the force feedback effect so it can later be canceled.
 	// Warning: If this is "None" it will stop ALL force feedback effects if it is canceled.
@@ -607,7 +607,7 @@ public:
 
 	// How the force feedback is attenuated over distance when played in world.
 	UPROPERTY(EditDefaultsOnly, BlueprintReadOnly, Category = GameplayCueNotify, Meta = (EditCondition = "bPlayInWorld"))
-	UForceFeedbackAttenuation* WorldAttenuation;
+	TObjectPtr<UForceFeedbackAttenuation> WorldAttenuation;
 };
 
 
@@ -639,7 +639,7 @@ public:
 
 	// Decal material to spawn.
 	UPROPERTY(EditDefaultsOnly, BlueprintReadOnly, Category = GameplayCueNotify)
-	UMaterialInterface* DecalMaterial;
+	TObjectPtr<UMaterialInterface> DecalMaterial;
 
 	// Decal size in local space (does not include the component scale).
 	UPROPERTY(EditDefaultsOnly, BlueprintReadOnly, Category = GameplayCueNotify, meta = (AllowPreserveRatio = "true"))
diff --git a/Engine/Plugins/Runtime/GameplayAbilities/Source/GameplayAbilities/Public/GameplayCueNotify_HitImpact.h b/Engine/Plugins/Runtime/GameplayAbilities/Source/GameplayAbilities/Public/GameplayCueNotify_HitImpact.h
--- a/Engine/Plugins/Runtime/GameplayAbilities/Source/GameplayAbilities/Public/GameplayCueNotify_HitImpact.h
+++ b/Engine/Plugins/Runtime/GameplayAbilities/Source/GameplayAbilities/Public/GameplayCueNotify_HitImpact.h
@@ -26,9 +26,9 @@ class GAMEPLAYABILITIES_API UGameplayCueNotify_HitImpact : public UGameplayCueNo
 	virtual void HandleGameplayCue(AActor* MyTarget, EGameplayCueEvent::Type EventType, const FGameplayCueParameters& Parameters) override;
 
 	UPROPERTY(EditAnywhere, BlueprintReadWrite, Category = GameplayCue)
-	USoundBase* Sound;
+	TObjectPtr<USoundBase> Sound;
 
 	/** Effects to play for weapon attacks against specific surfaces */
 	UPROPERTY(EditAnywhere, BlueprintReadWrite, Category = GameplayCue)
-	UParticleSystem* ParticleSystem;
+	TObjectPtr<UParticleSystem> ParticleSystem;
 };
diff --git a/Engine/Plugins/Runtime/GameplayAbilities/Source/GameplayAbilities/Public/GameplayCueSet.h b/Engine/Plugins/Runtime/GameplayAbilities/Source/GameplayAbilities/Public/GameplayCueSet.h
--- a/Engine/Plugins/Runtime/GameplayAbilities/Source/GameplayAbilities/Public/GameplayCueSet.h
+++ b/Engine/Plugins/Runtime/GameplayAbilities/Source/GameplayAbilities/Public/GameplayCueSet.h
@@ -28,7 +28,7 @@ struct FGameplayCueNotifyData
 	FSoftObjectPath GameplayCueNotifyObj;
 
 	UPROPERTY(transient)
-	UClass* LoadedGameplayCueClass;
+	TObjectPtr<UClass> LoadedGameplayCueClass;
 
 	int32 ParentDataIdx;
 };
diff --git a/Engine/Plugins/Runtime/GameplayAbilities/Source/GameplayAbilities/Public/GameplayCueTranslator.h b/Engine/Plugins/Runtime/GameplayAbilities/Source/GameplayAbilities/Public/GameplayCueTranslator.h
--- a/Engine/Plugins/Runtime/GameplayAbilities/Source/GameplayAbilities/Public/GameplayCueTranslator.h
+++ b/Engine/Plugins/Runtime/GameplayAbilities/Source/GameplayAbilities/Public/GameplayCueTranslator.h
@@ -128,7 +128,7 @@ struct FGameplayCueTranslationLink
 
 	/** The rule that provides this translation */
 	UPROPERTY()
-	const UGameplayCueTranslator* RulesCDO = nullptr;
+	TObjectPtr<const UGameplayCueTranslator> RulesCDO = nullptr;
 
 	/** Fixed size lookup. The RulesCDO return the index into this which will translate to the new node. */
 	TArray< FGameplayCueTranslatorNodeIndex > NodeLookup;
@@ -223,7 +223,7 @@ private:
 
 	/** Cached reference to tag manager */
 	UPROPERTY()
-	UGameplayTagsManager* TagManager = nullptr;
+	TObjectPtr<UGameplayTagsManager> TagManager = nullptr;
 		
 	/** All name swpa rules we have gathered */
 	TArray<FNameSwapData> AllNameSwaps;
diff --git a/Engine/Plugins/Runtime/GameplayAbilities/Source/GameplayAbilities/Public/GameplayCue_Types.h b/Engine/Plugins/Runtime/GameplayAbilities/Source/GameplayAbilities/Public/GameplayCue_Types.h
--- a/Engine/Plugins/Runtime/GameplayAbilities/Source/GameplayAbilities/Public/GameplayCue_Types.h
+++ b/Engine/Plugins/Runtime/GameplayAbilities/Source/GameplayAbilities/Public/GameplayCue_Types.h
@@ -48,7 +48,7 @@ struct FGameplayCuePendingExecute
 
 	/** What component to send the cue on */
 	UPROPERTY()
-	UAbilitySystemComponent* OwningComponent = nullptr;
+	TObjectPtr<UAbilitySystemComponent> OwningComponent = nullptr;
 
 	/** If this cue is from a spec, here's the copy of that spec */
 	UPROPERTY()
@@ -65,7 +65,7 @@ struct FGameplayCueNotifyActorArray
 	GENERATED_BODY()
 
 	UPROPERTY(transient)
-	TArray<AGameplayCueNotify_Actor*> Actors;
+	TArray<TObjectPtr<AGameplayCueNotify_Actor>> Actors;
 };
 
 /** Struct for pooling and preallocating gameplaycuenotify_actor classes. This data is per world and used to track what actors are available to recycle and which classes need to preallocate instances of those actors */
@@ -76,7 +76,7 @@ struct FPreallocationInfo
 
 	/** Raw list of pooled instances. This relies on NotifyGameplayCueActorEndPlay always being called when actor is destroyed */
 	UPROPERTY(transient)
-	TMap<UClass*, FGameplayCueNotifyActorArray> PreallocatedInstances;
+	TMap<TObjectPtr<UClass>, FGameplayCueNotifyActorArray> PreallocatedInstances;
 
 	/** List of classes that will be pooled */
 	UPROPERTY(transient)
diff --git a/Engine/Plugins/Runtime/GameplayAbilities/Source/GameplayAbilities/Public/GameplayEffect.h b/Engine/Plugins/Runtime/GameplayAbilities/Source/GameplayAbilities/Public/GameplayEffect.h
--- a/Engine/Plugins/Runtime/GameplayAbilities/Source/GameplayAbilities/Public/GameplayEffect.h
+++ b/Engine/Plugins/Runtime/GameplayAbilities/Source/GameplayAbilities/Public/GameplayEffect.h
@@ -1092,7 +1092,7 @@ public:
 
 	/** GameplayEfect definition. The static data that this spec points to. */
 	UPROPERTY()
-	const UGameplayEffect* Def;
+	TObjectPtr<const UGameplayEffect> Def;
 	
 	/** A list of attributes that were modified during the application of this spec */
 	UPROPERTY()
@@ -1190,7 +1190,7 @@ struct GAMEPLAYABILITIES_API FGameplayEffectSpecForRPC
 
 	/** GameplayEfect definition. The static data that this spec points to. */
 	UPROPERTY()
-	const UGameplayEffect* Def;
+	TObjectPtr<const UGameplayEffect> Def;
 
 	UPROPERTY()
 	TArray<FGameplayEffectModifiedAttribute> ModifiedAttributes;
@@ -1393,7 +1393,7 @@ public:
 
 	/** Matches on GameplayEffects which come from this source */
 	UPROPERTY(EditAnywhere, BlueprintReadWrite, Category = Query)
-	const UObject* EffectSource;
+	TObjectPtr<const UObject> EffectSource;
 
 	/** Matches on GameplayEffects with this definition */
 	UPROPERTY(EditAnywhere, BlueprintReadWrite, Category = Query)
@@ -1856,7 +1856,7 @@ private:
 
 	/** Active GEs that have immunity queries. This is an acceleration list to avoid searching through the Active GameplayEffect list frequetly. (We only search for the active GE if immunity procs) */
 	UPROPERTY(NotReplicated)
-	TArray<const UGameplayEffect*> ApplicationImmunityQueryEffects;
+	TArray<TObjectPtr<const UGameplayEffect>> ApplicationImmunityQueryEffects;
 
 	FAggregatorRef& FindOrCreateAttributeAggregator(FGameplayAttribute Attribute);
 
@@ -2045,7 +2045,7 @@ public:
 
 	/** Data for the UI representation of this effect. This should include things like text, icons, etc. Not available in server-only builds. */
 	UPROPERTY(EditDefaultsOnly, Instanced, BlueprintReadOnly, Category = Display)
-	class UGameplayEffectUIData* UIData;
+	TObjectPtr<class UGameplayEffectUIData> UIData;
 
 	// ----------------------------------------------------------------------
 	//	Tag Containers
diff --git a/Engine/Plugins/Runtime/GameplayAbilities/Source/GameplayAbilities/Public/GameplayEffectTypes.h b/Engine/Plugins/Runtime/GameplayAbilities/Source/GameplayAbilities/Public/GameplayEffectTypes.h
--- a/Engine/Plugins/Runtime/GameplayAbilities/Source/GameplayAbilities/Public/GameplayEffectTypes.h
+++ b/Engine/Plugins/Runtime/GameplayAbilities/Source/GameplayAbilities/Public/GameplayEffectTypes.h
@@ -1693,7 +1693,7 @@ struct GAMEPLAYABILITIES_API FMinimalReplicationTagCountMap
 	TMap<FGameplayTag, int32>	TagMap;
 
 	UPROPERTY()
-	class UAbilitySystemComponent* Owner;
+	TObjectPtr<class UAbilitySystemComponent> Owner;
 
 	/** Comparison operator */
 	bool operator==(FMinimalReplicationTagCountMap const& Other) const
diff --git a/Engine/Plugins/Runtime/GameplayAbilities/Source/GameplayAbilities/Public/GameplayPrediction.h b/Engine/Plugins/Runtime/GameplayAbilities/Source/GameplayAbilities/Public/GameplayPrediction.h
--- a/Engine/Plugins/Runtime/GameplayAbilities/Source/GameplayAbilities/Public/GameplayPrediction.h
+++ b/Engine/Plugins/Runtime/GameplayAbilities/Source/GameplayAbilities/Public/GameplayPrediction.h
@@ -302,7 +302,7 @@ struct GAMEPLAYABILITIES_API FPredictionKey
 	/** On the server, what network connection this was serialized on. */
 	UE_DEPRECATED(5.0, "No longer used in favor of GetPredictiveConnectionKey, to avoid holding a direct object reference.")
 	UPROPERTY(NotReplicated)
-	UPackageMap* PredictiveConnection = nullptr;
+	TObjectPtr<UPackageMap> PredictiveConnection = nullptr;
 
 	/** The unique ID of this prediction key */
 	UPROPERTY()
diff --git a/Engine/Plugins/Runtime/GameplayAbilities/Source/GameplayAbilities/Public/Sequencer/MovieSceneGameplayCueSections.h b/Engine/Plugins/Runtime/GameplayAbilities/Source/GameplayAbilities/Public/Sequencer/MovieSceneGameplayCueSections.h
--- a/Engine/Plugins/Runtime/GameplayAbilities/Source/GameplayAbilities/Public/Sequencer/MovieSceneGameplayCueSections.h
+++ b/Engine/Plugins/Runtime/GameplayAbilities/Source/GameplayAbilities/Public/Sequencer/MovieSceneGameplayCueSections.h
@@ -59,7 +59,7 @@ struct FMovieSceneGameplayCueKey
 
 	/** PhysMat of the hit, if there was a hit. */
 	UPROPERTY(EditAnywhere, AdvancedDisplay, Category="Gameplay Cue")
-	const UPhysicalMaterial* PhysicalMaterial;
+	TObjectPtr<const UPhysicalMaterial> PhysicalMaterial;
 
 	/** The level of that GameplayEffect */
 	UPROPERTY(EditAnywhere, AdvancedDisplay, Category="Gameplay Cue")
diff --git a/Engine/Plugins/Runtime/GameplayBehaviorSmartObjects/Source/GameplayBehaviorSmartObjectsModule/Public/AI/AITask_UseGameplayBehaviorSmartObject.h b/Engine/Plugins/Runtime/GameplayBehaviorSmartObjects/Source/GameplayBehaviorSmartObjectsModule/Public/AI/AITask_UseGameplayBehaviorSmartObject.h
--- a/Engine/Plugins/Runtime/GameplayBehaviorSmartObjects/Source/GameplayBehaviorSmartObjectsModule/Public/AI/AITask_UseGameplayBehaviorSmartObject.h
+++ b/Engine/Plugins/Runtime/GameplayBehaviorSmartObjects/Source/GameplayBehaviorSmartObjectsModule/Public/AI/AITask_UseGameplayBehaviorSmartObject.h
@@ -51,10 +51,10 @@ protected:
 	FGenericGameplayTaskDelegate OnFinished;
 
 	UPROPERTY()
-	UAITask_MoveTo* MoveToTask;
+	TObjectPtr<UAITask_MoveTo> MoveToTask;
 
 	UPROPERTY()
-	UGameplayBehavior* GameplayBehavior;
+	TObjectPtr<UGameplayBehavior> GameplayBehavior;
 
 	FSmartObjectClaimHandle ClaimedHandle;
 	FDelegateHandle OnBehaviorFinishedNotifyHandle;
diff --git a/Engine/Plugins/Runtime/GameplayBehaviorSmartObjects/Source/GameplayBehaviorSmartObjectsModule/Public/GameplayBehaviorSmartObjectBehaviorDefinition.h b/Engine/Plugins/Runtime/GameplayBehaviorSmartObjects/Source/GameplayBehaviorSmartObjectsModule/Public/GameplayBehaviorSmartObjectBehaviorDefinition.h
--- a/Engine/Plugins/Runtime/GameplayBehaviorSmartObjects/Source/GameplayBehaviorSmartObjectsModule/Public/GameplayBehaviorSmartObjectBehaviorDefinition.h
+++ b/Engine/Plugins/Runtime/GameplayBehaviorSmartObjects/Source/GameplayBehaviorSmartObjectsModule/Public/GameplayBehaviorSmartObjectBehaviorDefinition.h
@@ -17,5 +17,5 @@ class GAMEPLAYBEHAVIORSMARTOBJECTSMODULE_API UGameplayBehaviorSmartObjectBehavio
 
 public:
 	UPROPERTY(EditDefaultsOnly, Category = SmartObject, Instanced)
-	UGameplayBehaviorConfig* GameplayBehaviorConfig;
+	TObjectPtr<UGameplayBehaviorConfig> GameplayBehaviorConfig;
 };
\ No newline at end of file
diff --git a/Engine/Plugins/Runtime/GameplayStateTree/Source/GameplayStateTreeModule/Public/Components/StateTreeComponent.h b/Engine/Plugins/Runtime/GameplayStateTree/Source/GameplayStateTreeModule/Public/Components/StateTreeComponent.h
--- a/Engine/Plugins/Runtime/GameplayStateTree/Source/GameplayStateTreeModule/Public/Components/StateTreeComponent.h
+++ b/Engine/Plugins/Runtime/GameplayStateTree/Source/GameplayStateTreeModule/Public/Components/StateTreeComponent.h
@@ -59,7 +59,7 @@ protected:
 #if WITH_EDITORONLY_DATA
 	UE_DEPRECATED(5.1, "This property has been deprecated. Use StateTreeReference instead.")
 	UPROPERTY()
-	UStateTree* StateTree_DEPRECATED;
+	TObjectPtr<UStateTree> StateTree_DEPRECATED;
 #endif
 
 	UPROPERTY(EditAnywhere, Category = AI, meta=(Schema="/Script/GameplayStateTreeModule.StateTreeComponentSchema"))
diff --git a/Engine/Plugins/Runtime/GeoReferencing/Source/GeoReferencing/Public/RoundPlanetPawn.h b/Engine/Plugins/Runtime/GeoReferencing/Source/GeoReferencing/Public/RoundPlanetPawn.h
--- a/Engine/Plugins/Runtime/GeoReferencing/Source/GeoReferencing/Public/RoundPlanetPawn.h
+++ b/Engine/Plugins/Runtime/GeoReferencing/Source/GeoReferencing/Public/RoundPlanetPawn.h
@@ -34,7 +34,7 @@ public:
 	* This curve must be kept in the 0 to 1 range on both axes.
 	*/
 	UPROPERTY(EditAnywhere, Category = "FlyToLocation")
-	UCurveFloat* AltitudeProfileCurve;
+	TObjectPtr<UCurveFloat> AltitudeProfileCurve;
 
 	/**
 	* This curve is used to pick up a reference the maximum altitude when flying to a location. 
@@ -43,7 +43,7 @@ public:
 	* Y Axis is the maximum altitude to be used for this flight (meters)
 	*/
 	UPROPERTY(EditAnywhere, Category = "FlyToLocation")
-	UCurveFloat* MaximumAltitudeCurve;
+	TObjectPtr<UCurveFloat> MaximumAltitudeCurve;
 
 	/**
 	* In case MaximumAltitudeCurve is not defined, the AltitudeProfileCurve will use this Maximum altitude value for the flight, whatever the travel distance. 
@@ -60,7 +60,7 @@ public:
 	* A linear curve will give a constant progress speed whereas a easein//easeout one will slow down the speed at beginning/end of motion. 
 	*/
 	UPROPERTY(EditAnywhere, Category = "FlyToLocation")
-	UCurveFloat* ProgressCurve;
+	TObjectPtr<UCurveFloat> ProgressCurve;
 
 	/**
 	* Fly to Location duration (in seconds)
@@ -124,7 +124,7 @@ public:
 	* ActualMaxSpeed = BaseSpeedKmh * SpeedScalar * AltitudeSpeedModifierCurve(Altitude)
 	*/
 	UPROPERTY(EditAnywhere, Category = "Movement")
-	UCurveFloat* SpeedByHATModifierCurve;
+	TObjectPtr<UCurveFloat> SpeedByHATModifierCurve;
 
 	/**
 	* Height Above Terrain. The distance between the ground and the pawn
@@ -231,7 +231,7 @@ public:
 
 protected:
 	UPROPERTY()
-	AGeoReferencingSystem* GeoReferencingSystem;
+	TObjectPtr<AGeoReferencingSystem> GeoReferencingSystem;
 
 private:
 	/**
diff --git a/Engine/Plugins/Runtime/HairStrands/Source/HairStrandsCore/Public/GroomBindingAsset.h b/Engine/Plugins/Runtime/HairStrands/Source/HairStrandsCore/Public/GroomBindingAsset.h
--- a/Engine/Plugins/Runtime/HairStrands/Source/HairStrandsCore/Public/GroomBindingAsset.h
+++ b/Engine/Plugins/Runtime/HairStrands/Source/HairStrandsCore/Public/GroomBindingAsset.h
@@ -198,5 +198,5 @@ class HAIRSTRANDSCORE_API UGroomBindingAssetList : public UObject
 	GENERATED_BODY()
 public:
 	UPROPERTY(EditAnywhere, Transient, EditFixedSize, Category = "Bindings")
-	TArray<UGroomBindingAsset*> Bindings;
+	TArray<TObjectPtr<UGroomBindingAsset>> Bindings;
 };
\ No newline at end of file
diff --git a/Engine/Plugins/Runtime/MassEntity/Source/MassEntity/Public/MassEntitySettings.h b/Engine/Plugins/Runtime/MassEntity/Source/MassEntity/Public/MassEntitySettings.h
--- a/Engine/Plugins/Runtime/MassEntity/Source/MassEntity/Public/MassEntitySettings.h
+++ b/Engine/Plugins/Runtime/MassEntity/Source/MassEntity/Public/MassEntitySettings.h
@@ -28,13 +28,13 @@ struct FMassProcessingPhaseConfig
 	TSubclassOf<UMassCompositeProcessor> PhaseGroupClass = UMassCompositeProcessor::StaticClass();
 
 	UPROPERTY(Transient)
-	TArray<UMassProcessor*> ProcessorCDOs;
+	TArray<TObjectPtr<UMassProcessor>> ProcessorCDOs;
 
 #if WITH_EDITORONLY_DATA
 	// this processor is available only in editor since it's used to present the user the order in which processors
 	// will be executed when given processing phase gets triggered
 	UPROPERTY(Transient)
-	UMassCompositeProcessor* PhaseProcessor = nullptr;
+	TObjectPtr<UMassCompositeProcessor> PhaseProcessor = nullptr;
 
 	UPROPERTY(VisibleAnywhere, Category = Mass, Transient)
 	FText Description;
@@ -89,7 +89,7 @@ public:
 
 	/** This list contains all the processors available in the given binary (including plugins). The contents are sorted by display name.*/
 	UPROPERTY(VisibleAnywhere, Category = Mass, Transient, Instanced, EditFixedSize)
-	TArray<UMassProcessor*> ProcessorCDOs;
+	TArray<TObjectPtr<UMassProcessor>> ProcessorCDOs;
 
 #if WITH_EDITORONLY_DATA
 protected:
diff --git a/Engine/Plugins/Runtime/MassEntity/Source/MassEntity/Public/MassObserverManager.h b/Engine/Plugins/Runtime/MassEntity/Source/MassEntity/Public/MassObserverManager.h
--- a/Engine/Plugins/Runtime/MassEntity/Source/MassEntity/Public/MassObserverManager.h
+++ b/Engine/Plugins/Runtime/MassEntity/Source/MassEntity/Public/MassObserverManager.h
@@ -27,7 +27,7 @@ struct FMassObserversMap
 
 private:
 	UPROPERTY()
-	TMap<const UScriptStruct*, FMassRuntimePipeline> Container;
+	TMap<TObjectPtr<const UScriptStruct>, FMassRuntimePipeline> Container;
 };
 
 /** 
diff --git a/Engine/Plugins/Runtime/MassEntity/Source/MassEntity/Public/MassObserverProcessor.h b/Engine/Plugins/Runtime/MassEntity/Source/MassEntity/Public/MassObserverProcessor.h
--- a/Engine/Plugins/Runtime/MassEntity/Source/MassEntity/Public/MassObserverProcessor.h
+++ b/Engine/Plugins/Runtime/MassEntity/Source/MassEntity/Public/MassObserverProcessor.h
@@ -26,7 +26,7 @@ protected:
 protected:
 	/** Determines which Fragment or Tag type this given UMassObserverProcessor will be observing */
 	UPROPERTY()
-	UScriptStruct* ObservedType = nullptr;
+	TObjectPtr<UScriptStruct> ObservedType = nullptr;
 
 	EMassObservedOperation Operation = EMassObservedOperation::MAX;
 };
diff --git a/Engine/Plugins/Runtime/MassEntity/Source/MassEntity/Public/MassObserverRegistry.h b/Engine/Plugins/Runtime/MassEntity/Source/MassEntity/Public/MassObserverRegistry.h
--- a/Engine/Plugins/Runtime/MassEntity/Source/MassEntity/Public/MassObserverRegistry.h
+++ b/Engine/Plugins/Runtime/MassEntity/Source/MassEntity/Public/MassObserverRegistry.h
@@ -42,7 +42,7 @@ struct FMassEntityObserverClassesMap
 
 private:
 	UPROPERTY()
-	TMap<const UScriptStruct*, FMassProcessorClassCollection> Container;
+	TMap<TObjectPtr<const UScriptStruct>, FMassProcessorClassCollection> Container;
 };
 
 UCLASS()
diff --git a/Engine/Plugins/Runtime/MassEntity/Source/MassEntity/Public/MassProcessingPhase.h b/Engine/Plugins/Runtime/MassEntity/Source/MassEntity/Public/MassProcessingPhase.h
--- a/Engine/Plugins/Runtime/MassEntity/Source/MassEntity/Public/MassProcessingPhase.h
+++ b/Engine/Plugins/Runtime/MassEntity/Source/MassEntity/Public/MassProcessingPhase.h
@@ -45,10 +45,10 @@ protected:
 	friend UMassProcessingPhaseManager;
 
 	UPROPERTY(EditAnywhere, Category=Mass)
-	UMassCompositeProcessor* PhaseProcessor = nullptr;
+	TObjectPtr<UMassCompositeProcessor> PhaseProcessor = nullptr;
 
 	UPROPERTY()
-	UMassProcessingPhaseManager* PhaseManager = nullptr;
+	TObjectPtr<UMassProcessingPhaseManager> PhaseManager = nullptr;
 
 	EMassProcessingPhase Phase = EMassProcessingPhase::MAX;
 	FOnPhaseEvent OnPhaseStart;
diff --git a/Engine/Plugins/Runtime/MassEntity/Source/MassEntity/Public/MassProcessingTypes.h b/Engine/Plugins/Runtime/MassEntity/Source/MassEntity/Public/MassProcessingTypes.h
--- a/Engine/Plugins/Runtime/MassEntity/Source/MassEntity/Public/MassProcessingTypes.h
+++ b/Engine/Plugins/Runtime/MassEntity/Source/MassEntity/Public/MassProcessingTypes.h
@@ -71,7 +71,7 @@ struct MASSENTITY_API FMassRuntimePipeline
 	GENERATED_BODY()
 
 	UPROPERTY()
-	TArray<UMassProcessor*> Processors;
+	TArray<TObjectPtr<UMassProcessor>> Processors;
 
 	void Reset();
 	void Initialize(UObject& Owner);
diff --git a/Engine/Plugins/Runtime/MassEntity/Source/MassEntity/Public/MassSettings.h b/Engine/Plugins/Runtime/MassEntity/Source/MassEntity/Public/MassSettings.h
--- a/Engine/Plugins/Runtime/MassEntity/Source/MassEntity/Public/MassSettings.h
+++ b/Engine/Plugins/Runtime/MassEntity/Source/MassEntity/Public/MassSettings.h
@@ -29,5 +29,5 @@ public:
 
 public:
 	UPROPERTY(VisibleAnywhere, Category = "Mass", NoClear, EditFixedSize, meta = (EditInline))
-	TMap<FName, UMassModuleSettings*> ModuleSettings;
+	TMap<FName, TObjectPtr<UMassModuleSettings>> ModuleSettings;
 };
diff --git a/Engine/Plugins/Runtime/MassEntity/Source/MassEntityTestSuite/Public/LWComponentTestFarmPlot.h b/Engine/Plugins/Runtime/MassEntity/Source/MassEntityTestSuite/Public/LWComponentTestFarmPlot.h
--- a/Engine/Plugins/Runtime/MassEntity/Source/MassEntityTestSuite/Public/LWComponentTestFarmPlot.h
+++ b/Engine/Plugins/Runtime/MassEntity/Source/MassEntityTestSuite/Public/LWComponentTestFarmPlot.h
@@ -22,10 +22,10 @@ struct FFarmVisualDataRow : public FTableRowBase
 	GENERATED_BODY()
 
 	UPROPERTY(EditAnywhere, Category=Farm)
-	UStaticMesh* Mesh = nullptr;
+	TObjectPtr<UStaticMesh> Mesh = nullptr;
 
 	UPROPERTY(EditAnywhere, Category=Farm)
-	UMaterialInterface* MaterialOverride = nullptr;
+	TObjectPtr<UMaterialInterface> MaterialOverride = nullptr;
 };
 
 USTRUCT()
@@ -300,15 +300,15 @@ public:
 	TArray<FFarmVisualDataRow> VisualDataTable;
 
 	UPROPERTY(Transient)
-	TArray<UHierarchicalInstancedStaticMeshComponent*> VisualDataISMCs;
+	TArray<TObjectPtr<UHierarchicalInstancedStaticMeshComponent>> VisualDataISMCs;
 
 	float NextSecondTimer = 0.0f;
 
 	UPROPERTY(Transient)
-	TArray<UFragmentUpdateSystem*> PerFrameSystems;
+	TArray<TObjectPtr<UFragmentUpdateSystem>> PerFrameSystems;
 
 	UPROPERTY(Transient)
-	TArray<UFragmentUpdateSystem*> PerSecondSystems;
+	TArray<TObjectPtr<UFragmentUpdateSystem>> PerSecondSystems;
 
 	// Indicies into VisualDataTable for flowers
 	UPROPERTY(EditAnywhere, Category=Farm)
@@ -332,7 +332,7 @@ public:
 
 private:
 	UPROPERTY(VisibleAnywhere, BlueprintReadOnly, meta=(AllowPrivateAccess = "true"), Category=Farm)
-	UHierarchicalInstancedStaticMeshComponent* HarvestIconISMC;
+	TObjectPtr<UHierarchicalInstancedStaticMeshComponent> HarvestIconISMC;
 
 	TArray<int32> FreeHarvestIconIndicies;
 
diff --git a/Engine/Plugins/Runtime/MassGameplay/Source/MassActors/Public/MassActorSpawnerSubsystem.h b/Engine/Plugins/Runtime/MassGameplay/Source/MassActors/Public/MassActorSpawnerSubsystem.h
--- a/Engine/Plugins/Runtime/MassGameplay/Source/MassActors/Public/MassActorSpawnerSubsystem.h
+++ b/Engine/Plugins/Runtime/MassGameplay/Source/MassActors/Public/MassActorSpawnerSubsystem.h
@@ -81,7 +81,7 @@ public:
 
 	/** The pointer to the actor once it is spawned */
 	UPROPERTY(Transient)
-	AActor* SpawnedActor = nullptr;
+	TObjectPtr<AActor> SpawnedActor = nullptr;
 
 	/** Internal request serial number (used to cycle through next spawning request) */
 	uint32 SerialNumber = 0;
@@ -223,10 +223,10 @@ protected:
 	TArray<FInstancedStruct> SpawnRequests;
 
 	UPROPERTY()
-	TArray<AActor*> ActorsToDestroy;
+	TArray<TObjectPtr<AActor>> ActorsToDestroy;
 
 	UPROPERTY()
-	TArray<AActor*> DeactivatedActorsToDestroy;
+	TArray<TObjectPtr<AActor>> DeactivatedActorsToDestroy;
 
 	bool bActorPoolingEnabled = true;
 
diff --git a/Engine/Plugins/Runtime/MassGameplay/Source/MassActors/Public/MassAgentSubsystem.h b/Engine/Plugins/Runtime/MassGameplay/Source/MassActors/Public/MassAgentSubsystem.h
--- a/Engine/Plugins/Runtime/MassGameplay/Source/MassActors/Public/MassAgentSubsystem.h
+++ b/Engine/Plugins/Runtime/MassGameplay/Source/MassActors/Public/MassAgentSubsystem.h
@@ -29,7 +29,7 @@ struct MASSACTORS_API FMassAgentInitializationQueue
 	GENERATED_BODY()
 
 	UPROPERTY()
-	TArray<UMassAgentComponent*> AgentComponents;
+	TArray<TObjectPtr<UMassAgentComponent>> AgentComponents;
 };
 
 /**
@@ -124,10 +124,10 @@ protected:
 	TSharedPtr<FMassEntityManager> EntityManager;
 
 	UPROPERTY()
-	UMassSpawnerSubsystem* SpawnerSystem;
+	TObjectPtr<UMassSpawnerSubsystem> SpawnerSystem;
 
 	UPROPERTY()
-	UMassSimulationSubsystem* SimulationSystem;
+	TObjectPtr<UMassSimulationSubsystem> SimulationSystem;
 
 	UPROPERTY()
 	TMap<FMassEntityTemplateID, FMassAgentInitializationQueue> PendingAgentEntities;
@@ -136,10 +136,10 @@ protected:
 	TMap<FMassEntityTemplateID, FMassAgentInitializationQueue> PendingPuppets;
 
 	UPROPERTY()
-	UMassReplicationSubsystem* ReplicationSubsystem;
+	TObjectPtr<UMassReplicationSubsystem> ReplicationSubsystem;
 
 	UPROPERTY()
-	TMap<FMassNetworkID, UMassAgentComponent*> ReplicatedAgentComponents;
+	TMap<FMassNetworkID, TObjectPtr<UMassAgentComponent>> ReplicatedAgentComponents;
 
 	UE::MassActor::FMassAgentComponentDelegate OnMassAgentComponentEntityAssociated;
 	UE::MassActor::FMassAgentComponentDelegate OnMassAgentComponentEntityDetaching;
diff --git a/Engine/Plugins/Runtime/MassGameplay/Source/MassGameplayDebug/Public/MassDebugVisualizationComponent.h b/Engine/Plugins/Runtime/MassGameplay/Source/MassGameplayDebug/Public/MassDebugVisualizationComponent.h
--- a/Engine/Plugins/Runtime/MassGameplay/Source/MassGameplayDebug/Public/MassDebugVisualizationComponent.h
+++ b/Engine/Plugins/Runtime/MassGameplay/Source/MassGameplayDebug/Public/MassDebugVisualizationComponent.h
@@ -35,6 +35,6 @@ protected:
 	TArray<FAgentDebugVisualization> VisualDataTable;
 
 	UPROPERTY(Transient)
-	TArray<UHierarchicalInstancedStaticMeshComponent*> VisualDataISMCs;
+	TArray<TObjectPtr<UHierarchicalInstancedStaticMeshComponent>> VisualDataISMCs;
 #endif // WITH_EDITORONLY_DATA
 };
diff --git a/Engine/Plugins/Runtime/MassGameplay/Source/MassGameplayDebug/Public/MassDebugVisualizer.h b/Engine/Plugins/Runtime/MassGameplay/Source/MassGameplayDebug/Public/MassDebugVisualizer.h
--- a/Engine/Plugins/Runtime/MassGameplay/Source/MassGameplayDebug/Public/MassDebugVisualizer.h
+++ b/Engine/Plugins/Runtime/MassGameplay/Source/MassGameplayDebug/Public/MassDebugVisualizer.h
@@ -19,6 +19,6 @@ public:
 
 protected:
 	UPROPERTY()
-	class UMassDebugVisualizationComponent* DebugVisComponent;
+	TObjectPtr<class UMassDebugVisualizationComponent> DebugVisComponent;
 #endif // WITH_EDITORONLY_DATA
 };
diff --git a/Engine/Plugins/Runtime/MassGameplay/Source/MassGameplayDebug/Public/MassDebuggerSubsystem.h b/Engine/Plugins/Runtime/MassGameplay/Source/MassGameplayDebug/Public/MassDebuggerSubsystem.h
--- a/Engine/Plugins/Runtime/MassGameplay/Source/MassGameplayDebug/Public/MassDebuggerSubsystem.h
+++ b/Engine/Plugins/Runtime/MassGameplay/Source/MassGameplayDebug/Public/MassDebuggerSubsystem.h
@@ -63,10 +63,10 @@ protected:
 	FString SelectedEntityDetails;
 
 	UPROPERTY(Transient)
-	UMassDebugVisualizationComponent* VisualizationComponent;
+	TObjectPtr<UMassDebugVisualizationComponent> VisualizationComponent;
 
 	UPROPERTY(Transient)
-	AMassDebugVisualizer* DebugVisualizer;
+	TObjectPtr<AMassDebugVisualizer> DebugVisualizer;
 
 	FDelegateHandle OnEntitySelectedHandle;
 };
diff --git a/Engine/Plugins/Runtime/MassGameplay/Source/MassGameplayDebug/Public/MassGameplayDebugTypes.h b/Engine/Plugins/Runtime/MassGameplay/Source/MassGameplayDebug/Public/MassGameplayDebugTypes.h
--- a/Engine/Plugins/Runtime/MassGameplay/Source/MassGameplayDebug/Public/MassGameplayDebugTypes.h
+++ b/Engine/Plugins/Runtime/MassGameplay/Source/MassGameplayDebug/Public/MassGameplayDebugTypes.h
@@ -37,10 +37,10 @@ struct FSimDebugDataRow : public FTableRowBase
 	GENERATED_BODY()
 
 	UPROPERTY(EditAnywhere, Category = Debug)
-	UStaticMesh* Mesh = nullptr;
+	TObjectPtr<UStaticMesh> Mesh = nullptr;
 
 	UPROPERTY(EditAnywhere, Category = Debug)
-	UMaterialInterface* MaterialOverride = nullptr;
+	TObjectPtr<UMaterialInterface> MaterialOverride = nullptr;
 
 	UPROPERTY(EditAnywhere, Category = Debug)
 	float Scale = 1.f;
@@ -84,10 +84,10 @@ struct FAgentDebugVisualization : public FTableRowBase
 	GENERATED_BODY()
 
 	UPROPERTY(EditAnywhere, Category = "Mass|Debug")
-	UStaticMesh* Mesh = nullptr;
+	TObjectPtr<UStaticMesh> Mesh = nullptr;
 
 	UPROPERTY(EditAnywhere, Category = "Mass|Debug")
-	UMaterialInterface* MaterialOverride = nullptr;
+	TObjectPtr<UMaterialInterface> MaterialOverride = nullptr;
 
 	/** Near cull distance to override default value for that agent type */
 	UPROPERTY(EditAnywhere, Category = "Mass|Debug")
diff --git a/Engine/Plugins/Runtime/MassGameplay/Source/MassLOD/Public/MassLODSubsystem.h b/Engine/Plugins/Runtime/MassGameplay/Source/MassLOD/Public/MassLODSubsystem.h
--- a/Engine/Plugins/Runtime/MassGameplay/Source/MassLOD/Public/MassLODSubsystem.h
+++ b/Engine/Plugins/Runtime/MassGameplay/Source/MassLOD/Public/MassLODSubsystem.h
@@ -29,7 +29,7 @@ struct FViewerInfo
 	GENERATED_BODY()
 
 	UPROPERTY(transient)
-	APlayerController* PlayerController = nullptr;
+	TObjectPtr<APlayerController> PlayerController = nullptr;
 	
 	FName StreamingSourceName;
 
diff --git a/Engine/Plugins/Runtime/MassGameplay/Source/MassReplication/Public/MassClientBubbleSerializerBase.h b/Engine/Plugins/Runtime/MassGameplay/Source/MassReplication/Public/MassClientBubbleSerializerBase.h
--- a/Engine/Plugins/Runtime/MassGameplay/Source/MassReplication/Public/MassClientBubbleSerializerBase.h
+++ b/Engine/Plugins/Runtime/MassGameplay/Source/MassReplication/Public/MassClientBubbleSerializerBase.h
@@ -42,15 +42,15 @@ public:
 
 private:
 	UPROPERTY(Transient)
-	UWorld* World = nullptr;
+	TObjectPtr<UWorld> World = nullptr;
 
 	UPROPERTY(Transient)
-	UMassSpawnerSubsystem* SpawnerSubsystem = nullptr;
+	TObjectPtr<UMassSpawnerSubsystem> SpawnerSubsystem = nullptr;
 
 	TSharedPtr<FMassEntityManager> EntityManager;
 
 	UPROPERTY(Transient)
-	UMassReplicationSubsystem* ReplicationSubsystem = nullptr;
+	TObjectPtr<UMassReplicationSubsystem> ReplicationSubsystem = nullptr;
 
 	/** Pointer to the IClientBubbleHandlerInterface derived class in the class derived from this one */
 	IClientBubbleHandlerInterface* ClientHandler = nullptr;
diff --git a/Engine/Plugins/Runtime/MassGameplay/Source/MassReplication/Public/MassReplicationFragments.h b/Engine/Plugins/Runtime/MassGameplay/Source/MassReplication/Public/MassReplicationFragments.h
--- a/Engine/Plugins/Runtime/MassGameplay/Source/MassReplication/Public/MassReplicationFragments.h
+++ b/Engine/Plugins/Runtime/MassGameplay/Source/MassReplication/Public/MassReplicationFragments.h
@@ -126,7 +126,7 @@ public:
 	//TODO review if we need to have this as a UPROPERTY at all and also if we can make this use a TInlineAllocator
 	//Can not use TInlineAllocator with UPROPERTY()
 	UPROPERTY(Transient)
-	TArray<AMassClientBubbleInfoBase*> BubbleInfos;
+	TArray<TObjectPtr<AMassClientBubbleInfoBase>> BubbleInfos;
 
 	TMassLODCollector<FReplicationLODLogic> LODCollector;
 	TMassLODCalculator<FReplicationLODLogic> LODCalculator;
@@ -136,7 +136,7 @@ public:
 	FMassEntityQuery EntityQuery;
 
 	UPROPERTY(Transient)
-	mutable UMassReplicatorBase* CachedReplicator = nullptr;
+	mutable TObjectPtr<UMassReplicatorBase> CachedReplicator = nullptr;
 
 	template<typename T>
 	T& GetTypedClientBubbleInfoChecked(FMassClientHandle Handle)
diff --git a/Engine/Plugins/Runtime/MassGameplay/Source/MassReplication/Public/MassReplicationProcessor.h b/Engine/Plugins/Runtime/MassGameplay/Source/MassReplication/Public/MassReplicationProcessor.h
--- a/Engine/Plugins/Runtime/MassGameplay/Source/MassReplication/Public/MassReplicationProcessor.h
+++ b/Engine/Plugins/Runtime/MassGameplay/Source/MassReplication/Public/MassReplicationProcessor.h
@@ -39,7 +39,7 @@ protected:
 protected:
 
 	UPROPERTY(Transient)
-	UMassReplicationSubsystem* ReplicationSubsystem = nullptr;
+	TObjectPtr<UMassReplicationSubsystem> ReplicationSubsystem = nullptr;
 
 	FMassEntityQuery SyncClientData;
 	FMassEntityQuery CollectViewerInfoQuery;
diff --git a/Engine/Plugins/Runtime/MassGameplay/Source/MassReplication/Public/MassReplicationSubsystem.h b/Engine/Plugins/Runtime/MassGameplay/Source/MassReplication/Public/MassReplicationSubsystem.h
--- a/Engine/Plugins/Runtime/MassGameplay/Source/MassReplication/Public/MassReplicationSubsystem.h
+++ b/Engine/Plugins/Runtime/MassGameplay/Source/MassReplication/Public/MassReplicationSubsystem.h
@@ -73,7 +73,7 @@ struct FMassClientBubbleInfoData
 
 	/** A free list array of AMassClientBubbleInfos. This is organised so the index is that of the client FMassClientHandle */
 	UPROPERTY(Transient)
-	TArray<AMassClientBubbleInfoBase*> Bubbles;
+	TArray<TObjectPtr<AMassClientBubbleInfoBase>> Bubbles;
 
 	UPROPERTY(Transient)
 	TSubclassOf<AMassClientBubbleInfoBase> BubbleClass;
@@ -313,10 +313,10 @@ protected:
 	TArray<FViewerClientPair> ClientToViewerHandleArray;
 
 	UPROPERTY()
-	UWorld* World;
+	TObjectPtr<UWorld> World;
 
 	UPROPERTY()
-	UMassLODSubsystem* MassLODSubsystem;
+	TObjectPtr<UMassLODSubsystem> MassLODSubsystem;
 
 #if UE_REPLICATION_COMPILE_CLIENT_CODE
 	TMap<FMassNetworkID, FMassReplicationEntityInfo> EntityInfoMap;
diff --git a/Engine/Plugins/Runtime/MassGameplay/Source/MassRepresentation/Public/MassRepresentationFragments.h b/Engine/Plugins/Runtime/MassGameplay/Source/MassRepresentation/Public/MassRepresentationFragments.h
--- a/Engine/Plugins/Runtime/MassGameplay/Source/MassRepresentation/Public/MassRepresentationFragments.h
+++ b/Engine/Plugins/Runtime/MassGameplay/Source/MassRepresentation/Public/MassRepresentationFragments.h
@@ -61,7 +61,7 @@ struct FMassRepresentationSubsystemSharedFragment : public FMassSharedFragment
 	GENERATED_BODY()
 
 	UPROPERTY(Transient)
-	UMassRepresentationSubsystem* RepresentationSubsystem = nullptr;
+	TObjectPtr<UMassRepresentationSubsystem> RepresentationSubsystem = nullptr;
 };
 
 template<>
@@ -116,7 +116,7 @@ struct FMassRepresentationParameters : public FMassSharedFragment
 	mutable EMassRepresentationType CachedDefaultRepresentationType = EMassRepresentationType::None;
 
 	UPROPERTY(Transient)
-	mutable UMassRepresentationActorManagement* CachedRepresentationActorManagement = nullptr;
+	mutable TObjectPtr<UMassRepresentationActorManagement> CachedRepresentationActorManagement = nullptr;
 };
 
 template<>
@@ -179,7 +179,7 @@ struct FMassVisualizationLODParameters : public FMassSharedFragment
 
 	/** Filter these settings with specified tag */
 	UPROPERTY(EditAnywhere, Category = "Mass|LOD", meta = (BaseStruct = "MassTag"))
-	UScriptStruct* FilterTag = nullptr;
+	TObjectPtr<UScriptStruct> FilterTag = nullptr;
 };
 
 USTRUCT()
@@ -194,5 +194,5 @@ struct FMassVisualizationLODSharedFragment : public FMassSharedFragment
 	bool bHasAdjustedDistancesFromCount = false;
 
 	UPROPERTY(Transient)
-	const UScriptStruct* FilterTag = nullptr;
+	TObjectPtr<const UScriptStruct> FilterTag = nullptr;
 };
diff --git a/Engine/Plugins/Runtime/MassGameplay/Source/MassRepresentation/Public/MassRepresentationSubsystem.h b/Engine/Plugins/Runtime/MassGameplay/Source/MassRepresentation/Public/MassRepresentationSubsystem.h
--- a/Engine/Plugins/Runtime/MassGameplay/Source/MassRepresentation/Public/MassRepresentationSubsystem.h
+++ b/Engine/Plugins/Runtime/MassGameplay/Source/MassRepresentation/Public/MassRepresentationSubsystem.h
@@ -140,19 +140,19 @@ protected:
 
 	/** The component that handles all the static mesh instances */
 	UPROPERTY(Transient)
-	UMassVisualizationComponent* VisualizationComponent;
+	TObjectPtr<UMassVisualizationComponent> VisualizationComponent;
 
 	/** The actor owning the above visualization component */
 	UPROPERTY(Transient)
-	AMassVisualizer* Visualizer;
+	TObjectPtr<AMassVisualizer> Visualizer;
 
 	UPROPERTY(Transient)
-	UMassActorSpawnerSubsystem* ActorSpawnerSubsystem;
+	TObjectPtr<UMassActorSpawnerSubsystem> ActorSpawnerSubsystem;
 
 	TSharedPtr<FMassEntityManager> EntityManager;
 
 	UPROPERTY(Transient)
-	UWorldPartitionSubsystem* WorldPartitionSubsystem;
+	TObjectPtr<UWorldPartitionSubsystem> WorldPartitionSubsystem;
 
 	/** The time to wait before retrying a to spawn actor that failed */
 	float RetryMovedDistanceSq = 1000000.0f;
diff --git a/Engine/Plugins/Runtime/MassGameplay/Source/MassRepresentation/Public/MassRepresentationTypes.h b/Engine/Plugins/Runtime/MassGameplay/Source/MassRepresentation/Public/MassRepresentationTypes.h
--- a/Engine/Plugins/Runtime/MassGameplay/Source/MassRepresentation/Public/MassRepresentationTypes.h
+++ b/Engine/Plugins/Runtime/MassGameplay/Source/MassRepresentation/Public/MassRepresentationTypes.h
@@ -41,7 +41,7 @@ struct FStaticMeshInstanceVisualizationMeshDesc
 	
 	/** The static mesh visual representation */
 	UPROPERTY(EditAnywhere, Category = "Mass|Visual")
-	UStaticMesh* Mesh = nullptr;
+	TObjectPtr<UStaticMesh> Mesh = nullptr;
 
 	/**
 	 * Material overrides for the static mesh visual representation. 
@@ -49,7 +49,7 @@ struct FStaticMeshInstanceVisualizationMeshDesc
 	 * Array indices correspond to material slot indices on the static mesh.
 	 */
 	UPROPERTY(EditAnywhere, Category = "Mass|Visual")
-	TArray<UMaterialInterface*> MaterialOverrides;
+	TArray<TObjectPtr<UMaterialInterface>> MaterialOverrides;
 
 	/** The minimum inclusive LOD significance to start using this static mesh */
 	UPROPERTY(EditAnywhere, Category = "Mass|Visual")
@@ -322,7 +322,7 @@ protected:
 
 	/** The component handling these instances */
 	UPROPERTY()
-		TArray<UInstancedStaticMeshComponent*> InstancedStaticMeshComponents;
+		TArray<TObjectPtr<UInstancedStaticMeshComponent>> InstancedStaticMeshComponents;
 
 	UPROPERTY()
 		TArray<FMassLODSignificanceRange> LODSignificanceRanges;
diff --git a/Engine/Plugins/Runtime/MassGameplay/Source/MassRepresentation/Public/MassVisualizationLODProcessor.h b/Engine/Plugins/Runtime/MassGameplay/Source/MassRepresentation/Public/MassVisualizationLODProcessor.h
--- a/Engine/Plugins/Runtime/MassGameplay/Source/MassRepresentation/Public/MassVisualizationLODProcessor.h
+++ b/Engine/Plugins/Runtime/MassGameplay/Source/MassRepresentation/Public/MassVisualizationLODProcessor.h
@@ -44,5 +44,5 @@ protected:
 	bool bForceOFFLOD = false;
 
 	UPROPERTY(Transient)
-	const UScriptStruct* FilterTag = nullptr;
+	TObjectPtr<const UScriptStruct> FilterTag = nullptr;
 };
\ No newline at end of file
diff --git a/Engine/Plugins/Runtime/MassGameplay/Source/MassRepresentation/Public/MassVisualizer.h b/Engine/Plugins/Runtime/MassGameplay/Source/MassRepresentation/Public/MassVisualizer.h
--- a/Engine/Plugins/Runtime/MassGameplay/Source/MassRepresentation/Public/MassVisualizer.h
+++ b/Engine/Plugins/Runtime/MassGameplay/Source/MassRepresentation/Public/MassVisualizer.h
@@ -25,5 +25,5 @@ public:
 
 protected:
 	UPROPERTY()
-	class UMassVisualizationComponent* VisComponent;
+	TObjectPtr<class UMassVisualizationComponent> VisComponent;
 };
diff --git a/Engine/Plugins/Runtime/MassGameplay/Source/MassSimulation/Public/MassSimulationSubsystem.h b/Engine/Plugins/Runtime/MassGameplay/Source/MassSimulation/Public/MassSimulationSubsystem.h
--- a/Engine/Plugins/Runtime/MassGameplay/Source/MassSimulation/Public/MassSimulationSubsystem.h
+++ b/Engine/Plugins/Runtime/MassGameplay/Source/MassSimulation/Public/MassSimulationSubsystem.h
@@ -55,7 +55,7 @@ protected:
 	TSharedPtr<FMassEntityManager> EntityManager;
 
 	UPROPERTY()
-	UMassProcessingPhaseManager* PhaseManager;
+	TObjectPtr<UMassProcessingPhaseManager> PhaseManager;
 
 	inline static FOnSimulationStarted OnSimulationStarted={};
 
diff --git a/Engine/Plugins/Runtime/MassGameplay/Source/MassSmartObjects/Public/SmartObjectZoneAnnotations.h b/Engine/Plugins/Runtime/MassGameplay/Source/MassSmartObjects/Public/SmartObjectZoneAnnotations.h
--- a/Engine/Plugins/Runtime/MassGameplay/Source/MassSmartObjects/Public/SmartObjectZoneAnnotations.h
+++ b/Engine/Plugins/Runtime/MassGameplay/Source/MassSmartObjects/Public/SmartObjectZoneAnnotations.h
@@ -141,5 +141,5 @@ protected:
 
 	/** Cached SmartObjectSubsystem */
 	UPROPERTY(Transient)
-	USmartObjectSubsystem* SmartObjectSubsystem = nullptr;
+	TObjectPtr<USmartObjectSubsystem> SmartObjectSubsystem = nullptr;
 };
diff --git a/Engine/Plugins/Runtime/MassGameplay/Source/MassSpawner/Public/MassEntityConfigAsset.h b/Engine/Plugins/Runtime/MassGameplay/Source/MassSpawner/Public/MassEntityConfigAsset.h
--- a/Engine/Plugins/Runtime/MassGameplay/Source/MassSpawner/Public/MassEntityConfigAsset.h
+++ b/Engine/Plugins/Runtime/MassGameplay/Source/MassSpawner/Public/MassEntityConfigAsset.h
@@ -57,11 +57,11 @@ protected:
 
 	/** Reference to parent config asset */
 	UPROPERTY(Category = "Derived Traits", EditAnywhere)
-	UMassEntityConfigAsset* Parent = nullptr;
+	TObjectPtr<UMassEntityConfigAsset> Parent = nullptr;
 
 	/** Array of unique traits of this config */
 	UPROPERTY(Category = "Traits", EditAnywhere, Instanced)
-	TArray<UMassEntityTraitBase*> Traits;
+	TArray<TObjectPtr<UMassEntityTraitBase>> Traits;
 
 private:
 	const FMassEntityTemplate* GetEntityTemplateInternal(const UWorld& World, const UObject& ConfigOwner, uint32& HashOut, FMassEntityTemplateID& TemplateIDOut, TArray<UMassEntityTraitBase*>& CombinedTraitsOut) const;
diff --git a/Engine/Plugins/Runtime/MassGameplay/Source/MassSpawner/Public/MassSpawner.h b/Engine/Plugins/Runtime/MassGameplay/Source/MassSpawner/Public/MassSpawner.h
--- a/Engine/Plugins/Runtime/MassGameplay/Source/MassSpawner/Public/MassSpawner.h
+++ b/Engine/Plugins/Runtime/MassGameplay/Source/MassSpawner/Public/MassSpawner.h
@@ -155,7 +155,7 @@ protected:
 #if WITH_EDITORONLY_DATA
 private:
 	UPROPERTY()
-	UBillboardComponent* SpriteComponent;
+	TObjectPtr<UBillboardComponent> SpriteComponent;
 #endif // WITH_EDITORONLY_DATA
  };
 
diff --git a/Engine/Plugins/Runtime/MassGameplay/Source/MassSpawner/Public/MassSpawnerSubsystem.h b/Engine/Plugins/Runtime/MassGameplay/Source/MassSpawner/Public/MassSpawnerSubsystem.h
--- a/Engine/Plugins/Runtime/MassGameplay/Source/MassSpawner/Public/MassSpawnerSubsystem.h
+++ b/Engine/Plugins/Runtime/MassGameplay/Source/MassSpawner/Public/MassSpawnerSubsystem.h
@@ -55,9 +55,9 @@ protected:
 	TSharedPtr<FMassEntityManager> EntityManager;
 
 	UPROPERTY()
-	UMassEntityTemplateRegistry* TemplateRegistryInstance;
+	TObjectPtr<UMassEntityTemplateRegistry> TemplateRegistryInstance;
 
 	UPROPERTY()
-	UMassSimulationSubsystem* SimulationSystem;
+	TObjectPtr<UMassSimulationSubsystem> SimulationSystem;
 };
 
diff --git a/Engine/Plugins/Runtime/MassGameplay/Source/MassSpawner/Public/MassSpawnerTypes.h b/Engine/Plugins/Runtime/MassGameplay/Source/MassSpawner/Public/MassSpawnerTypes.h
--- a/Engine/Plugins/Runtime/MassGameplay/Source/MassSpawner/Public/MassSpawnerTypes.h
+++ b/Engine/Plugins/Runtime/MassGameplay/Source/MassSpawner/Public/MassSpawnerTypes.h
@@ -51,7 +51,7 @@ struct FMassSpawnedEntityType
 
 private:
 	UPROPERTY(Transient)
-	mutable UMassEntityConfigAsset* EntityConfigPtr = nullptr;
+	mutable TObjectPtr<UMassEntityConfigAsset> EntityConfigPtr = nullptr;
 };
 
 USTRUCT(BlueprintType)
diff --git a/Engine/Plugins/Runtime/MixedRealityCaptureFramework/Source/MixedRealityCaptureFramework/Private/MrcProjectionBillboard.h b/Engine/Plugins/Runtime/MixedRealityCaptureFramework/Source/MixedRealityCaptureFramework/Private/MrcProjectionBillboard.h
--- a/Engine/Plugins/Runtime/MixedRealityCaptureFramework/Source/MixedRealityCaptureFramework/Private/MrcProjectionBillboard.h
+++ b/Engine/Plugins/Runtime/MixedRealityCaptureFramework/Source/MixedRealityCaptureFramework/Private/MrcProjectionBillboard.h
@@ -45,7 +45,7 @@ class AMrcProjectionActor : public AActor
 	GENERATED_UCLASS_BODY()
 
 	UPROPERTY()
-	UMixedRealityCaptureBillboard* ProjectionComponent;
+	TObjectPtr<UMixedRealityCaptureBillboard> ProjectionComponent;
 
 public:
 	//~ AActor interface
diff --git a/Engine/Plugins/Runtime/MixedRealityCaptureFramework/Source/MixedRealityCaptureFramework/Public/MixedRealityCaptureActor.h b/Engine/Plugins/Runtime/MixedRealityCaptureFramework/Source/MixedRealityCaptureFramework/Public/MixedRealityCaptureActor.h
--- a/Engine/Plugins/Runtime/MixedRealityCaptureFramework/Source/MixedRealityCaptureFramework/Public/MixedRealityCaptureActor.h
+++ b/Engine/Plugins/Runtime/MixedRealityCaptureFramework/Source/MixedRealityCaptureFramework/Public/MixedRealityCaptureActor.h
@@ -22,7 +22,7 @@ class MIXEDREALITYCAPTUREFRAMEWORK_API AMixedRealityCaptureActor : public AActor
 
 public:
 	UPROPERTY(VisibleAnywhere, BlueprintReadOnly, Category=MixedRealityCapture, meta=(AllowPrivateAccess="true"))
-	UMixedRealityCaptureComponent* CaptureComponent;
+	TObjectPtr<UMixedRealityCaptureComponent> CaptureComponent;
 
 	UPROPERTY(EditAnywhere, config, BlueprintReadOnly, Category=Tracking)
 	bool bAutoAttachToVRPlayer;
@@ -84,5 +84,5 @@ private:
 	FCastingModeRestore BroadcastManager;
 
 	UPROPERTY(Transient)
-	UStaticMesh* DebugVisualizerMesh;
+	TObjectPtr<UStaticMesh> DebugVisualizerMesh;
 };
diff --git a/Engine/Plugins/Runtime/MixedRealityCaptureFramework/Source/MixedRealityCaptureFramework/Public/MixedRealityCaptureComponent.h b/Engine/Plugins/Runtime/MixedRealityCaptureFramework/Source/MixedRealityCaptureFramework/Public/MixedRealityCaptureComponent.h
--- a/Engine/Plugins/Runtime/MixedRealityCaptureFramework/Source/MixedRealityCaptureFramework/Public/MixedRealityCaptureComponent.h
+++ b/Engine/Plugins/Runtime/MixedRealityCaptureFramework/Source/MixedRealityCaptureFramework/Public/MixedRealityCaptureComponent.h
@@ -38,10 +38,10 @@ class MIXEDREALITYCAPTUREFRAMEWORK_API UMixedRealityCaptureComponent : public US
 
 public:
 	UPROPERTY(EditAnywhere, BlueprintReadWrite, Category=VideoCapture)
-	UMediaPlayer* MediaSource;
+	TObjectPtr<UMediaPlayer> MediaSource;
 	
 	UPROPERTY(EditAnywhere, BlueprintReadWrite, BlueprintSetter=SetVidProjectionMat, Category=Composition)
-	UMaterialInterface* VideoProcessingMaterial;
+	TObjectPtr<UMaterialInterface> VideoProcessingMaterial;
 
 	UPROPERTY(EditAnywhere, BlueprintReadWrite, BlueprintSetter=SetVidProcessingParams, Category=Composition)
 	FMrcVideoProcessingParams VideoProcessingParams;
@@ -56,7 +56,7 @@ public:
 	FName TrackingSourceName;
 
 	UPROPERTY(EditAnywhere, BlueprintReadWrite, Category=Composition)
-	UTextureRenderTarget2D* GarbageMatteCaptureTextureTarget;
+	TObjectPtr<UTextureRenderTarget2D> GarbageMatteCaptureTextureTarget;
 
 	/** Millisecond delay to apply to motion controller components when rendering to the capture view (to better align with latent camera feeds) */
 	UPROPERTY(BlueprintReadWrite, Config, BlueprintGetter=GetTrackingDelay, BlueprintSetter=SetTrackingDelay, Category=Composition, meta=(ClampMin="0", UIMin="0"))
@@ -230,28 +230,28 @@ private:
 private:
 #if WITH_EDITORONLY_DATA
 	UPROPERTY(Transient)
-	UStaticMesh* ProxyMesh;
+	TObjectPtr<UStaticMesh> ProxyMesh;
 
 	UPROPERTY(Transient)
-	UStaticMeshComponent* ProxyMeshComponent;
+	TObjectPtr<UStaticMeshComponent> ProxyMeshComponent;
 #endif
 
 	UPROPERTY(Transient)
-	UChildActorComponent* ProjectionActor;
+	TObjectPtr<UChildActorComponent> ProjectionActor;
 
 	UPROPERTY(Transient)
-	UMotionControllerComponent* PairedTracker;
+	TObjectPtr<UMotionControllerComponent> PairedTracker;
 
 	UPROPERTY(Transient)
-	USceneComponent* TrackingOriginOffset;
+	TObjectPtr<USceneComponent> TrackingOriginOffset;
 
 	EHMDTrackingOrigin::Type RelativeOriginType;
 
 	UPROPERTY(Transient)
-	UMrcGarbageMatteCaptureComponent* GarbageMatteCaptureComponent;
+	TObjectPtr<UMrcGarbageMatteCaptureComponent> GarbageMatteCaptureComponent;
 
 	UPROPERTY(Transient)
-	UTexture2D* DistortionDisplacementMap;
+	TObjectPtr<UTexture2D> DistortionDisplacementMap;
 
 	float CalibratedFOV;
 	FOpenCVCameraViewInfo UndistortedCameraInfo;
diff --git a/Engine/Plugins/Runtime/MixedRealityCaptureFramework/Source/MixedRealityCaptureFramework/Public/MrcGarbageMatteCaptureComponent.h b/Engine/Plugins/Runtime/MixedRealityCaptureFramework/Source/MixedRealityCaptureFramework/Public/MrcGarbageMatteCaptureComponent.h
--- a/Engine/Plugins/Runtime/MixedRealityCaptureFramework/Source/MixedRealityCaptureFramework/Public/MrcGarbageMatteCaptureComponent.h
+++ b/Engine/Plugins/Runtime/MixedRealityCaptureFramework/Source/MixedRealityCaptureFramework/Public/MrcGarbageMatteCaptureComponent.h
@@ -57,10 +57,10 @@ private:
 	TSubclassOf<AMrcGarbageMatteActor> GarbageMatteActorClass;
 
 	UPROPERTY(Transient)
-	AMrcGarbageMatteActor* GarbageMatteActor;
+	TObjectPtr<AMrcGarbageMatteActor> GarbageMatteActor;
 
 	UPROPERTY(Transient)
-	TArray<AMrcGarbageMatteActor*> SpawnedActors;
+	TArray<TObjectPtr<AMrcGarbageMatteActor>> SpawnedActors;
 
 	UPROPERTY(Transient)
 	TWeakObjectPtr<USceneComponent> TrackingOriginPtr;
@@ -96,11 +96,11 @@ public:
 
 private:
 	UPROPERTY(Transient)
-	UStaticMesh* GarbageMatteMesh;
+	TObjectPtr<UStaticMesh> GarbageMatteMesh;
 
 	UPROPERTY(Transient)
-	UMaterial* GarbageMatteMaterial;
+	TObjectPtr<UMaterial> GarbageMatteMaterial;
 
 	UPROPERTY(Transient, BlueprintReadOnly, Category = "MixedRealityCapture|GarbageMatting", meta=(AllowPrivateAccess="true"))
-	TArray<UPrimitiveComponent*> GarbageMattes;
+	TArray<TObjectPtr<UPrimitiveComponent>> GarbageMattes;
 };
diff --git a/Engine/Plugins/Runtime/ModelViewViewModel/Source/ModelViewViewModel/Public/Bindings/MVVMCompiledBindingLibrary.h b/Engine/Plugins/Runtime/ModelViewViewModel/Source/ModelViewViewModel/Public/Bindings/MVVMCompiledBindingLibrary.h
--- a/Engine/Plugins/Runtime/ModelViewViewModel/Source/ModelViewViewModel/Public/Bindings/MVVMCompiledBindingLibrary.h
+++ b/Engine/Plugins/Runtime/ModelViewViewModel/Source/ModelViewViewModel/Public/Bindings/MVVMCompiledBindingLibrary.h
@@ -293,7 +293,7 @@ private:
 	TArray<FProperty*> LoadedProperties;
 
 	UPROPERTY(Transient)
-	TArray<UFunction*> LoadedFunctions;
+	TArray<TObjectPtr<UFunction>> LoadedFunctions;
 
 	TArray<UE::FieldNotification::FFieldId> LoadedFieldIds;
 
diff --git a/Engine/Plugins/Runtime/ModelViewViewModel/Source/ModelViewViewModel/Public/MVVMSubsystem.h b/Engine/Plugins/Runtime/ModelViewViewModel/Source/ModelViewViewModel/Public/MVVMSubsystem.h
--- a/Engine/Plugins/Runtime/ModelViewViewModel/Source/ModelViewViewModel/Public/MVVMSubsystem.h
+++ b/Engine/Plugins/Runtime/ModelViewViewModel/Source/ModelViewViewModel/Public/MVVMSubsystem.h
@@ -95,5 +95,5 @@ public:
 
 private:
 	UPROPERTY(Transient)
-	UMVVMViewModelCollectionObject* GlobalViewModelCollection;
+	TObjectPtr<UMVVMViewModelCollectionObject> GlobalViewModelCollection;
 };
diff --git a/Engine/Plugins/Runtime/ModelViewViewModel/Source/ModelViewViewModelBlueprint/Public/MVVMBlueprintViewModelContext.h b/Engine/Plugins/Runtime/ModelViewViewModel/Source/ModelViewViewModelBlueprint/Public/MVVMBlueprintViewModelContext.h
--- a/Engine/Plugins/Runtime/ModelViewViewModel/Source/ModelViewViewModelBlueprint/Public/MVVMBlueprintViewModelContext.h
+++ b/Engine/Plugins/Runtime/ModelViewViewModel/Source/ModelViewViewModelBlueprint/Public/MVVMBlueprintViewModelContext.h
@@ -80,7 +80,7 @@ private:
 	FGuid ViewModelContextId;
 
 	UPROPERTY(VisibleAnywhere, Category = "MVVM", NoClear, meta = (AllowedClasses = "/Script/UMG.NotifyFieldValueChanged", DisallowedClasses = "/Script/UMG.Widget"))
-	UClass* NotifyFieldValueClass = nullptr;
+	TObjectPtr<UClass> NotifyFieldValueClass = nullptr;
 
 	UPROPERTY()
 	TSubclassOf<UMVVMViewModelBase> ViewModelClass_DEPRECATED;
diff --git a/Engine/Plugins/Runtime/NavCorridor/Source/NavCorridor/Public/NavCorridorTestingComponent.h b/Engine/Plugins/Runtime/NavCorridor/Source/NavCorridor/Public/NavCorridorTestingComponent.h
--- a/Engine/Plugins/Runtime/NavCorridor/Source/NavCorridor/Public/NavCorridorTestingComponent.h
+++ b/Engine/Plugins/Runtime/NavCorridor/Source/NavCorridor/Public/NavCorridorTestingComponent.h
@@ -102,5 +102,5 @@ public:
 
 protected:
 	UPROPERTY(Category = Default, VisibleAnywhere, meta = (AllowPrivateAccess = "true"))
-	UNavCorridorTestingComponent* DebugComp;
+	TObjectPtr<UNavCorridorTestingComponent> DebugComp;
 };
diff --git a/Engine/Plugins/Runtime/NetworkPrediction/Source/NetworkPrediction/Public/NetworkPredictionPhysicsComponent.h b/Engine/Plugins/Runtime/NetworkPrediction/Source/NetworkPrediction/Public/NetworkPredictionPhysicsComponent.h
--- a/Engine/Plugins/Runtime/NetworkPrediction/Source/NetworkPrediction/Public/NetworkPredictionPhysicsComponent.h
+++ b/Engine/Plugins/Runtime/NetworkPrediction/Source/NetworkPrediction/Public/NetworkPredictionPhysicsComponent.h
@@ -53,7 +53,7 @@ protected:
 	FNetworkPredictionProxy NetworkPredictionProxy;
 
 	UPROPERTY()
-	UPrimitiveComponent* UpdatedPrimitive = nullptr;
+	TObjectPtr<UPrimitiveComponent> UpdatedPrimitive = nullptr;
 
 	UPROPERTY(Replicated, transient)
 	FReplicationProxy ReplicationProxy;
diff --git a/Engine/Plugins/Runtime/NetworkPrediction/Source/NetworkPrediction/Public/NetworkPredictionProxy.h b/Engine/Plugins/Runtime/NetworkPrediction/Source/NetworkPrediction/Public/NetworkPredictionProxy.h
--- a/Engine/Plugins/Runtime/NetworkPrediction/Source/NetworkPrediction/Public/NetworkPredictionProxy.h
+++ b/Engine/Plugins/Runtime/NetworkPrediction/Source/NetworkPrediction/Public/NetworkPredictionProxy.h
@@ -202,7 +202,7 @@ private:
 	TFunction<void(FNetworkPredictionProxy* const, FNetworkPredictionID NewID, EConfigAction Action)>	ConfigFunc;
 
 	UPROPERTY()
-	UNetworkPredictionWorldManager* WorldManager = nullptr;
+	TObjectPtr<UNetworkPredictionWorldManager> WorldManager = nullptr;
 };
 
 template<>
diff --git a/Engine/Plugins/Runtime/NetworkPrediction/Source/NetworkPrediction/Public/NetworkPredictionWorldManager.h b/Engine/Plugins/Runtime/NetworkPrediction/Source/NetworkPrediction/Public/NetworkPredictionWorldManager.h
--- a/Engine/Plugins/Runtime/NetworkPrediction/Source/NetworkPrediction/Public/NetworkPredictionWorldManager.h
+++ b/Engine/Plugins/Runtime/NetworkPrediction/Source/NetworkPrediction/Public/NetworkPredictionWorldManager.h
@@ -31,7 +31,7 @@ public:
 
 	// Server created, replicated manager (only used for centralized/system wide data replication)
 	UPROPERTY()
-	ANetworkPredictionReplicatedManager* ReplicatedManager = nullptr;
+	TObjectPtr<ANetworkPredictionReplicatedManager> ReplicatedManager = nullptr;
 
 	// Subsystem Init/Deinit
 	virtual void Initialize(FSubsystemCollectionBase& Collection) override;
diff --git a/Engine/Plugins/Runtime/NetworkPredictionExtras/Source/NetworkPredictionExtras/Public/BaseMovementComponent.h b/Engine/Plugins/Runtime/NetworkPredictionExtras/Source/NetworkPredictionExtras/Public/BaseMovementComponent.h
--- a/Engine/Plugins/Runtime/NetworkPredictionExtras/Source/NetworkPredictionExtras/Public/BaseMovementComponent.h
+++ b/Engine/Plugins/Runtime/NetworkPredictionExtras/Source/NetworkPredictionExtras/Public/BaseMovementComponent.h
@@ -41,10 +41,10 @@ protected:
 	virtual void PhysicsVolumeChanged(class APhysicsVolume* NewVolume);	
 
 	UPROPERTY()
-	USceneComponent* UpdatedComponent = nullptr;
+	TObjectPtr<USceneComponent> UpdatedComponent = nullptr;
 
 	UPROPERTY()
-	UPrimitiveComponent* UpdatedPrimitive = nullptr;
+	TObjectPtr<UPrimitiveComponent> UpdatedPrimitive = nullptr;
 
 private:
 
diff --git a/Engine/Plugins/Runtime/NetworkPredictionExtras/Source/NetworkPredictionExtras/Public/MockRootMotionComponent.h b/Engine/Plugins/Runtime/NetworkPredictionExtras/Source/NetworkPredictionExtras/Public/MockRootMotionComponent.h
--- a/Engine/Plugins/Runtime/NetworkPredictionExtras/Source/NetworkPredictionExtras/Public/MockRootMotionComponent.h
+++ b/Engine/Plugins/Runtime/NetworkPredictionExtras/Source/NetworkPredictionExtras/Public/MockRootMotionComponent.h
@@ -18,7 +18,7 @@ struct FRootMotionSourceCache
 	GENERATED_BODY()
 
 	UPROPERTY()
-	UMockRootMotionSource* Instance = nullptr;
+	TObjectPtr<UMockRootMotionSource> Instance = nullptr;
 
 	int32 ClassID = 0;
 };
diff --git a/Engine/Plugins/Runtime/NetworkPredictionExtras/Source/NetworkPredictionExtras/Public/MockRootMotionSourceObject.h b/Engine/Plugins/Runtime/NetworkPredictionExtras/Source/NetworkPredictionExtras/Public/MockRootMotionSourceObject.h
--- a/Engine/Plugins/Runtime/NetworkPredictionExtras/Source/NetworkPredictionExtras/Public/MockRootMotionSourceObject.h
+++ b/Engine/Plugins/Runtime/NetworkPredictionExtras/Source/NetworkPredictionExtras/Public/MockRootMotionSourceObject.h
@@ -141,7 +141,7 @@ class UMockRootMotionSource_Montage : public UMockRootMotionSource
 protected:
 
 	UPROPERTY(EditAnywhere, Category=RootMotion)
-	UAnimMontage* Montage;
+	TObjectPtr<UAnimMontage> Montage;
 
 	UPROPERTY(EditAnywhere, Category=RootMotion)
 	float PlayRate = 1.f;
@@ -170,7 +170,7 @@ class UMockRootMotionSource_Curve : public UMockRootMotionSource
 protected:
 
 	UPROPERTY(EditAnywhere, Category=RootMotion)
-	UCurveVector* Curve;
+	TObjectPtr<UCurveVector> Curve;
 
 	UPROPERTY(EditAnywhere, Category=RootMotion)
 	float PlayRate = 1.f;
