Fix typos with type name, FMinimapReplicationTagCountMap to actually match the type name FMinimalReplicationTagCountMap.

diff --git a/Engine/Plugins/Runtime/GameplayAbilities/Source/GameplayAbilities/Private/GameplayEffectTypes.cpp b/Engine/Plugins/Runtime/GameplayAbilities/Source/GameplayAbilities/Private/GameplayEffectTypes.cpp
--- a/Engine/Plugins/Runtime/GameplayAbilities/Source/GameplayAbilities/Private/GameplayEffectTypes.cpp
+++ b/Engine/Plugins/Runtime/GameplayAbilities/Source/GameplayAbilities/Private/GameplayEffectTypes.cpp
@@ -1133,7 +1133,7 @@ bool FMinimalReplicationTagCountMap::NetSerialize(FArchive& Ar, class UPackageMa
 		if (Count > MaxCount)
 		{
 #if UE_BUILD_SHIPPING
-			ABILITY_LOG(Error, TEXT("FMinimalReplicationTagCountMap has too many tags (%d) when the limit is %d. This will cause tags to not replicate. See FMinimapReplicationTagCountMap::NetSerialize"), TagMap.Num(), MaxCount);
+			ABILITY_LOG(Error, TEXT("FMinimalReplicationTagCountMap has too many tags (%d) when the limit is %d. This will cause tags to not replicate. See FMinimalReplicationTagCountMap::NetSerialize"), TagMap.Num(), MaxCount);
 #else
 			TArray<FGameplayTag> TagKeys;
 			TagMap.GetKeys(TagKeys);
@@ -1152,7 +1152,7 @@ bool FMinimalReplicationTagCountMap::NetSerialize(FArchive& Ar, class UPackageMa
 				TagKeys[TagIndex].GetTagName().AppendString(TagsString);
 			}
 
-			ABILITY_LOG(Error, TEXT("FMinimalReplicationTagCountMap has too many tags (%d) when the limit is %d. This will cause tags to not replicate. See FMinimapReplicationTagCountMap::NetSerialize\nTagMap tags: %s"), TagMap.Num(), MaxCount, *TagsString);
+			ABILITY_LOG(Error, TEXT("FMinimalReplicationTagCountMap has too many tags (%d) when the limit is %d. This will cause tags to not replicate. See FMinimalReplicationTagCountMap::NetSerialize\nTagMap tags: %s"), TagMap.Num(), MaxCount, *TagsString);
 #endif	
 
 			//clamp the count
diff --git a/Engine/Plugins/Runtime/GameplayAbilities/Source/GameplayAbilities/Public/AbilitySystemGlobals.h b/Engine/Plugins/Runtime/GameplayAbilities/Source/GameplayAbilities/Public/AbilitySystemGlobals.h
--- a/Engine/Plugins/Runtime/GameplayAbilities/Source/GameplayAbilities/Public/AbilitySystemGlobals.h
+++ b/Engine/Plugins/Runtime/GameplayAbilities/Source/GameplayAbilities/Public/AbilitySystemGlobals.h
@@ -225,7 +225,7 @@ class GAMEPLAYABILITIES_API UAbilitySystemGlobals : public UObject
 	UPROPERTY(config)
 	FName ActivateFailNetworkingName;
 
-	/** How many bits to use for "number of tags" in FMinimapReplicationTagCountMap::NetSerialize.  */
+	/** How many bits to use for "number of tags" in FMinimalReplicationTagCountMap::NetSerialize.  */
 	UPROPERTY(config)
 	int32	MinimalReplicationTagCountBits;
 
diff --git a/Engine/Plugins/Runtime/GameplayAbilities/Source/GameplayAbilities/Public/GameplayEffectTypes.h b/Engine/Plugins/Runtime/GameplayAbilities/Source/GameplayAbilities/Public/GameplayEffectTypes.h
--- a/Engine/Plugins/Runtime/GameplayAbilities/Source/GameplayAbilities/Public/GameplayEffectTypes.h
+++ b/Engine/Plugins/Runtime/GameplayAbilities/Source/GameplayAbilities/Public/GameplayEffectTypes.h
@@ -1654,7 +1654,7 @@ struct GAMEPLAYABILITIES_API FMinimalReplicationTagCountMap
 		}
 		else if (Count < 0)
 		{
-			ABILITY_LOG(Error, TEXT("FMinimapReplicationTagCountMap::RemoveTag called on Tag %s and count is now < 0"), *Tag.ToString());
+			ABILITY_LOG(Error, TEXT("FMinimalReplicationTagCountMap::RemoveTag called on Tag %s and count is now < 0"), *Tag.ToString());
 			Count = 0;
 		}
 	}
