Add mechanism to filter DebugAbilityTags output

diff --git a/Engine/Plugins/Runtime/GameplayAbilities/Source/GameplayAbilities/Private/AbilitySystemDebugHUD.cpp b/Engine/Plugins/Runtime/GameplayAbilities/Source/GameplayAbilities/Private/AbilitySystemDebugHUD.cpp
--- a/Engine/Plugins/Runtime/GameplayAbilities/Source/GameplayAbilities/Private/AbilitySystemDebugHUD.cpp
+++ b/Engine/Plugins/Runtime/GameplayAbilities/Source/GameplayAbilities/Private/AbilitySystemDebugHUD.cpp
@@ -7,6 +7,7 @@
 #include "EngineGlobals.h"
 #include "Engine/Engine.h"
 #include "Debug/DebugDrawService.h"
+#include "Editor.h"
 #include "Engine/LocalPlayer.h"
 #include "EngineUtils.h"
 #include "AbilitySystemComponent.h"
@@ -23,21 +24,39 @@ namespace
 	);
 }
 
-bool UAbilitySystemDebugHUDExtension_Tags::bEnabled = false;
-
-TSet<FString> UAbilitySystemDebugHUDExtension_Attributes::AttributesToDisplay;
-
-bool UAbilitySystemDebugHUDExtension_Attributes::bIncludeModifiers = false;
-
-bool UAbilitySystemDebugHUDExtension_BlockedAbilityTags::bEnabled = false;
-
 bool AAbilitySystemDebugHUD::bEnableBasicHUD = false;
 
 TSubclassOf<AAbilitySystemDebugHUD> AAbilitySystemDebugHUD::HUDClass = AAbilitySystemDebugHUD::StaticClass();
 
 void UAbilitySystemDebugHUDExtension_Tags::ToggleExtension(const TArray<FString>& Args, UWorld* World)
 {
-	bEnabled = !bEnabled;
+	UAbilitySystemDebugHUDExtension_Tags* Extension = GetMutableDefault<UAbilitySystemDebugHUDExtension_Tags>();
+	if (Args.IsEmpty())
+	{
+		const bool bIsCurrentlyEnabled = Extension->bEnabled || !Extension->TagsToDisplay.IsEmpty();
+
+		if (bIsCurrentlyEnabled)
+		{
+			Extension->TagsToDisplay.Empty();
+		}
+
+		Extension->bEnabled = !bIsCurrentlyEnabled;
+	}
+	else
+	{
+		for (const FString& Tag : Args)
+		{
+			FSetElementId Element = Extension->TagsToDisplay.FindId(Tag);
+			if (Element.IsValidId())
+			{
+				Extension->TagsToDisplay.Remove(Element);
+			}
+			else
+			{
+				Extension->TagsToDisplay.Add(Tag);
+			}
+		}
+	}
 
 	AAbilitySystemDebugHUD::NotifyExtensionEnableChanged(World);
 }
@@ -71,29 +90,34 @@ void UAbilitySystemDebugHUDExtension_Tags::GetDebugStrings(const AActor* Actor,
 	Comp->GetOwnedGameplayTags(Container);
 	for (auto ContainerIt = Container.CreateConstIterator(); ContainerIt; ++ContainerIt)
 	{
-		if (const uint32* TagCount = TagCounts.Find(*ContainerIt))
+		if (TagsToDisplay.IsEmpty() || TagsToDisplay.Contains(ContainerIt->ToString()))
 		{
-			OutDebugStrings.Add(FString::Printf(TEXT("%s x%d"), *ContainerIt->ToString(), *TagCount));
-		}
-		else
-		{
-			OutDebugStrings.Add(ContainerIt->ToString());
+			if (const uint32* TagCount = TagCounts.Find(*ContainerIt))
+			{
+				OutDebugStrings.Add(FString::Printf(TEXT("%s x%d"), *ContainerIt->ToString(), *TagCount));
+			}
+			else
+			{
+				OutDebugStrings.Add(ContainerIt->ToString());
+			}
 		}
 	}
 }
 
 void UAbilitySystemDebugHUDExtension_Attributes::ToggleExtension(const TArray<FString>& Args, UWorld* World)
 {
+	UAbilitySystemDebugHUDExtension_Attributes* Extension = GetMutableDefault<UAbilitySystemDebugHUDExtension_Attributes>();
+
 	for (const FString& Arg : Args)
 	{
-		FSetElementId FoundId = AttributesToDisplay.FindId(Arg);
+		FSetElementId FoundId = Extension->AttributesToDisplay.FindId(Arg);
 		if (FoundId.IsValidId())
 		{
-			AttributesToDisplay.Remove(FoundId);
+			Extension->AttributesToDisplay.Remove(FoundId);
 		}
 		else
 		{
-			AttributesToDisplay.Add(Arg);
+			Extension->AttributesToDisplay.Add(Arg);
 		}
 	}
 
@@ -102,12 +126,14 @@ void UAbilitySystemDebugHUDExtension_Attributes::ToggleExtension(const TArray<FS
 
 void UAbilitySystemDebugHUDExtension_Attributes::ToggleIncludeModifiers()
 {
-	bIncludeModifiers = !bIncludeModifiers;
+	UAbilitySystemDebugHUDExtension_Attributes* Extension = GetMutableDefault<UAbilitySystemDebugHUDExtension_Attributes>();
+	Extension->bIncludeModifiers = !Extension->bIncludeModifiers;
 }
 
 void UAbilitySystemDebugHUDExtension_Attributes::ClearDisplayedAttributes(const TArray<FString>& Args, UWorld* World)
 {
-	AttributesToDisplay.Empty();
+	UAbilitySystemDebugHUDExtension_Attributes* Extension = GetMutableDefault<UAbilitySystemDebugHUDExtension_Attributes>();
+	Extension->AttributesToDisplay.Empty();
 
 	AAbilitySystemDebugHUD::NotifyExtensionEnableChanged(World);
 }
@@ -170,7 +196,33 @@ void UAbilitySystemDebugHUDExtension_Attributes::GetDebugStrings(const AActor* A
 
 void UAbilitySystemDebugHUDExtension_BlockedAbilityTags::ToggleExtension(const TArray<FString>& Args, UWorld* World)
 {
-	bEnabled = !bEnabled;
+	UAbilitySystemDebugHUDExtension_BlockedAbilityTags* Extension = GetMutableDefault<UAbilitySystemDebugHUDExtension_BlockedAbilityTags>();
+	if (Args.IsEmpty())
+	{
+		const bool bIsCurrentlyEnabled = Extension->bEnabled || !Extension->TagsToDisplay.IsEmpty();
+
+		if (bIsCurrentlyEnabled)
+		{
+			Extension->TagsToDisplay.Empty();
+		}
+
+		Extension->bEnabled = !bIsCurrentlyEnabled;
+	}
+	else
+	{
+		for (const FString& Tag : Args)
+		{
+			FSetElementId Element = Extension->TagsToDisplay.FindId(Tag);
+			if (Element.IsValidId())
+			{
+				Extension->TagsToDisplay.Remove(Element);
+			}
+			else
+			{
+				Extension->TagsToDisplay.Add(Tag);
+			}
+		}
+	}
 
 	AAbilitySystemDebugHUD::NotifyExtensionEnableChanged(World);
 }
@@ -204,13 +256,16 @@ void UAbilitySystemDebugHUDExtension_BlockedAbilityTags::GetDebugStrings(const A
 	Comp->GetBlockedAbilityTags(Container);
 	for (auto ContainerIt = Container.CreateConstIterator(); ContainerIt; ++ContainerIt)
 	{
-		if (const uint32* TagCount = TagCounts.Find(*ContainerIt))
-		{
-			OutDebugStrings.Add(FString::Printf(TEXT("%s x%d"), *ContainerIt->ToString(), *TagCount));
-		}
-		else
+		if (TagsToDisplay.IsEmpty() || TagsToDisplay.Contains(ContainerIt->ToString()))
 		{
-			OutDebugStrings.Add(ContainerIt->ToString());
+			if (const uint32* TagCount = TagCounts.Find(*ContainerIt))
+			{
+				OutDebugStrings.Add(FString::Printf(TEXT("%s x%d"), *ContainerIt->ToString(), *TagCount));
+			}
+			else
+			{
+				OutDebugStrings.Add(ContainerIt->ToString());
+			}
 		}
 	}
 }
@@ -248,6 +303,7 @@ void AAbilitySystemDebugHUD::NotifyExtensionEnableChanged(UWorld* InWorld)
 	}
 
 	static FDelegateHandle DrawDebugDelegateHandle;
+	static FDelegateHandle PostBeginPIEDelegateHandle;
 
 	if (!HUD && bAnyExtensionEnabled)
 	{
@@ -261,11 +317,19 @@ void AAbilitySystemDebugHUD::NotifyExtensionEnableChanged(UWorld* InWorld)
 
 		FDebugDrawDelegate DrawDebugDelegate = FDebugDrawDelegate::CreateUObject(HUD, &AAbilitySystemDebugHUD::DrawDebugHUD);
 		DrawDebugDelegateHandle = UDebugDrawService::Register(TEXT("GameplayDebug"), DrawDebugDelegate);
+
+		PostBeginPIEDelegateHandle = FEditorDelegates::PostPIEStarted.AddLambda([](const bool bIsSimulating)
+			{
+				if (FWorldContext* PIEWorldContext = GEditor->GetPIEWorldContext())
+				{
+					NotifyExtensionEnableChanged(PIEWorldContext->World());
+				}
+			});
 	}
 	else if (HUD && !bAnyExtensionEnabled)
 	{
-		FDebugDrawDelegate DrawDebugDelegate = FDebugDrawDelegate::CreateUObject(HUD, &AAbilitySystemDebugHUD::DrawDebugHUD);
 		UDebugDrawService::Unregister(DrawDebugDelegateHandle);
+		FEditorDelegates::PostPIEStarted.Remove(PostBeginPIEDelegateHandle);
 		HUD->Destroy();
 	}
 }
@@ -537,7 +601,7 @@ FAutoConsoleCommandWithWorldAndArgs AbilitySystemToggleDebugHUDCommand(
 
 FAutoConsoleCommandWithWorldAndArgs AbilitySystemToggleDebugTagsCommand(
 	TEXT("AbilitySystem.DebugAbilityTags"),
-	TEXT("Toggles Drawing Ability Tags on Actors with AbilitySystemComponents"),
+	TEXT("Usage: AbilitySystem.DebugAbilityTags [TagName] [TagName]...\nToggles Drawing Ability Tags on Actors with AbilitySystemComponents. If no tags are given, draws all owned tags."),
 	FConsoleCommandWithWorldAndArgsDelegate::CreateStatic(UAbilitySystemDebugHUDExtension_Tags::ToggleExtension)
 );
 
@@ -561,7 +625,7 @@ FAutoConsoleCommand AbilitySystemToggleIncludeModifiersCommand(
 
 FAutoConsoleCommandWithWorldAndArgs AbilitySystemToggleDebugBlockedTagsCommand(
 	TEXT("AbilitySystem.DebugBlockedAbilityTags"),
-	TEXT("Toggles Drawing Blocked Ability Tags on Actors with AbilitySystemComponents"),
+	TEXT("Usage: AbilitySystem.DebugBlockedAbilityTags [TagName] [TagName]...\nToggles Drawing Blocked Ability Tags on Actors with AbilitySystemComponents. If no tags are given, draws all blocked tags."),
 	FConsoleCommandWithWorldAndArgsDelegate::CreateStatic(UAbilitySystemDebugHUDExtension_BlockedAbilityTags::ToggleExtension)
 );
 
diff --git a/Engine/Plugins/Runtime/GameplayAbilities/Source/GameplayAbilities/Public/AbilitySystemDebugHUD.h b/Engine/Plugins/Runtime/GameplayAbilities/Source/GameplayAbilities/Public/AbilitySystemDebugHUD.h
--- a/Engine/Plugins/Runtime/GameplayAbilities/Source/GameplayAbilities/Public/AbilitySystemDebugHUD.h
+++ b/Engine/Plugins/Runtime/GameplayAbilities/Source/GameplayAbilities/Public/AbilitySystemDebugHUD.h
@@ -65,7 +65,9 @@ class UAbilitySystemDebugHUDExtension_Tags final : public UAbilitySystemDebugHUD
 {
 	GENERATED_BODY()
 
-	static bool bEnabled;
+	TSet<FString> TagsToDisplay;
+
+	bool bEnabled;
 
 public:
 
@@ -86,9 +88,9 @@ class UAbilitySystemDebugHUDExtension_Attributes final : public UAbilitySystemDe
 {
 	GENERATED_BODY()
 
-	static TSet<FString> AttributesToDisplay;
+	TSet<FString> AttributesToDisplay;
 
-	static bool bIncludeModifiers;
+	bool bIncludeModifiers;
 public:
 	static void ToggleExtension(const TArray<FString>& Args, UWorld* World);
 	
@@ -111,7 +113,9 @@ class UAbilitySystemDebugHUDExtension_BlockedAbilityTags final : public UAbility
 {
 	GENERATED_BODY()
 
-	static bool bEnabled;
+	TSet<FString> TagsToDisplay;
+
+	bool bEnabled;
 
 public:
 
