Fix Editor References in runtime functions

#rnx

diff --git a/Engine/Plugins/Runtime/GameplayAbilities/Source/GameplayAbilities/Private/AbilitySystemDebugHUD.cpp b/Engine/Plugins/Runtime/GameplayAbilities/Source/GameplayAbilities/Private/AbilitySystemDebugHUD.cpp
--- a/Engine/Plugins/Runtime/GameplayAbilities/Source/GameplayAbilities/Private/AbilitySystemDebugHUD.cpp
+++ b/Engine/Plugins/Runtime/GameplayAbilities/Source/GameplayAbilities/Private/AbilitySystemDebugHUD.cpp
@@ -7,13 +7,16 @@
 #include "EngineGlobals.h"
 #include "Engine/Engine.h"
 #include "Debug/DebugDrawService.h"
-#include "Editor.h"
 #include "Engine/LocalPlayer.h"
 #include "EngineUtils.h"
 #include "AbilitySystemComponent.h"
 #include "AbilitySystemGlobals.h"
 #include "GameplayEffect.h"
 
+#if WITH_EDITOR
+#include "Editor.h"
+#endif
+
 namespace
 {
 	static float DebugDrawMaxDistance = 2048.f;
@@ -318,6 +321,7 @@ void AAbilitySystemDebugHUD::NotifyExtensionEnableChanged(UWorld* InWorld)
 		FDebugDrawDelegate DrawDebugDelegate = FDebugDrawDelegate::CreateUObject(HUD, &AAbilitySystemDebugHUD::DrawDebugHUD);
 		DrawDebugDelegateHandle = UDebugDrawService::Register(TEXT("GameplayDebug"), DrawDebugDelegate);
 
+#if WITH_EDITOR
 		PostBeginPIEDelegateHandle = FEditorDelegates::PostPIEStarted.AddLambda([](const bool bIsSimulating)
 			{
 				if (FWorldContext* PIEWorldContext = GEditor->GetPIEWorldContext())
@@ -325,11 +329,15 @@ void AAbilitySystemDebugHUD::NotifyExtensionEnableChanged(UWorld* InWorld)
 					NotifyExtensionEnableChanged(PIEWorldContext->World());
 				}
 			});
+#endif
+
 	}
 	else if (HUD && !bAnyExtensionEnabled)
 	{
 		UDebugDrawService::Unregister(DrawDebugDelegateHandle);
+#if WITH_EDITOR
 		FEditorDelegates::PostPIEStarted.Remove(PostBeginPIEDelegateHandle);
+#endif
 		HUD->Destroy();
 	}
 }
