Modify FGameplayAbilityActorInfo to defer to Pawn for the IsLocallyControlled call in the case where we don't have a player controller

diff --git a/Engine/Plugins/Runtime/GameplayAbilities/Source/GameplayAbilities/Private/GameplayAbilityTypes.cpp b/Engine/Plugins/Runtime/GameplayAbilities/Source/GameplayAbilities/Private/GameplayAbilityTypes.cpp
--- a/Engine/Plugins/Runtime/GameplayAbilities/Source/GameplayAbilities/Private/GameplayAbilityTypes.cpp
+++ b/Engine/Plugins/Runtime/GameplayAbilities/Source/GameplayAbilities/Private/GameplayAbilityTypes.cpp
@@ -104,6 +104,10 @@ bool FGameplayAbilityActorInfo::IsLocallyControlled() const
 	{
 		return PC->IsLocalController();
 	}
+	else if (const APawn* OwnerPawn = Cast<APawn>(OwnerActor))
+	{
+		return OwnerPawn->IsLocallyControlled();
+	}
 	else if (IsNetAuthority())
 	{
 		// Non-players are always locally controlled on the server
