Remove code from UGameplay Ability that was incorrectly only registering the last instanced ability with the BP debugger for breakpoints

diff --git a/Engine/Plugins/Runtime/GameplayAbilities/Source/GameplayAbilities/Private/Abilities/GameplayAbility.cpp b/Engine/Plugins/Runtime/GameplayAbilities/Source/GameplayAbilities/Private/Abilities/GameplayAbility.cpp
--- a/Engine/Plugins/Runtime/GameplayAbilities/Source/GameplayAbilities/Private/Abilities/GameplayAbility.cpp
+++ b/Engine/Plugins/Runtime/GameplayAbilities/Source/GameplayAbilities/Private/Abilities/GameplayAbility.cpp
@@ -58,18 +58,6 @@ UGameplayAbility::UGameplayAbility(const FObjectInitializer& ObjectInitializer)
 		UFunction* ActivateFunction = GetClass()->FindFunctionByName(FuncName);
 		bHasBlueprintActivateFromEvent = ImplementedInBlueprint(ActivateFunction);
 	}
-	
-#if WITH_EDITOR
-	/** Autoregister abilities with the blueprint debugger in the editor.*/
-	if (!HasAnyFlags(RF_ClassDefaultObject))
-	{
-		UBlueprint* BP = Cast<UBlueprint>(GetClass()->ClassGeneratedBy);
-		if (BP && (BP->GetWorldBeingDebugged() == nullptr || BP->GetWorldBeingDebugged() == GetWorld()))
-		{
-			BP->SetObjectBeingDebugged(this);
-		}
-	}
-#endif
 
 	bServerRespectsRemoteAbilityCancellation = true;
 	bReplicateInputDirectly = false;
