Add debug command to filter GameplayCue events

diff --git a/Engine/Plugins/Runtime/GameplayAbilities/Source/GameplayAbilities/Private/GameplayCueSet.cpp b/Engine/Plugins/Runtime/GameplayAbilities/Source/GameplayAbilities/Private/GameplayCueSet.cpp
--- a/Engine/Plugins/Runtime/GameplayAbilities/Source/GameplayAbilities/Private/GameplayCueSet.cpp
+++ b/Engine/Plugins/Runtime/GameplayAbilities/Source/GameplayAbilities/Private/GameplayCueSet.cpp
@@ -11,6 +11,36 @@
 
 UE_DEFINE_GAMEPLAY_TAG_STATIC(StaticTag_GameplayCue, TEXT("GameplayCue"));
 
+namespace GameplayCueDebug
+{
+	static FGameplayTagContainer DebugGameplayCueFilter;
+	FAutoConsoleCommand ConCommandDebugGameplayCueFilter(
+		TEXT("GameplayCue.FilterCuesByTag"),
+		TEXT("Adds or removes a GameplayCue tag to the debug filter list. If the filter is populated, only GameplayCues with tags in the filter will be invoked."),
+		FConsoleCommandWithArgsDelegate::CreateLambda([](const TArray<FString>& Args) 
+			{
+				for (const FString& TagToFilter : Args)
+				{
+					const FGameplayTag ExistingTag = FGameplayTag::RequestGameplayTag(FName(TagToFilter));
+					if (!ExistingTag.IsValid())
+					{
+						continue;
+					}
+
+					if (DebugGameplayCueFilter.HasTagExact(ExistingTag))
+					{
+						DebugGameplayCueFilter.RemoveTag(ExistingTag);
+					}
+					else
+					{
+						DebugGameplayCueFilter.AddTagFast(ExistingTag);
+					}
+				}
+			})
+	);
+
+}
+
 // --------------------------------------------------------------------------------------------------------------------------------------------------------
 //
 //	UGameplayCueSet
@@ -29,6 +59,16 @@ bool UGameplayCueSet::HandleGameplayCue(AActor* TargetActor, FGameplayTag Gamepl
 #if WITH_SERVER_CODE
 	QUICK_SCOPE_CYCLE_COUNTER(STAT_GameplayCueSet_HandleGameplayCue);
 #endif
+	
+#if !UE_BUILD_SHIPPING
+	if (!GameplayCueDebug::DebugGameplayCueFilter.IsEmpty())
+	{
+		if (!GameplayCueDebug::DebugGameplayCueFilter.HasTagExact(GameplayCueTag))
+		{
+			return false;
+		}
+	}
+#endif
 
 	// GameplayCueTags could have been removed from the dictionary but not content. When the content is resaved the old tag will be cleaned up, but it could still come through here
 	// at runtime. Since we only populate the map with dictionary gameplaycue tags, we may not find it here.
