Modify AbilityTask debugging CVars to enable debug recording and printing to log by default in non-shipping builds (with ability to hotfix on/off as needed).
Can now set AbilitySystem.AbilityTask.Debug.RecordingEnabled to 0 to disable, 1 to enable in non-shipping builds, and 2 to enable all builds (including shipping).
Can use AbilitySystem.AbilityTask.Debug.AbilityTaskDebugPrintTopNResults to only print the top N results in log (hoping to avoid log spam).

Perf: STAT_AbilityTaskDebugRecording can be used to test perf impact on these on-by-default debugging changes.

diff --git a/Engine/Plugins/Runtime/GameplayAbilities/Source/GameplayAbilities/Private/Abilities/Tasks/AbilityTask.cpp b/Engine/Plugins/Runtime/GameplayAbilities/Source/GameplayAbilities/Private/Abilities/Tasks/AbilityTask.cpp
--- a/Engine/Plugins/Runtime/GameplayAbilities/Source/GameplayAbilities/Private/Abilities/Tasks/AbilityTask.cpp
+++ b/Engine/Plugins/Runtime/GameplayAbilities/Source/GameplayAbilities/Private/Abilities/Tasks/AbilityTask.cpp
@@ -4,11 +4,25 @@
 #include "AbilitySystemComponent.h"
 #include "AbilitySystemStats.h"
 
-#if !UE_BUILD_SHIPPING
 static void DebugRecordAbilityTaskCreated(const UAbilityTask* NewTask);
 static void DebugRecordAbilityTaskDestroyed(const UAbilityTask* NewTask);
 static void DebugPrintAbilityTasksByClass();
-#endif // !UE_BUILD_SHIPPING
+
+enum class EDebugBuildRecordFlag : int32
+{
+	Disabled = 0,
+	EnableForNonShippingBuilds = 1,
+	EnableForAllBuilds = 2
+};
+
+namespace AbilityTaskConstants
+{
+#if UE_BUILD_SHIPPING
+	constexpr int32 DebugMinValueToEnableRecording = static_cast<int32>(EDebugBuildRecordFlag::EnableForAllBuilds);
+#else
+	constexpr int32 DebugMinValueToEnableRecording = static_cast<int32>(EDebugBuildRecordFlag::EnableForNonShippingBuilds);
+#endif
+}
 
 namespace AbilityTaskCVars
 {
@@ -19,19 +33,26 @@ namespace AbilityTaskCVars
 		TEXT("Global limit on the number of extant AbilityTasks. Use 'AbilitySystem.Debug.RecordingEnabled' and 'AbilitySystem.AbilityTask.Debug.PrintCounts' to debug why you are hitting this before raising the cap.")
 	);
 
-#if !UE_BUILD_SHIPPING
-	static bool bRecordAbilityTaskCounts = false;
+	// 0 - disabled, 1 - enabled in non-shipping builds, 2 - enabled in all builds (including shipping)
+	static int32 AbilityTaskRecordingType = static_cast<int32>(EDebugBuildRecordFlag::EnableForNonShippingBuilds);
 	static FAutoConsoleVariableRef CVarRecordAbilityTaskCounts(
 		TEXT("AbilitySystem.AbilityTask.Debug.RecordingEnabled"),
-		bRecordAbilityTaskCounts,
-		TEXT("If this is enabled, all new AbilityTasks will be counted by type. Use 'AbilitySystem.AbilityTask.Debug.PrintCounts' to print out the current counts.")
+		AbilityTaskRecordingType,
+		TEXT("Set to 0 to disable, 1 to enable in non-shipping builds, and 2 to enable in all builds (including shipping). If this is enabled, all new AbilityTasks will be counted by type. Use 'AbilitySystem.AbilityTask.Debug.PrintCounts' to print out the current counts.")
 	);
 
-	static bool bRecordAbilityTaskSourceAbilityCounts = false;
+	static bool bRecordAbilityTaskSourceAbilityCounts = true;
 	static FAutoConsoleVariableRef CVarRecordAbilityTaskSourceAbilityCounts(
 		TEXT("AbilitySystem.AbilityTask.Debug.SourceRecordingEnabled"),
 		bRecordAbilityTaskSourceAbilityCounts,
-		TEXT("Requires bRecordAbilityTaskCounts to be set to true for this value to do anything.  If both are enabled, all new AbilityTasks (after InitTask is called in NewAbilityTask) will be counted by the class of the ability that created them.  Use 'AbilitySystem.AbilityTask.Debug.PrintCounts' to print out the current counts.")
+		TEXT("Requires bRecordAbilityTaskCounts to be set to enabled (1 for non-shipping builds, 2 for all builds) for this value to do anything.  If both are enabled, all new AbilityTasks (after InitTask is called in NewAbilityTask) will be counted by the class of the ability that created them.  Use 'AbilitySystem.AbilityTask.Debug.PrintCounts' to print out the current counts.")
+	);
+
+	static int32 AbilityTaskDebugPrintTopNResults = 5;
+	static FAutoConsoleVariableRef CVarAbilityTaskDebugPrintTopNResults(
+		TEXT("AbilitySystem.AbilityTask.Debug.AbilityTaskDebugPrintTopNResults"),
+		AbilityTaskDebugPrintTopNResults,
+		TEXT("Set N to only print the top N results when printing ability task counts (N = 5 by default, if N = 0 prints all).  Use 'AbilitySystem.AbilityTask.Debug.PrintCounts' to print out the current counts.")
 	);
 
 	static FAutoConsoleCommand AbilityTaskPrintAbilityTaskCountsCmd(
@@ -39,7 +60,6 @@ namespace AbilityTaskCVars
 		TEXT("Print out the current AbilityTask counts by class. 'AbilitySystem.AbilityTask.Debug.RecordingEnabled' must be turned on for this to function."),
 		FConsoleCommandDelegate::CreateStatic(DebugPrintAbilityTasksByClass)
 	);
-#endif
 }
 
 static int32 GlobalAbilityTaskCount = 0;
@@ -49,22 +69,23 @@ UAbilityTask::UAbilityTask(const FObjectInitializer& ObjectInitializer)
 {
 	WaitStateBitMask = static_cast<uint8>(EAbilityTaskWaitState::WaitingOnGame);
 
-#if !UE_BUILD_SHIPPING
-	if (AbilityTaskCVars::bRecordAbilityTaskCounts)
+	if (AbilityTaskCVars::AbilityTaskRecordingType >= AbilityTaskConstants::DebugMinValueToEnableRecording)
 	{
 		DebugRecordAbilityTaskCreated(this);
 	}
-#endif  // !UE_BUILD_SHIPPING
+
+	bool bExceededAbilityTaskMaxCount = false;
 
 	++GlobalAbilityTaskCount;
 	SET_DWORD_STAT(STAT_AbilitySystem_TaskCount, GlobalAbilityTaskCount);
-	if (!ensure(GlobalAbilityTaskCount < AbilityTaskCVars::AbilityTaskMaxCount))
+	if (!(GlobalAbilityTaskCount < AbilityTaskCVars::AbilityTaskMaxCount))
 	{
+		bExceededAbilityTaskMaxCount = true;
+
 		ABILITY_LOG(Warning, TEXT("Way too many AbilityTasks are currently active! %d. %s"), GlobalAbilityTaskCount, *GetClass()->GetName());
 
-#if !UE_BUILD_SHIPPING
 		// Auto dump the counts if we hit the limit
-		if (AbilityTaskCVars::bRecordAbilityTaskCounts)
+		if (AbilityTaskCVars::AbilityTaskRecordingType >= AbilityTaskConstants::DebugMinValueToEnableRecording)
 		{
 			static bool bHasDumpedAbilityTasks = false;  // The dump is spammy, so we only want to auto-dump once
 
@@ -72,10 +93,15 @@ UAbilityTask::UAbilityTask(const FObjectInitializer& ObjectInitializer)
 			{
 				DebugPrintAbilityTasksByClass();
 				bHasDumpedAbilityTasks = true;
+
+				// If we don't flush here the ensure is hit without debug ability task info printed in log
+				GLog->FlushThreadedLogs();
+				GLog->Flush();
 			}
 		}
-#endif  // !UE_BUILD_SHIPPING
 	}
+
+	ensureMsgf(!bExceededAbilityTaskMaxCount, TEXT("Exceeded AbilityTaskMaxCount. For more information in log set AbilitySystem.AbilityTask.Debug.SourceRecordingEnabled to 1 for non-shipping builds, or 2 for all builds (including shipping)."));
 }
 
 void UAbilityTask::OnDestroy(bool bInOwnerFinished)
@@ -86,12 +112,10 @@ void UAbilityTask::OnDestroy(bool bInOwnerFinished)
 
 	bWasSuccessfullyDestroyed = true;
 
-#if !UE_BUILD_SHIPPING
-	if (AbilityTaskCVars::bRecordAbilityTaskCounts)
+	if (AbilityTaskCVars::AbilityTaskRecordingType >= AbilityTaskConstants::DebugMinValueToEnableRecording)
 	{
 		DebugRecordAbilityTaskDestroyed(this);
 	}
-#endif  // !UE_BUILD_SHIPPING
 
 	Super::OnDestroy(bInOwnerFinished);
 }
@@ -108,12 +132,10 @@ void UAbilityTask::BeginDestroy()
 		SET_DWORD_STAT(STAT_AbilitySystem_TaskCount, GlobalAbilityTaskCount);
 		bWasSuccessfullyDestroyed = true;
 
-#if !UE_BUILD_SHIPPING
-		if (AbilityTaskCVars::bRecordAbilityTaskCounts)
+		if (AbilityTaskCVars::AbilityTaskRecordingType >= AbilityTaskConstants::DebugMinValueToEnableRecording)
 		{
 			DebugRecordAbilityTaskDestroyed(this);
 		}
-#endif  // !UE_BUILD_SHIPPING
 	}
 }
 
@@ -217,12 +239,13 @@ bool UAbilityTask::IsWaitingOnAvatar() const
 	return (WaitStateBitMask & (uint8)EAbilityTaskWaitState::WaitingOnAvatar) != 0;
 }
 
-#if !UE_BUILD_SHIPPING
 static TMap<const UClass*, int32> StaticAbilityTasksByClass = {};
 static TMap<const UClass*, int32> StaticAbilityTasksByAbilityClass = {};
 
 void DebugRecordAbilityTaskCreated(const UAbilityTask* NewTask)
 {
+	QUICK_SCOPE_CYCLE_COUNTER(STAT_AbilityTaskDebugRecording);
+
 	const UClass* ClassPtr = (NewTask != nullptr) ? NewTask->GetClass() : nullptr;
 	if (ClassPtr != nullptr)
 	{
@@ -239,7 +262,9 @@ void DebugRecordAbilityTaskCreated(const UAbilityTask* NewTask)
 
 void UAbilityTask::DebugRecordAbilityTaskCreatedByAbility(const UObject* Ability)
 {
-	if (!AbilityTaskCVars::bRecordAbilityTaskSourceAbilityCounts || !AbilityTaskCVars::bRecordAbilityTaskCounts)
+	QUICK_SCOPE_CYCLE_COUNTER(STAT_AbilityTaskDebugRecording);
+
+	if (!AbilityTaskCVars::bRecordAbilityTaskSourceAbilityCounts || AbilityTaskCVars::AbilityTaskRecordingType < AbilityTaskConstants::DebugMinValueToEnableRecording)
 	{	// Both the more detailed and the basic recording is required for the detailed recording to work properly.
 		return;
 	}
@@ -260,6 +285,8 @@ void UAbilityTask::DebugRecordAbilityTaskCreatedByAbility(const UObject* Ability
 
 static void DebugRecordAbilityTaskDestroyed(const UAbilityTask* DestroyedTask)
 {
+	QUICK_SCOPE_CYCLE_COUNTER(STAT_AbilityTaskDebugRecording);
+
 	const UClass* ClassPtr = (DestroyedTask != nullptr) ? DestroyedTask->GetClass() : nullptr;
 	if (ClassPtr != nullptr)
 	{
@@ -292,17 +319,31 @@ static void DebugRecordAbilityTaskDestroyed(const UAbilityTask* DestroyedTask)
 	}
 }
 
-static void DebugPrintAbilityTasksByClass()
+void DebugPrintAbilityTasksByClass()
 {
-	if (AbilityTaskCVars::bRecordAbilityTaskCounts)
+	if (AbilityTaskCVars::AbilityTaskRecordingType >= AbilityTaskConstants::DebugMinValueToEnableRecording)
 	{
+		const int32 NumberOfTopResultsToShow = AbilityTaskCVars::AbilityTaskDebugPrintTopNResults;
+		const int32 NumStaticAbilityTasks = StaticAbilityTasksByClass.Num();
+
+		// If NumberOfTopResultsToShow == 0 print all, otherwise print the top N results
+		const int32 LogElementMax = NumberOfTopResultsToShow > 0 ? NumberOfTopResultsToShow : NumStaticAbilityTasks;
+		int32 LogElementCount = 0;
+
 		int32 AccumulatedAbilityTasks = 0;
-		ABILITY_LOG(Display, TEXT("Logging global UAbilityTask counts:"));
+		ABILITY_LOG(Display, TEXT("Logging global UAbilityTask counts (showing top %d results):"), NumberOfTopResultsToShow);
 		StaticAbilityTasksByClass.ValueSort(TGreater<int32>());
 		for (const TPair<const UClass*, int32>& Pair : StaticAbilityTasksByClass)
 		{
-			FString SafeName = GetNameSafe(Pair.Key);
-			ABILITY_LOG(Display, TEXT("- Class '%s': %d"), *SafeName, Pair.Value);
+			// Only log top NumberOfTopResultsToShow entries
+			if (LogElementCount < LogElementMax)
+			{
+				FString SafeName = GetNameSafe(Pair.Key);
+				ABILITY_LOG(Display, TEXT("- Class '%s': %d"), *SafeName, Pair.Value);
+
+				++LogElementCount;
+			}
+			
 			AccumulatedAbilityTasks += Pair.Value;
 		}
 
@@ -316,12 +357,21 @@ static void DebugPrintAbilityTasksByClass()
 
 		if (AbilityTaskCVars::bRecordAbilityTaskSourceAbilityCounts)
 		{
-			ABILITY_LOG(Display, TEXT("UAbilityTask counts per Ability Class:"));
+			LogElementCount = 0;
+			ABILITY_LOG(Display, TEXT("UAbilityTask counts per Ability Class (showing top %d results):"), NumberOfTopResultsToShow);
 			StaticAbilityTasksByAbilityClass.ValueSort(TGreater<int32>());
 			for (const TPair<const UClass*, int32>& Pair : StaticAbilityTasksByAbilityClass)
 			{
+				// Only log top NumberOfTopResultsToShow entries
+				if (LogElementCount >= LogElementMax)
+				{
+					break;
+				}
+
 				FString SafeName = GetNameSafe(Pair.Key);
 				ABILITY_LOG(Display, TEXT("- Ability Class '%s': %d"), *SafeName, Pair.Value);
+
+				++LogElementCount;
 			}
 		}
 
@@ -329,7 +379,6 @@ static void DebugPrintAbilityTasksByClass()
 	}
 	else
 	{
-		ABILITY_LOG(Display, TEXT("Recording of UAbilityTask counts is disabled! Enable 'AbilitySystem.AbilityTask.Debug.RecordingEnabled' to turn on recording."))
+		ABILITY_LOG(Display, TEXT("Recording of UAbilityTask counts is disabled! Enable 'AbilitySystem.AbilityTask.Debug.RecordingEnabled' (1 for non-shipping builds, 2 for all builds) to turn on recording."))
 	}
 }
-#endif  // !UE_BUILD_SHIPPING
diff --git a/Engine/Plugins/Runtime/GameplayAbilities/Source/GameplayAbilities/Public/Abilities/Tasks/AbilityTask.h b/Engine/Plugins/Runtime/GameplayAbilities/Source/GameplayAbilities/Public/Abilities/Tasks/AbilityTask.h
--- a/Engine/Plugins/Runtime/GameplayAbilities/Source/GameplayAbilities/Public/Abilities/Tasks/AbilityTask.h
+++ b/Engine/Plugins/Runtime/GameplayAbilities/Source/GameplayAbilities/Public/Abilities/Tasks/AbilityTask.h
@@ -121,9 +121,7 @@ class GAMEPLAYABILITIES_API UAbilityTask : public UGameplayTask
 
 	virtual void InitSimulatedTask(UGameplayTasksComponent& InGameplayTasksComponent) override;
 
-#if !UE_BUILD_SHIPPING
 	static void DebugRecordAbilityTaskCreatedByAbility(const UObject* Ability);
-#endif  // !UE_BUILD_SHIPPING
 
 	/** Helper function for instantiating and initializing a new task */
 	template <class T>
@@ -134,9 +132,7 @@ class GAMEPLAYABILITIES_API UAbilityTask : public UGameplayTask
 		T* MyObj = NewObject<T>();
 		MyObj->InitTask(*ThisAbility, ThisAbility->GetGameplayTaskDefaultPriority());
 
-#if !UE_BUILD_SHIPPING
 		UAbilityTask::DebugRecordAbilityTaskCreatedByAbility(ThisAbility);
-#endif  // !UE_BUILD_SHIPPING
 
 		MyObj->InstanceName = InstanceName;
 		return MyObj;
