Fix DebugAbilityTags not appearing if command is first run with a filter and not appearing if it was still enabled on begin PIE/level changed

diff --git a/Engine/Plugins/Runtime/GameplayAbilities/Source/GameplayAbilities/Private/AbilitySystemDebugHUD.cpp b/Engine/Plugins/Runtime/GameplayAbilities/Source/GameplayAbilities/Private/AbilitySystemDebugHUD.cpp
--- a/Engine/Plugins/Runtime/GameplayAbilities/Source/GameplayAbilities/Private/AbilitySystemDebugHUD.cpp
+++ b/Engine/Plugins/Runtime/GameplayAbilities/Source/GameplayAbilities/Private/AbilitySystemDebugHUD.cpp
@@ -29,6 +29,8 @@ namespace
 
 bool AAbilitySystemDebugHUD::bEnableBasicHUD = false;
 
+FDelegateHandle AAbilitySystemDebugHUD::DrawDebugDelegateHandle;
+
 TSubclassOf<AAbilitySystemDebugHUD> AAbilitySystemDebugHUD::HUDClass = AAbilitySystemDebugHUD::StaticClass();
 
 void UAbilitySystemDebugHUDExtension_Tags::ToggleExtension(const TArray<FString>& Args, UWorld* World)
@@ -66,7 +68,7 @@ void UAbilitySystemDebugHUDExtension_Tags::ToggleExtension(const TArray<FString>
 
 bool UAbilitySystemDebugHUDExtension_Tags::IsEnabled() const
 {
-	return bEnabled;
+	return bEnabled || !TagsToDisplay.IsEmpty();
 }
 
 void UAbilitySystemDebugHUDExtension_Tags::GetDebugStrings(const AActor* Actor, const UAbilitySystemComponent* Comp, OUT TArray<FString>& OutDebugStrings) const
@@ -232,7 +234,7 @@ void UAbilitySystemDebugHUDExtension_BlockedAbilityTags::ToggleExtension(const T
 
 bool UAbilitySystemDebugHUDExtension_BlockedAbilityTags::IsEnabled() const
 {
-	return bEnabled;
+	return bEnabled || !TagsToDisplay.IsEmpty();
 }
 
 void UAbilitySystemDebugHUDExtension_BlockedAbilityTags::GetDebugStrings(const AActor* Actor, const UAbilitySystemComponent* Comp, OUT TArray<FString>& OutDebugStrings) const
@@ -292,7 +294,7 @@ void AAbilitySystemDebugHUD::NotifyExtensionEnableChanged(UWorld* InWorld)
 		return;
 	}
 
-	AAbilitySystemDebugHUD* HUD = NULL;
+	AAbilitySystemDebugHUD* HUD = nullptr;
 	for (TActorIterator<AAbilitySystemDebugHUD> It(InWorld); It; ++It)
 	{
 		HUD = *It;
@@ -305,43 +307,43 @@ void AAbilitySystemDebugHUD::NotifyExtensionEnableChanged(UWorld* InWorld)
 		bAnyExtensionEnabled |= ExtIt->IsEnabled();
 	}
 
-	static FDelegateHandle DrawDebugDelegateHandle;
-	static FDelegateHandle PostBeginPIEDelegateHandle;
+	static FDelegateHandle PostWorldInitDelegateHandle;
 
 	if (!HUD && bAnyExtensionEnabled)
 	{
-		UClass* HUDClassToSpawn = HUDClass.Get();
-		if (!HUDClassToSpawn)
-		{
-			HUDClassToSpawn = AAbilitySystemDebugHUD::StaticClass();
-		}
-
-		HUD = InWorld->SpawnActor<AAbilitySystemDebugHUD>();
+		CreateHUD(InWorld);
 
-		FDebugDrawDelegate DrawDebugDelegate = FDebugDrawDelegate::CreateUObject(HUD, &AAbilitySystemDebugHUD::DrawDebugHUD);
-		DrawDebugDelegateHandle = UDebugDrawService::Register(TEXT("GameplayDebug"), DrawDebugDelegate);
-
-#if WITH_EDITOR
-		PostBeginPIEDelegateHandle = FEditorDelegates::PostPIEStarted.AddLambda([](const bool bIsSimulating)
+		PostWorldInitDelegateHandle = FWorldDelegates::OnPostWorldInitialization.AddLambda([](UWorld* World, const UWorld::InitializationValues IVS)
 			{
-				if (FWorldContext* PIEWorldContext = GEditor->GetPIEWorldContext())
+				if (World->WorldType == EWorldType::PIE || World->WorldType == EWorldType::Game)
 				{
-					NotifyExtensionEnableChanged(PIEWorldContext->World());
+					CreateHUD(World);
 				}
 			});
-#endif
 
 	}
 	else if (HUD && !bAnyExtensionEnabled)
 	{
+		FEditorDelegates::PostPIEStarted.Remove(PostWorldInitDelegateHandle);
 		UDebugDrawService::Unregister(DrawDebugDelegateHandle);
-#if WITH_EDITOR
-		FEditorDelegates::PostPIEStarted.Remove(PostBeginPIEDelegateHandle);
-#endif
 		HUD->Destroy();
 	}
 }
 
+void AAbilitySystemDebugHUD::CreateHUD(UWorld* World)
+{
+	UClass* HUDClassToSpawn = HUDClass.Get();
+	if (!HUDClassToSpawn)
+	{
+		HUDClassToSpawn = AAbilitySystemDebugHUD::StaticClass();
+	}
+
+	AAbilitySystemDebugHUD* HUD = World->SpawnActor<AAbilitySystemDebugHUD>();
+
+	FDebugDrawDelegate DrawDebugDelegate = FDebugDrawDelegate::CreateUObject(HUD, &AAbilitySystemDebugHUD::DrawDebugHUD);
+	DrawDebugDelegateHandle = UDebugDrawService::Register(TEXT("GameplayDebug"), DrawDebugDelegate);
+}
+
 void AAbilitySystemDebugHUD::DrawWithBackground(UFont* InFont, const FString& Text, const FColor& TextColor, EAlignHorizontal::Type HAlign, float& OffsetX, EAlignVertical::Type VAlign, float& OffsetY, float Alpha)
 {
 	float SizeX, SizeY;
diff --git a/Engine/Plugins/Runtime/GameplayAbilities/Source/GameplayAbilities/Public/AbilitySystemDebugHUD.h b/Engine/Plugins/Runtime/GameplayAbilities/Source/GameplayAbilities/Public/AbilitySystemDebugHUD.h
--- a/Engine/Plugins/Runtime/GameplayAbilities/Source/GameplayAbilities/Public/AbilitySystemDebugHUD.h
+++ b/Engine/Plugins/Runtime/GameplayAbilities/Source/GameplayAbilities/Public/AbilitySystemDebugHUD.h
@@ -137,6 +137,7 @@ public:
 
 	static void ToggleBasicHUD(const TArray<FString>& Args, UWorld* World);
 
+
 	/**
 	 * Notifies the AbilitySystemDebugHUD that an extension has been enabled/disabled
 	 * This will create the HUD if it did not exist and there are any extensions enabled
@@ -158,6 +159,7 @@ protected:
 	virtual bool ShouldDisplayDebugForActor(UCanvas* InCanvas, const AActor* Actor, const FVector& CameraPosition, const FVector& CameraDir) const;
 
 private:
+	static void CreateHUD(UWorld* World);
 
 	void DrawWithBackground(UFont* InFont, const FString& Text, const FColor& TextColor, EAlignHorizontal::Type HAlign, float& OffsetX, EAlignVertical::Type VAlign, float& OffsetY, float Alpha = 1.f);
 
@@ -172,5 +174,7 @@ private:
 	static bool bEnableBasicHUD;
 
 	static TSubclassOf<AAbilitySystemDebugHUD> HUDClass;
+	
+	static FDelegateHandle DrawDebugDelegateHandle;
 
 };
