Fix typo caught by CIS

diff --git a/Engine/Plugins/Runtime/GameplayAbilities/Source/GameplayAbilities/Private/AbilitySystemDebugHUD.cpp b/Engine/Plugins/Runtime/GameplayAbilities/Source/GameplayAbilities/Private/AbilitySystemDebugHUD.cpp
--- a/Engine/Plugins/Runtime/GameplayAbilities/Source/GameplayAbilities/Private/AbilitySystemDebugHUD.cpp
+++ b/Engine/Plugins/Runtime/GameplayAbilities/Source/GameplayAbilities/Private/AbilitySystemDebugHUD.cpp
@@ -324,7 +324,7 @@ void AAbilitySystemDebugHUD::NotifyExtensionEnableChanged(UWorld* InWorld)
 	}
 	else if (HUD && !bAnyExtensionEnabled)
 	{
-		FEditorDelegates::PostPIEStarted.Remove(PostWorldInitDelegateHandle);
+		FWorldDelegates::OnPostWorldInitialization.Remove(PostWorldInitDelegateHandle);
 		UDebugDrawService::Unregister(DrawDebugDelegateHandle);
 		HUD->Destroy();
 	}
