Fix Removal Tag requirment not working for apply cooldown GE
- GE with removal tag requirement specified would be removed unexpectedly
    - the chain of UGameplayAbility::ApplyGameplayEffectToOwner never populate the target tags, so when trying to apply other GE, the cooldown GE would be removed
    - Cooldown GE for GA won't be applied at all, still because the target tags being empty
- Now will populate the target tags with ActorInfo's ASC tags in the two places above

diff --git a/Engine/Plugins/Runtime/GameplayAbilities/Source/GameplayAbilities/Private/Abilities/GameplayAbility.cpp b/Engine/Plugins/Runtime/GameplayAbilities/Source/GameplayAbilities/Private/Abilities/GameplayAbility.cpp
--- a/Engine/Plugins/Runtime/GameplayAbilities/Source/GameplayAbilities/Private/Abilities/GameplayAbility.cpp
+++ b/Engine/Plugins/Runtime/GameplayAbilities/Source/GameplayAbilities/Private/Abilities/GameplayAbility.cpp
@@ -1726,6 +1726,13 @@ FActiveGameplayEffectHandle UGameplayAbility::ApplyGameplayEffectToOwner(const F
 		if (SpecHandle.IsValid())
 		{
 			SpecHandle.Data->StackCount = Stacks;
+			if (const UAbilitySystemComponent* OwnerASC = ActorInfo->AbilitySystemComponent.Get())
+			{
+				FGameplayTagContainer TargetTags;
+				OwnerASC->GetOwnedGameplayTags(TargetTags);
+				SpecHandle.Data->CapturedTargetTags.GetActorTags().AppendTags(TargetTags);
+			}
+
 			return ApplyGameplayEffectSpecToOwner(Handle, ActorInfo, ActivationInfo, SpecHandle);
 		}
 	}
