Add better logging to TryActivateAbility to help indicate why abilities failed

diff --git a/Engine/Plugins/Runtime/GameplayAbilities/Source/GameplayAbilities/Private/Abilities/GameplayAbility.cpp b/Engine/Plugins/Runtime/GameplayAbilities/Source/GameplayAbilities/Private/Abilities/GameplayAbility.cpp
--- a/Engine/Plugins/Runtime/GameplayAbilities/Source/GameplayAbilities/Private/Abilities/GameplayAbility.cpp
+++ b/Engine/Plugins/Runtime/GameplayAbilities/Source/GameplayAbilities/Private/Abilities/GameplayAbility.cpp
@@ -24,6 +24,8 @@ namespace FAbilitySystemTweaks
 	FAutoConsoleVariableRef CVarClearAbilityTimers(TEXT("AbilitySystem.ClearAbilityTimers"), FAbilitySystemTweaks::ClearAbilityTimers, TEXT("Whether to call ClearAllTimersForObject as part of EndAbility call"), ECVF_Default);
 }
 
+int32 FScopedCanActivateAbilityLogEnabler::LogEnablerCounter = 0;
+
 
 UGameplayAbility::UGameplayAbility(const FObjectInitializer& ObjectInitializer)
 	: Super(ObjectInitializer)
@@ -331,29 +333,45 @@ bool UGameplayAbility::CanActivateAbility(const FGameplayAbilitySpecHandle Handl
 
 	if (!AbilitySystemGlobals.ShouldIgnoreCooldowns() && !CheckCooldown(Handle, ActorInfo, OptionalRelevantTags))
 	{
+		if (FScopedCanActivateAbilityLogEnabler::IsLoggingEnabled())
+		{
+			ABILITY_VLOG(ActorInfo->OwnerActor.Get(), Verbose, TEXT("Ability could not be activated due to Cooldown: %s"), *GetName());
+		}
 		return false;
 	}
 
 	if (!AbilitySystemGlobals.ShouldIgnoreCosts() && !CheckCost(Handle, ActorInfo, OptionalRelevantTags))
 	{
+		if (FScopedCanActivateAbilityLogEnabler::IsLoggingEnabled())
+		{
+			ABILITY_VLOG(ActorInfo->OwnerActor.Get(), Verbose, TEXT("Ability could not be activated due to Cost: %s"), *GetName());
+		}
 		return false;
 	}
 
 	if (!DoesAbilitySatisfyTagRequirements(*AbilitySystemComponent, SourceTags, TargetTags, OptionalRelevantTags))
 	{	// If the ability's tags are blocked, or if it has a "Blocking" tag or is missing a "Required" tag, then it can't activate.
+		if (FScopedCanActivateAbilityLogEnabler::IsLoggingEnabled())
+		{
+			ABILITY_VLOG(ActorInfo->OwnerActor.Get(), Verbose, TEXT("Ability could not be activated due to Blocking Tags or Missing Required Tags: %s"), *GetName());
+		}
 		return false;
 	}
 
 	FGameplayAbilitySpec* Spec = AbilitySystemComponent->FindAbilitySpecFromHandle(Handle);
 	if (!Spec)
 	{
-		ABILITY_LOG(Warning, TEXT("CanActivateAbility called with invalid Handle"));
+		ABILITY_LOG(Warning, TEXT("CanActivateAbility %s failed, called with invalid Handle"), *GetName());
 		return false;
 	}
 
 	// Check if this ability's input binding is currently blocked
 	if (AbilitySystemComponent->IsAbilityInputBlocked(Spec->InputID))
 	{
+		if (FScopedCanActivateAbilityLogEnabler::IsLoggingEnabled())
+		{
+			ABILITY_VLOG(ActorInfo->OwnerActor.Get(), Verbose, TEXT("Ability could not be activated due to blocked input ID %i: %s"), Spec->InputID, *GetName());
+		}
 		return false;
 	}
 
diff --git a/Engine/Plugins/Runtime/GameplayAbilities/Source/GameplayAbilities/Private/AbilitySystemComponent_Abilities.cpp b/Engine/Plugins/Runtime/GameplayAbilities/Source/GameplayAbilities/Private/AbilitySystemComponent_Abilities.cpp
--- a/Engine/Plugins/Runtime/GameplayAbilities/Source/GameplayAbilities/Private/AbilitySystemComponent_Abilities.cpp
+++ b/Engine/Plugins/Runtime/GameplayAbilities/Source/GameplayAbilities/Private/AbilitySystemComponent_Abilities.cpp
@@ -1391,6 +1391,7 @@ bool UAbilitySystemComponent::TryActivateAbility(FGameplayAbilitySpecHandle Abil
 	{
 		if (bAllowRemoteActivation)
 		{
+			FScopedCanActivateAbilityLogEnabler LogEnabler;
 			if (Ability->CanActivateAbility(AbilityToActivate, ActorInfo, nullptr, nullptr, &FailureTags))
 			{
 				// No prediction key, server will assign a server-generated key
@@ -1536,6 +1537,7 @@ bool UAbilitySystemComponent::InternalTryActivateAbility(FGameplayAbilitySpecHan
 		// If we have an instanced ability, call CanActivateAbility on it.
 		// Otherwise we always do a non instanced CanActivateAbility check using the CDO of the Ability.
 		UGameplayAbility* const CanActivateAbilitySource = InstancedAbility ? InstancedAbility : Ability;
+		FScopedCanActivateAbilityLogEnabler LogEnabler;
 
 		if (!CanActivateAbilitySource->CanActivateAbility(Handle, ActorInfo, SourceTags, TargetTags, &InternalTryActivateAbilityFailureTags))
 		{
diff --git a/Engine/Plugins/Runtime/GameplayAbilities/Source/GameplayAbilities/Public/Abilities/GameplayAbility.h b/Engine/Plugins/Runtime/GameplayAbilities/Source/GameplayAbilities/Public/Abilities/GameplayAbility.h
--- a/Engine/Plugins/Runtime/GameplayAbilities/Source/GameplayAbilities/Public/Abilities/GameplayAbility.h
+++ b/Engine/Plugins/Runtime/GameplayAbilities/Source/GameplayAbilities/Public/Abilities/GameplayAbility.h
@@ -21,6 +21,19 @@ class UGameplayAbility;
 class UGameplayTask;
 class UGameplayTasksComponent;
 
+struct FScopedCanActivateAbilityLogEnabler
+{
+	FScopedCanActivateAbilityLogEnabler() { ++LogEnablerCounter; }
+
+	~FScopedCanActivateAbilityLogEnabler() { --LogEnablerCounter; }
+
+	static bool IsLoggingEnabled() { return LogEnablerCounter > 0; }
+
+private:
+
+	static int32 LogEnablerCounter;
+};
+
 /**
  * UGameplayAbility
  *	
