Refine TargetTags lacking problem in UAbilitySystemComponent::ApplyGameplayEffectSpecToSelf
    - Previously was simply populating the tags from owner ability system component when UGameplayAbility::BP_ApplyGameplayEffectToOwner is called, now going to check the actually applied effect's target tags instead

diff --git a/Engine/Plugins/Runtime/GameplayAbilities/Source/GameplayAbilities/Private/Abilities/GameplayAbility.cpp b/Engine/Plugins/Runtime/GameplayAbilities/Source/GameplayAbilities/Private/Abilities/GameplayAbility.cpp
--- a/Engine/Plugins/Runtime/GameplayAbilities/Source/GameplayAbilities/Private/Abilities/GameplayAbility.cpp
+++ b/Engine/Plugins/Runtime/GameplayAbilities/Source/GameplayAbilities/Private/Abilities/GameplayAbility.cpp
@@ -1744,13 +1744,6 @@ FActiveGameplayEffectHandle UGameplayAbility::ApplyGameplayEffectToOwner(const F
 		if (SpecHandle.IsValid())
 		{
 			SpecHandle.Data->StackCount = Stacks;
-			if (const UAbilitySystemComponent* OwnerASC = ActorInfo->AbilitySystemComponent.Get())
-			{
-				FGameplayTagContainer TargetTags;
-				OwnerASC->GetOwnedGameplayTags(TargetTags);
-				SpecHandle.Data->CapturedTargetTags.GetActorTags().AppendTags(TargetTags);
-			}
-
 			return ApplyGameplayEffectSpecToOwner(Handle, ActorInfo, ActivationInfo, SpecHandle);
 		}
 	}
diff --git a/Engine/Plugins/Runtime/GameplayAbilities/Source/GameplayAbilities/Private/AbilitySystemComponent.cpp b/Engine/Plugins/Runtime/GameplayAbilities/Source/GameplayAbilities/Private/AbilitySystemComponent.cpp
--- a/Engine/Plugins/Runtime/GameplayAbilities/Source/GameplayAbilities/Private/AbilitySystemComponent.cpp
+++ b/Engine/Plugins/Runtime/GameplayAbilities/Source/GameplayAbilities/Private/AbilitySystemComponent.cpp
@@ -941,7 +941,7 @@ FActiveGameplayEffectHandle UAbilitySystemComponent::ApplyGameplayEffectSpecToSe
 	// evaluate if any active effects need to be removed by the application of this effect
 	if (bIsNetAuthority)
 	{
-		ActiveGameplayEffects.AttemptRemoveActiveEffectsOnEffectApplication(Spec, MyHandle);
+		ActiveGameplayEffects.AttemptRemoveActiveEffectsOnEffectApplication(*OurCopyOfSpec, MyHandle);
 	}
 	
 	// ------------------------------------------------------
