Removing uses of master/slave found in GameplayAbilitiesEditor with base folder Source

diff --git a/Engine/Plugins/Runtime/GameplayAbilities/Source/GameplayAbilitiesEditor/Private/Sequencer/GameplayCueTrackEditor.cpp b/Engine/Plugins/Runtime/GameplayAbilities/Source/GameplayAbilitiesEditor/Private/Sequencer/GameplayCueTrackEditor.cpp
--- a/Engine/Plugins/Runtime/GameplayAbilities/Source/GameplayAbilitiesEditor/Private/Sequencer/GameplayCueTrackEditor.cpp
+++ b/Engine/Plugins/Runtime/GameplayAbilities/Source/GameplayAbilitiesEditor/Private/Sequencer/GameplayCueTrackEditor.cpp
@@ -173,10 +173,10 @@ void FGameplayCueTrackEditor::AddTracks(TRange<FFrameNumber> SectionTickRange, U
 
 	if (InObjectBindingIDs.Num() == 0)
 	{
-		UMovieSceneGameplayCueTrack* NewMasterTrack = FocusedMovieScene->AddMasterTrack<UMovieSceneGameplayCueTrack>();
-		AddSectionToTrack(NewMasterTrack, SectionTickRange, SectionClass);
+		UMovieSceneGameplayCueTrack* NewMainTrack = FocusedMovieScene->AddMasterTrack<UMovieSceneGameplayCueTrack>();
+		AddSectionToTrack(NewMainTrack, SectionTickRange, SectionClass);
 
-		SequencerPtr->OnAddTrack(NewMasterTrack, FGuid());
+		SequencerPtr->OnAddTrack(NewMainTrack, FGuid());
 	}
 	else for (const FGuid& ObjectBindingID : InObjectBindingIDs)
 	{
