Fix it so GameplayEffects that use Data Registries do not warn on load
Add proper verification warnings to Scalable Floats, activated when a context string is passed in
Change IsValid to never warn and add IsValidWithWarnings for explicit verification, and change the default level to match other functions

diff --git a/Engine/Plugins/Runtime/GameplayAbilities/Source/GameplayAbilities/Private/GameplayEffect.cpp b/Engine/Plugins/Runtime/GameplayAbilities/Source/GameplayAbilities/Private/GameplayEffect.cpp
--- a/Engine/Plugins/Runtime/GameplayAbilities/Source/GameplayAbilities/Private/GameplayEffect.cpp
+++ b/Engine/Plugins/Runtime/GameplayAbilities/Source/GameplayAbilities/Private/GameplayEffect.cpp
@@ -34,13 +34,13 @@ const float UGameplayEffect::INVALID_LEVEL = FGameplayEffectConstants::INVALID_L
 DECLARE_CYCLE_STAT(TEXT("MakeQuery"), STAT_MakeGameplayEffectQuery, STATGROUP_AbilitySystem);
 
 #if WITH_EDITOR
-#define GETCURVE_REPORTERROR_WITHPOSTLOAD(Handle) \
-	if (Handle.CurveTable) const_cast<UCurveTable*>(ToRawPtr(Handle.CurveTable))->ConditionalPostLoad(); \
-	GETCURVE_REPORTERROR(Handle);
+#define SCALABLEFLOAT_REPORTERROR_WITHPOSTLOAD(Scalable) \
+	if (Scalable.Curve.CurveTable) const_cast<UCurveTable*>(ToRawPtr(Scalable.Curve.CurveTable))->ConditionalPostLoad(); \
+	SCALABLEFLOAT_REPORTERROR(Scalable);
 
-#define GETCURVE_REPORTERROR_WITHPATHNAME_WITHPOSTLOAD(Handle, PathNameString) \
-	if (Handle.CurveTable) const_cast<UCurveTable*>(ToRawPtr(Handle.CurveTable))->ConditionalPostLoad(); \
-	GETCURVE_REPORTERROR_WITHPATHNAME(Handle, PathNameString);
+#define SCALABLEFLOAT_REPORTERROR_WITHPATHNAME_WITHPOSTLOAD(Scalable, PathNameString) \
+	if (Scalable.Curve.CurveTable) const_cast<UCurveTable*>(ToRawPtr(Scalable.Curve.CurveTable))->ConditionalPostLoad(); \
+	SCALABLEFLOAT_REPORTERROR_WITHPATHNAME(Scalable, PathNameString);
 #endif // WITH_EDITOR
 
 // --------------------------------------------------------------------------------------------------------------------------------------------------------
@@ -103,8 +103,8 @@ void UGameplayEffect::PostLoad()
 	HasRemoveGameplayEffectsQuery = !RemoveGameplayEffectQuery.IsEmpty();
 
 #if WITH_EDITOR
-	GETCURVE_REPORTERROR_WITHPOSTLOAD(Period.Curve);
-	GETCURVE_REPORTERROR_WITHPOSTLOAD(ChanceToApplyToTarget.Curve);
+	SCALABLEFLOAT_REPORTERROR_WITHPOSTLOAD(Period);
+	SCALABLEFLOAT_REPORTERROR_WITHPOSTLOAD(ChanceToApplyToTarget);
 	DurationMagnitude.ReportErrors(GetPathName());
 #endif // WITH_EDITOR
 
@@ -598,19 +598,19 @@ void FGameplayEffectModifierMagnitude::ReportErrors(const FString& PathName) con
 {
 	if (MagnitudeCalculationType == EGameplayEffectMagnitudeCalculation::ScalableFloat)
 	{
-		GETCURVE_REPORTERROR_WITHPATHNAME_WITHPOSTLOAD(ScalableFloatMagnitude.Curve, PathName);
+		SCALABLEFLOAT_REPORTERROR_WITHPATHNAME_WITHPOSTLOAD(ScalableFloatMagnitude, PathName);
 	}
 	else if (MagnitudeCalculationType == EGameplayEffectMagnitudeCalculation::AttributeBased)
 	{
-		GETCURVE_REPORTERROR_WITHPATHNAME_WITHPOSTLOAD(AttributeBasedMagnitude.Coefficient.Curve, PathName);
-		GETCURVE_REPORTERROR_WITHPATHNAME_WITHPOSTLOAD(AttributeBasedMagnitude.PreMultiplyAdditiveValue.Curve, PathName);
-		GETCURVE_REPORTERROR_WITHPATHNAME_WITHPOSTLOAD(AttributeBasedMagnitude.PostMultiplyAdditiveValue.Curve, PathName);
+		SCALABLEFLOAT_REPORTERROR_WITHPATHNAME_WITHPOSTLOAD(AttributeBasedMagnitude.Coefficient, PathName);
+		SCALABLEFLOAT_REPORTERROR_WITHPATHNAME_WITHPOSTLOAD(AttributeBasedMagnitude.PreMultiplyAdditiveValue, PathName);
+		SCALABLEFLOAT_REPORTERROR_WITHPATHNAME_WITHPOSTLOAD(AttributeBasedMagnitude.PostMultiplyAdditiveValue, PathName);
 	}
 	else if (MagnitudeCalculationType == EGameplayEffectMagnitudeCalculation::CustomCalculationClass)
 	{
-		GETCURVE_REPORTERROR_WITHPATHNAME_WITHPOSTLOAD(CustomMagnitude.Coefficient.Curve, PathName);
-		GETCURVE_REPORTERROR_WITHPATHNAME_WITHPOSTLOAD(CustomMagnitude.PreMultiplyAdditiveValue.Curve, PathName);
-		GETCURVE_REPORTERROR_WITHPATHNAME_WITHPOSTLOAD(CustomMagnitude.PostMultiplyAdditiveValue.Curve, PathName);
+		SCALABLEFLOAT_REPORTERROR_WITHPATHNAME_WITHPOSTLOAD(CustomMagnitude.Coefficient, PathName);
+		SCALABLEFLOAT_REPORTERROR_WITHPATHNAME_WITHPOSTLOAD(CustomMagnitude.PreMultiplyAdditiveValue, PathName);
+		SCALABLEFLOAT_REPORTERROR_WITHPATHNAME_WITHPOSTLOAD(CustomMagnitude.PostMultiplyAdditiveValue, PathName);
 	}
 }
 #endif // WITH_EDITOR
diff --git a/Engine/Plugins/Runtime/GameplayAbilities/Source/GameplayAbilities/Private/ScalableFloat.cpp b/Engine/Plugins/Runtime/GameplayAbilities/Source/GameplayAbilities/Private/ScalableFloat.cpp
--- a/Engine/Plugins/Runtime/GameplayAbilities/Source/GameplayAbilities/Private/ScalableFloat.cpp
+++ b/Engine/Plugins/Runtime/GameplayAbilities/Source/GameplayAbilities/Private/ScalableFloat.cpp
@@ -15,7 +15,7 @@
 #include "EditorReimportHandler.h"
 #endif
 
-bool FScalableFloat::EvaluateCurveAtLevel(float& OutValue, const FRealCurve*& OutCurve, float Level, const FString& ContextString) const
+bool FScalableFloat::EvaluateCurveAtLevel(float& OutValue, const FRealCurve*& OutCurve, float Level, const FString& ContextString, bool bWarnIfInvalid) const
 {
 	if (!Curve.RowName.IsNone())
 	{
@@ -33,7 +33,7 @@ bool FScalableFloat::EvaluateCurveAtLevel(float& OutValue, const FRealCurve*& Ou
 			// Cache not valid, look at sources
 			if (Curve.CurveTable)
 			{
-				CachedCurve = OutCurve = Curve.GetCurve(ContextString);
+				CachedCurve = OutCurve = Curve.GetCurve(ContextString, bWarnIfInvalid);
 				LocalCachedCurveID = GlobalCachedCurveID;
 			}
 			else if (RegistryType.IsValid())
@@ -70,10 +70,36 @@ bool FScalableFloat::EvaluateCurveAtLevel(float& OutValue, const FRealCurve*& Ou
 		else
 		{
 			// Has row name but no curve, this is an error but fallback to raw value
+			if (bWarnIfInvalid && !Curve.CurveTable)
+			{
+				// CurveTable case is handled in GetCurve above
+				if (RegistryType.IsValid())
+				{
+					// This can happen if the data registry hasn't loaded yet so don't warn by default
+					ABILITY_LOG(Verbose, TEXT("FScalableFloat could not find curve for DataRegistryId %s:%s (%s)."), *RegistryType.ToString(), *Curve.RowName.ToString(), *ContextString);
+				}
+				else
+				{
+					ABILITY_LOG(Warning, TEXT("FScalableFloat has no CurveTable or DataRegistry for row %s (%s)."), *Curve.RowName.ToString(), *ContextString);
+				}
+			}
+
 			OutValue = Value;
 			return false;
 		}
 	}
+	else if (bWarnIfInvalid)
+	{
+		// These cases are errors, but fallback to no curve case for compatibility
+		if (RegistryType.IsValid())
+		{
+			ABILITY_LOG(Warning, TEXT("FScalableFloat has no row for DataRegistry %s (%s)."), *RegistryType.ToString(), *ContextString);
+		}
+		else if (Curve.CurveTable)
+		{
+			ABILITY_LOG(Warning, TEXT("FScalableFloat has no row for CurveTable %s (%s)."), *Curve.CurveTable->GetPathName(), *ContextString);
+		}
+	}
 
 	// No curve, use raw value
 	OutCurve = nullptr;
@@ -87,7 +113,8 @@ float FScalableFloat::GetValueAtLevel(float Level, const FString* ContextString)
 	const FRealCurve* FoundCurve;
 	static const FString DefaultContextString = TEXT("FScalableFloat::GetValueAtLevel");
 
-	EvaluateCurveAtLevel(OutFloat, FoundCurve, Level, ContextString ? *ContextString : DefaultContextString);
+	// Only print warnings if we have a useful context string
+	EvaluateCurveAtLevel(OutFloat, FoundCurve, Level, ContextString != nullptr ? *ContextString : DefaultContextString, ContextString != nullptr);
 
 	return OutFloat;
 }
@@ -142,9 +169,17 @@ bool FScalableFloat::IsValid() const
 {
 	float OutFloat;
 	const FRealCurve* FoundCurve;
-	static const FString DefaultContextString = TEXT("FScalableFloat::GetValueAtLevel");
+	static const FString DefaultContextString = TEXT("FScalableFloat::IsValid");
+
+	return EvaluateCurveAtLevel(OutFloat, FoundCurve, 0.f, DefaultContextString, false);
+}
+
+bool FScalableFloat::IsValidWithWarnings(const FString& ContextString) const
+{
+	float OutFloat;
+	const FRealCurve* FoundCurve;
 
-	return EvaluateCurveAtLevel(OutFloat, FoundCurve, 1.f, DefaultContextString);
+	return EvaluateCurveAtLevel(OutFloat, FoundCurve, 0.f, ContextString, true);
 }
 
 bool FScalableFloat::SerializeFromMismatchedTag(const FPropertyTag& Tag, FStructuredArchive::FSlot Slot)
diff --git a/Engine/Plugins/Runtime/GameplayAbilities/Source/GameplayAbilities/Public/ScalableFloat.h b/Engine/Plugins/Runtime/GameplayAbilities/Source/GameplayAbilities/Public/ScalableFloat.h
--- a/Engine/Plugins/Runtime/GameplayAbilities/Source/GameplayAbilities/Public/ScalableFloat.h
+++ b/Engine/Plugins/Runtime/GameplayAbilities/Source/GameplayAbilities/Public/ScalableFloat.h
@@ -61,8 +61,8 @@ public:
 	/** Returns the value as an int32 */
 	int32 AsInteger(float Level = 0, const FString* ContextString = nullptr) const;
 
-	/** Gets the value and possible curve at a given level, returning false if any errors occurred */
-	bool EvaluateCurveAtLevel(float& OutValue, const FRealCurve*& OutCurve, float Level, const FString& ContextString) const;
+	/** Gets the value and possible curve at a given level, returning false if it failed to find a good value */
+	bool EvaluateCurveAtLevel(float& OutValue, const FRealCurve*& OutCurve, float Level, const FString& ContextString, bool bWarnIfInvalid = true) const;
 
 	/** True if there is no curve lookup */
 	bool IsStatic() const
@@ -86,9 +86,12 @@ public:
 	/** Outputs human readable string */
 	FString ToSimpleString() const;
 
-	/** Error checking: checks if we have a curve source specified but no valid curve entry */
+	/** Error checking: Returns false if this has an invalid curve reference but will not print warnings */
 	bool IsValid() const;
 
+	/** Error checking: Checks if float is valid and prints detailed warnings if not valid */
+	bool IsValidWithWarnings(const FString& ContextString) const;
+
 	/** Equality/Inequality operators */
 	bool operator==(const FScalableFloat& Other) const;
 	bool operator!=(const FScalableFloat& Other) const;
@@ -115,3 +118,10 @@ struct TStructOpsTypeTraits<FScalableFloat>
 	};
 };
 
+/** Macro to call IsValidWithWarnings with a correct error info. Assumed to be called within a UObject */
+#define SCALABLEFLOAT_REPORTERROR(Scalable)\
+	Scalable.IsValidWithWarnings(FString::Printf(TEXT("%s.%s"), *GetPathName(), TEXT(#Scalable)))\
+
+/** Macro to call IsValidWithWarnings with a correct error info */
+#define SCALABLEFLOAT_REPORTERROR_WITHPATHNAME(Scalable, PathNameString)\
+	Scalable.IsValidWithWarnings(FString::Printf(TEXT("%s.%s"), *PathNameString, TEXT(#Scalable)))
\ No newline at end of file
