Don't trigger replicated gameplaycues from gameplayeffects that no longer exist if the gameplayeffect no longer exists

diff --git a/Engine/Plugins/Runtime/GameplayAbilities/Source/GameplayAbilities/Private/AbilitySystemComponent.cpp b/Engine/Plugins/Runtime/GameplayAbilities/Source/GameplayAbilities/Private/AbilitySystemComponent.cpp
--- a/Engine/Plugins/Runtime/GameplayAbilities/Source/GameplayAbilities/Private/AbilitySystemComponent.cpp
+++ b/Engine/Plugins/Runtime/GameplayAbilities/Source/GameplayAbilities/Private/AbilitySystemComponent.cpp
@@ -1217,6 +1217,10 @@ void UAbilitySystemComponent::InvokeGameplayCueEvent(const FGameplayEffectSpecFo
 
 	FGameplayCueParameters CueParameters(Spec);
 
+	FGameplayEffectQuery EffectQuery;
+	EffectQuery.EffectDefinition = Spec.Def->GetClass();
+	CueParameters.bGameplayEffectActive = Spec.Def->DurationPolicy == EGameplayEffectDurationType::Instant || ActiveGameplayEffects.GetActiveEffectCount(EffectQuery) > 0;
+
 	for (FGameplayEffectCue CueInfo : Spec.Def->GameplayCues)
 	{
 		if (CueInfo.MagnitudeAttribute.IsValid())
diff --git a/Engine/Plugins/Runtime/GameplayAbilities/Source/GameplayAbilities/Private/GameplayCueSet.cpp b/Engine/Plugins/Runtime/GameplayAbilities/Source/GameplayAbilities/Private/GameplayCueSet.cpp
--- a/Engine/Plugins/Runtime/GameplayAbilities/Source/GameplayAbilities/Private/GameplayCueSet.cpp
+++ b/Engine/Plugins/Runtime/GameplayAbilities/Source/GameplayAbilities/Private/GameplayCueSet.cpp
@@ -305,6 +305,12 @@ bool UGameplayCueSet::HandleGameplayCueNotify_Internal(AActor* TargetActor, int3
 		}
 		else if (AGameplayCueNotify_Actor* InstancedCue = Cast<AGameplayCueNotify_Actor>(CueData.LoadedGameplayCueClass->ClassDefaultObject))
 		{
+			bool bShouldDestroy = false;
+			if (EventType == EGameplayCueEvent::Executed && !Parameters.bGameplayEffectActive && InstancedCue->bAutoDestroyOnRemove)
+			{
+				bShouldDestroy = true;
+			}
+
 			if (InstancedCue->HandlesEvent(EventType))
 			{
 				if (TargetActor)
@@ -319,6 +325,11 @@ bool UGameplayCueSet::HandleGameplayCueNotify_Internal(AActor* TargetActor, int3
 						{
 							HandleGameplayCueNotify_Internal(TargetActor, CueData.ParentDataIdx, EventType, Parameters);
 						}
+
+						if (bShouldDestroy)
+						{
+							SpawnedInstancedCue->HandleGameplayCue(TargetActor, EGameplayCueEvent::Removed, Parameters);
+						}
 					}
 				}
 			}
diff --git a/Engine/Plugins/Runtime/GameplayAbilities/Source/GameplayAbilities/Public/GameplayEffectTypes.h b/Engine/Plugins/Runtime/GameplayAbilities/Source/GameplayAbilities/Public/GameplayEffectTypes.h
--- a/Engine/Plugins/Runtime/GameplayAbilities/Source/GameplayAbilities/Public/GameplayEffectTypes.h
+++ b/Engine/Plugins/Runtime/GameplayAbilities/Source/GameplayAbilities/Public/GameplayEffectTypes.h
@@ -1009,6 +1009,9 @@ struct GAMEPLAYABILITIES_API FGameplayCueParameters
 	UPROPERTY(BlueprintReadWrite, Category = GameplayCue)
 	bool bReplicateLocationWhenUsingMinimalRepProxy = false;
 
+	/** If originating from a GameplayEffect, whether that GameplayEffect is still Active */
+	bool bGameplayEffectActive = true;
+
 	/** Optimized serializer */
 	bool NetSerialize(FArchive& Ar, class UPackageMap* Map, bool& bOutSuccess);
 
