Remove AbilitySystem.Debug.NextTarget and PrevTarget commands in favor of global HUD NextDebugTarget and PrevDebugTarget commands

diff --git a/Engine/Plugins/Runtime/GameplayAbilities/Source/GameplayAbilities/Private/AbilitySystemComponent.cpp b/Engine/Plugins/Runtime/GameplayAbilities/Source/GameplayAbilities/Private/AbilitySystemComponent.cpp
--- a/Engine/Plugins/Runtime/GameplayAbilities/Source/GameplayAbilities/Private/AbilitySystemComponent.cpp
+++ b/Engine/Plugins/Runtime/GameplayAbilities/Source/GameplayAbilities/Private/AbilitySystemComponent.cpp
@@ -2134,61 +2134,6 @@ UAbilitySystemComponent* GetDebugTarget(FASCDebugTargetInfo* Info)
 	return Info->LastDebugTarget.Get();
 }
 
-void CycleDebugTarget(FASCDebugTargetInfo* TargetInfo, bool Next)
-{
-	GetDebugTarget(TargetInfo);
-
-	// Build a list	of ASCs
-	TArray<UAbilitySystemComponent*> List;
-	for (TObjectIterator<UAbilitySystemComponent> It; It; ++It)
-	{
-		if (UAbilitySystemComponent* ASC = *It)
-		{
-			if (ASC->GetWorld() == TargetInfo->TargetWorld.Get())
-			{
-				List.Add(ASC);
-			}
-		}
-	}
-
-	// Search through list to find prev/next target
-	UAbilitySystemComponent* Previous = nullptr;
-	for (int32 idx=0; idx < List.Num() + 1; ++idx)
-	{
-		UAbilitySystemComponent* ASC = List[idx % List.Num()];
-
-		if (Next && Previous == TargetInfo->LastDebugTarget.Get())
-		{
-			TargetInfo->LastDebugTarget = ASC;
-			return;
-		}
-		if (!Next && ASC == TargetInfo->LastDebugTarget.Get())
-		{
-			TargetInfo->LastDebugTarget = Previous;
-			return;
-		}
-
-		Previous = ASC;
-	}
-}
-
-static void	AbilitySystemCycleDebugTarget(UWorld* InWorld, bool Next)
-{
-	CycleDebugTarget( GetDebugTargetInfo(InWorld), Next );
-}
-
-FAutoConsoleCommandWithWorld AbilitySystemNextDebugTargetCmd(
-	TEXT("AbilitySystem.Debug.NextTarget"),
-	TEXT("Targets next AbilitySystemComponent in ShowDebug AbilitySystem"),
-	FConsoleCommandWithWorldDelegate::CreateStatic(AbilitySystemCycleDebugTarget, true)
-	);
-
-FAutoConsoleCommandWithWorld AbilitySystemPrevDebugTargetCmd(
-	TEXT("AbilitySystem.Debug.PrevTarget"),
-	TEXT("Targets previous AbilitySystemComponent in ShowDebug AbilitySystem"),
-	FConsoleCommandWithWorldDelegate::CreateStatic(AbilitySystemCycleDebugTarget, false)
-	);
-
 FAutoConsoleCommandWithWorld AbilitySystemDebugNextCategoryCmd(
 	TEXT("AbilitySystem.Debug.NextCategory"),
 	TEXT("Switches to the next ShowDebug AbilitySystem category"),
