[KillPendingKill] Nulling out ability reference inside AbilityTask's OnDestroy so it doesn't get garbage collected

diff --git a/Engine/Plugins/Runtime/GameplayAbilities/Source/GameplayAbilities/Private/Abilities/Tasks/AbilityTask.cpp b/Engine/Plugins/Runtime/GameplayAbilities/Source/GameplayAbilities/Private/Abilities/Tasks/AbilityTask.cpp
--- a/Engine/Plugins/Runtime/GameplayAbilities/Source/GameplayAbilities/Private/Abilities/Tasks/AbilityTask.cpp
+++ b/Engine/Plugins/Runtime/GameplayAbilities/Source/GameplayAbilities/Private/Abilities/Tasks/AbilityTask.cpp
@@ -117,6 +117,9 @@ void UAbilityTask::OnDestroy(bool bInOwnerFinished)
 		DebugRecordAbilityTaskDestroyed(this);
 	}
 
+	// #KillPendingKill Clear ability reference so we don't hold onto it and GC can delete it.
+	Ability = nullptr;
+
 	Super::OnDestroy(bInOwnerFinished);
 }
 
