Add useful comments explaining gameplay cue active/execute enum and functions

diff --git a/Engine/Plugins/Runtime/GameplayAbilities/Source/GameplayAbilities/Public/GameplayCueNotify_Actor.h b/Engine/Plugins/Runtime/GameplayAbilities/Source/GameplayAbilities/Public/GameplayCueNotify_Actor.h
--- a/Engine/Plugins/Runtime/GameplayAbilities/Source/GameplayAbilities/Public/GameplayCueNotify_Actor.h
+++ b/Engine/Plugins/Runtime/GameplayAbilities/Source/GameplayAbilities/Public/GameplayCueNotify_Actor.h
@@ -81,15 +81,19 @@ public:
 	UFUNCTION(BlueprintImplementableEvent, Category = "GameplayCueNotify", DisplayName = "HandleGameplayCue", meta=(ScriptName = "HandleGameplayCue"))
 	void K2_HandleGameplayCue(AActor* MyTarget, EGameplayCueEvent::Type EventType, const FGameplayCueParameters& Parameters);
 
+	/** Called when a GameplayCue is executed, this is used for instant effects or periodic ticks */
 	UFUNCTION(BlueprintNativeEvent, Category = "GameplayCueNotify")
 	bool OnExecute(AActor* MyTarget, const FGameplayCueParameters& Parameters);
 
+	/** Called when a GameplayCue with duration is first activated, this will only be called if the client witnessed the activation */
 	UFUNCTION(BlueprintNativeEvent, Category = "GameplayCueNotify")
 	bool OnActive(AActor* MyTarget, const FGameplayCueParameters& Parameters);
 
+	/** Called when a GameplayCue with duration is first seen as active, even if it wasn't actually just applied (Join in progress, etc) */
 	UFUNCTION(BlueprintNativeEvent, Category = "GameplayCueNotify")
 	bool WhileActive(AActor* MyTarget, const FGameplayCueParameters& Parameters);
 
+	/** Called when a GameplayCue with duration is removed */
 	UFUNCTION(BlueprintNativeEvent, Category = "GameplayCueNotify")
 	bool OnRemove(AActor* MyTarget, const FGameplayCueParameters& Parameters);
 
diff --git a/Engine/Plugins/Runtime/GameplayAbilities/Source/GameplayAbilities/Public/GameplayCueNotify_Static.h b/Engine/Plugins/Runtime/GameplayAbilities/Source/GameplayAbilities/Public/GameplayCueNotify_Static.h
--- a/Engine/Plugins/Runtime/GameplayAbilities/Source/GameplayAbilities/Public/GameplayCueNotify_Static.h
+++ b/Engine/Plugins/Runtime/GameplayAbilities/Source/GameplayAbilities/Public/GameplayCueNotify_Static.h
@@ -39,15 +39,19 @@ class GAMEPLAYABILITIES_API UGameplayCueNotify_Static : public UObject
 	UFUNCTION(BlueprintImplementableEvent, Category = "GameplayCueNotify", DisplayName = "HandleGameplayCue", meta=(ScriptName = "HandleGameplayCue"))
 	void K2_HandleGameplayCue(AActor* MyTarget, EGameplayCueEvent::Type EventType, const FGameplayCueParameters& Parameters) const;
 
+	/** Called when a GameplayCue is executed, this is used for instant effects or periodic ticks */
 	UFUNCTION(BlueprintNativeEvent, BlueprintPure, Category = "GameplayCueNotify")
 	bool OnExecute(AActor* MyTarget, const FGameplayCueParameters& Parameters) const;
 
+	/** Called when a GameplayCue with duration is first activated, this will only be called if the client witnessed the activation */
 	UFUNCTION(BlueprintNativeEvent, BlueprintPure, Category = "GameplayCueNotify")
 	bool OnActive(AActor* MyTarget, const FGameplayCueParameters& Parameters) const;
 
+	/** Called when a GameplayCue with duration is first seen as active, even if it wasn't actually just applied (Join in progress, etc) */
 	UFUNCTION(BlueprintNativeEvent, BlueprintPure, Category = "GameplayCueNotify")
 	bool WhileActive(AActor* MyTarget, const FGameplayCueParameters& Parameters) const;
 
+	/** Called when a GameplayCue with duration is removed */
 	UFUNCTION(BlueprintNativeEvent, BlueprintPure, Category = "GameplayCueNotify")
 	bool OnRemove(AActor* MyTarget, const FGameplayCueParameters& Parameters) const;
 
diff --git a/Engine/Plugins/Runtime/GameplayAbilities/Source/GameplayAbilities/Public/GameplayEffectTypes.h b/Engine/Plugins/Runtime/GameplayAbilities/Source/GameplayAbilities/Public/GameplayEffectTypes.h
--- a/Engine/Plugins/Runtime/GameplayAbilities/Source/GameplayAbilities/Public/GameplayEffectTypes.h
+++ b/Engine/Plugins/Runtime/GameplayAbilities/Source/GameplayAbilities/Public/GameplayEffectTypes.h
@@ -1047,16 +1047,16 @@ namespace EGameplayCueEvent
 	/** Indicates what type of action happened to a specific gameplay cue tag. Sometimes you will get multiple events at once */
 	enum Type
 	{
-		/** Called when GameplayCue is activated */
+		/** Called when a GameplayCue with duration is first activated, this will only be called if the client witnessed the activation */
 		OnActive,
 
-		/** Called when GameplayCue is active, even if it wasn't actually just applied (Join in progress, etc) */
+		/** Called when a GameplayCue with duration is first seen as active, even if it wasn't actually just applied (Join in progress, etc) */
 		WhileActive,
 
-		/** Called when a GameplayCue is executed: instant effects or periodic tick */
+		/** Called when a GameplayCue is executed, this is used for instant effects or periodic ticks */
 		Executed,
 
-		/** Called when GameplayCue is removed */
+		/** Called when a GameplayCue with duration is removed */
 		Removed
 	};
 }
