Fix StartAbilityState task to correctly call delegates when it is destroyed now that UAbilityTask::OnDestroy nulls the ability pointer

diff --git a/Engine/Plugins/Runtime/GameplayAbilities/Source/GameplayAbilities/Private/Abilities/Tasks/AbilityTask_StartAbilityState.cpp b/Engine/Plugins/Runtime/GameplayAbilities/Source/GameplayAbilities/Private/Abilities/Tasks/AbilityTask_StartAbilityState.cpp
--- a/Engine/Plugins/Runtime/GameplayAbilities/Source/GameplayAbilities/Private/Abilities/Tasks/AbilityTask_StartAbilityState.cpp
+++ b/Engine/Plugins/Runtime/GameplayAbilities/Source/GameplayAbilities/Private/Abilities/Tasks/AbilityTask_StartAbilityState.cpp
@@ -33,7 +33,12 @@ void UAbilityTask_StartAbilityState::Activate()
 
 void UAbilityTask_StartAbilityState::OnDestroy(bool AbilityEnded)
 {
-	Super::OnDestroy(AbilityEnded);
+	// Unbind delegates so this doesn't get recursively called
+	if (Ability)
+	{
+		Ability->OnGameplayAbilityCancelled.Remove(InterruptStateHandle);
+		Ability->OnGameplayAbilityStateEnded.Remove(EndStateHandle);
+	}
 
 	if (bWasInterrupted && OnStateInterrupted.IsBound())
 	{
@@ -50,11 +55,8 @@ void UAbilityTask_StartAbilityState::OnDestroy(bool AbilityEnded)
 		}
 	}
 
-	if (Ability)
-	{
-		Ability->OnGameplayAbilityCancelled.Remove(InterruptStateHandle);
-		Ability->OnGameplayAbilityStateEnded.Remove(EndStateHandle);
-	}
+	// This will invalidate the task so needs to happen after callbacks
+	Super::OnDestroy(AbilityEnded);
 }
 
 void UAbilityTask_StartAbilityState::OnEndState(FName StateNameToEnd)
