Suppressing the GameplayCue folder requirement warning

diff --git a/Engine/Plugins/Runtime/GameplayAbilities/Source/GameplayAbilities/Private/GameplayCueManager.cpp b/Engine/Plugins/Runtime/GameplayAbilities/Source/GameplayAbilities/Private/GameplayCueManager.cpp
--- a/Engine/Plugins/Runtime/GameplayAbilities/Source/GameplayAbilities/Private/GameplayCueManager.cpp
+++ b/Engine/Plugins/Runtime/GameplayAbilities/Source/GameplayAbilities/Private/GameplayCueManager.cpp
@@ -1034,7 +1034,7 @@ void UGameplayCueManager::HandleAssetAdded(UObject *Object)
 		
 		if (StaticCDO || ActorCDO)
 		{
-			if (VerifyNotifyAssetIsInValidPath(Blueprint->GetOuter()->GetPathName()))
+			if (!Blueprint->GetOutermost()->HasAnyPackageFlags(PKG_ForDiffing) && VerifyNotifyAssetIsInValidPath(Blueprint->GetOuter()->GetPathName()))
 			{
 				FSoftObjectPath StringRef;
 				StringRef.SetPath(Blueprint->GeneratedClass->GetPathName());
