GameplayAbility
* Removed the warning when calling an RPC on non-registered subobjects.

diff --git a/Engine/Plugins/Runtime/GameplayAbilities/Source/GameplayAbilities/Private/Abilities/GameplayAbility.cpp b/Engine/Plugins/Runtime/GameplayAbilities/Source/GameplayAbilities/Private/Abilities/GameplayAbility.cpp
--- a/Engine/Plugins/Runtime/GameplayAbilities/Source/GameplayAbilities/Private/Abilities/GameplayAbility.cpp
+++ b/Engine/Plugins/Runtime/GameplayAbilities/Source/GameplayAbilities/Private/Abilities/GameplayAbility.cpp
@@ -109,15 +109,6 @@ bool UGameplayAbility::CallRemoteFunction(UFunction* Function, void* Parameters,
 	FWorldContext* const Context = GEngine->GetWorldContextFromWorld(GetWorld());
 	if (Context != nullptr)
 	{
-#if DO_CHECK
-		if (UAbilitySystemComponent* AbilitySystemComponent = GetAbilitySystemComponentFromActorInfo())
-		{
-			// If the component is using the subobject list, detect subobjects that send RPC's before they are registered.
-			const bool bIsSubObjectRegistered = !AbilitySystemComponent->IsUsingRegisteredSubObjectList() || Owner->IsActorComponentReplicatedSubObjectRegistered(AbilitySystemComponent, this);
-			UE_CLOG(!bIsSubObjectRegistered, LogAbilitySystem, Warning, TEXT("%s sent an RPC before it was registered as a subobject on %s::%s"), *GetName(), *Owner->GetName(), *AbilitySystemComponent->GetName());
-		}
-#endif
-		
 		for (FNamedNetDriver& Driver : Context->ActiveNetDrivers)
 		{
 			if (Driver.NetDriver != nullptr && Driver.NetDriver->ShouldReplicateFunction(Owner, Function))
