Fixes #jira UE-166117: GetOwnedGameplayTags blueprint output pin retains values between calls.

diff --git a/Engine/Plugins/Runtime/GameplayAbilities/Source/GameplayAbilities/Public/AbilitySystemComponent.h b/Engine/Plugins/Runtime/GameplayAbilities/Source/GameplayAbilities/Public/AbilitySystemComponent.h
--- a/Engine/Plugins/Runtime/GameplayAbilities/Source/GameplayAbilities/Public/AbilitySystemComponent.h
+++ b/Engine/Plugins/Runtime/GameplayAbilities/Source/GameplayAbilities/Public/AbilitySystemComponent.h
@@ -546,6 +546,7 @@ class GAMEPLAYABILITIES_API UAbilitySystemComponent : public UGameplayTasksCompo
 
 	FORCEINLINE void GetOwnedGameplayTags(FGameplayTagContainer& TagContainer) const override
 	{
+		TagContainer.Reset();
 		TagContainer.AppendTags(GameplayTagCountContainer.GetExplicitGameplayTags());
 	}
 
