Update the editor icons for the gameplay cue editor to be in the UE5 styling

#rnx

diff --git a/Engine/Plugins/Runtime/GameplayAbilities/Source/GameplayAbilitiesEditor/Private/GameplayAbilitiesEditorModule.cpp b/Engine/Plugins/Runtime/GameplayAbilities/Source/GameplayAbilitiesEditor/Private/GameplayAbilitiesEditorModule.cpp
--- a/Engine/Plugins/Runtime/GameplayAbilities/Source/GameplayAbilitiesEditor/Private/GameplayAbilitiesEditorModule.cpp
+++ b/Engine/Plugins/Runtime/GameplayAbilities/Source/GameplayAbilitiesEditor/Private/GameplayAbilitiesEditorModule.cpp
@@ -10,6 +10,9 @@
 #include "Styling/AppStyle.h"
 #include "HAL/IConsoleManager.h"
 #include "Framework/Application/SlateApplication.h"
+#include "Styling/SlateStyle.h"
+#include "Styling/SlateStyleMacros.h"
+#include "Styling/SlateStyleRegistry.h"
 
 #include "PropertyEditorModule.h"
 #include "AttributeDetails.h"
@@ -127,6 +130,42 @@ public:
 
 };
 
+/** The style set for the Gameplay Abilities plugin. Any custom icons or color themes can go in here. */
+class FGameplayAbilitiesEditorStyleSet final : public FSlateStyleSet
+{
+public:
+
+	FGameplayAbilitiesEditorStyleSet()
+		: FSlateStyleSet("GameplayAbilitiesEditor")
+	{
+		SetParentStyleName(FAppStyle::GetAppStyleSetName());
+
+		// The icons are located in \Engine\Plugins\Runtime\GameplayAbilities\Content\Editor\Slate\Icons
+		SetContentRoot(FPaths::EnginePluginsDir() / TEXT("Runtime/GameplayAbilities/Content/Editor/Slate"));
+		SetCoreContentRoot(FPaths::EngineContentDir() / TEXT("Slate"));
+
+		// Gameplay Abilities Editor icons
+		static const FVector2D Icon16 = FVector2D(16.0, 16.0);
+		static const FVector2D Icon64 = FVector2D(64.0, 64.0);
+
+		Set("GameplayCueEditor_Small", new IMAGE_BRUSH_SVG("Icons/GameplayCueEditor_16", Icon16));
+		Set("GameplayCueEditor_Large", new IMAGE_BRUSH_SVG("Icons/GameplayCueEditor_64", Icon64));
+
+		FSlateStyleRegistry::RegisterSlateStyle(*this);
+	}
+
+	~FGameplayAbilitiesEditorStyleSet()
+	{
+		FSlateStyleRegistry::UnRegisterSlateStyle(*this);
+	}
+
+	static FGameplayAbilitiesEditorStyleSet& Get()
+	{
+		static FGameplayAbilitiesEditorStyleSet Inst;
+		return Inst;
+	}
+};
+
 IMPLEMENT_MODULE(FGameplayAbilitiesEditorModule, GameplayAbilitiesEditor)
 
 void FGameplayAbilitiesEditorModule::StartupModule()
@@ -169,14 +208,13 @@ void FGameplayAbilitiesEditorModule::StartupModule()
 	// Listen for changes to the gameplay tag tree so we can refresh blueprint actions for the GameplayCueEvent node
 	UGameplayTagsManager& GameplayTagsManager = UGameplayTagsManager::Get();
 	GameplayTagTreeChangedDelegateHandle = IGameplayTagsModule::OnGameplayTagTreeChanged.AddStatic(&FGameplayAbilitiesEditorModule::GameplayTagTreeChanged);
-
+	
 	// GameplayCue editor
 	FGlobalTabmanager::Get()->RegisterNomadTabSpawner( FName(TEXT("GameplayCueApp")), FOnSpawnTab::CreateRaw(this, &FGameplayAbilitiesEditorModule::SpawnGameplayCueEditorTab))
 		.SetDisplayName(NSLOCTEXT("GameplayAbilitiesEditorModule", "GameplayCueTabTitle", "GameplayCue Editor"))
 		.SetTooltipText(NSLOCTEXT("GameplayAbilitiesEditorModule", "GameplayCueTooltipText", "Open GameplayCue Editor tab."))
 		.SetGroup(WorkspaceMenu::GetMenuStructure().GetToolsCategory())
-		//.SetGroup(WorkspaceMenu::GetMenuStructure().GetDeveloperToolsDebugCategory());
-		.SetIcon(FSlateIcon(FAppStyle::GetAppStyleSetName(), "Profiler.EventGraph.ExpandHotPath16"));
+		.SetIcon(FSlateIcon(FGameplayAbilitiesEditorStyleSet::Get().GetStyleSetName(), "GameplayCueEditor_Small"));
 		
 	ApplyGameplayModEvaluationChannelAliasesToEnumMetadata();
 
