Iris
- Fixed bIsQuantizedState flag not being property set in a bunch of NetSerializers

diff --git a/Engine/Plugins/Runtime/GameplayAbilities/Source/GameplayAbilities/Private/Serialization/GameplayAbilityRepAnimMontageNetSerializer.cpp b/Engine/Plugins/Runtime/GameplayAbilities/Source/GameplayAbilities/Private/Serialization/GameplayAbilityRepAnimMontageNetSerializer.cpp
--- a/Engine/Plugins/Runtime/GameplayAbilities/Source/GameplayAbilities/Private/Serialization/GameplayAbilityRepAnimMontageNetSerializer.cpp
+++ b/Engine/Plugins/Runtime/GameplayAbilities/Source/GameplayAbilities/Private/Serialization/GameplayAbilityRepAnimMontageNetSerializer.cpp
@@ -326,7 +326,7 @@ bool FGameplayAbilityRepAnimMontageNetSerializer::IsEqual(FNetSerializationConte
 		}
 
 		// Forward to normal StructNetSerializer
-		FNetIsEqualArgs IsEqualArgs = {};
+		FNetIsEqualArgs IsEqualArgs = Args;
 		IsEqualArgs.NetSerializerConfig = &StructNetSerializerConfigForBase;
 		IsEqualArgs.Source0 = NetSerializerValuePointer(&Value0.GameplayAbilityRepAnimMontage);
 		IsEqualArgs.Source1 = NetSerializerValuePointer(&Value1.GameplayAbilityRepAnimMontage);
@@ -347,10 +347,11 @@ bool FGameplayAbilityRepAnimMontageNetSerializer::IsEqual(FNetSerializationConte
 		const SourceType& SourceValue1 = *reinterpret_cast<const SourceType*>(Args.Source1);
 
 		// Forward to normal StructNetSerializer
-		FNetIsEqualArgs IsEqualArgs ={};
+		FNetIsEqualArgs IsEqualArgs = Args;
 		IsEqualArgs.NetSerializerConfig = &StructNetSerializerConfigForBase;
 		IsEqualArgs.Source0 = Args.Source0;
 		IsEqualArgs.Source1 = Args.Source1;
+
 		if (!StructNetSerializer->IsEqual(Context, IsEqualArgs))
 		{
 			return false;
diff --git a/Engine/Plugins/Runtime/GameplayAbilities/Source/GameplayAbilities/Private/Serialization/GameplayEffectContextNetSerializer.cpp b/Engine/Plugins/Runtime/GameplayAbilities/Source/GameplayAbilities/Private/Serialization/GameplayEffectContextNetSerializer.cpp
--- a/Engine/Plugins/Runtime/GameplayAbilities/Source/GameplayAbilities/Private/Serialization/GameplayEffectContextNetSerializer.cpp
+++ b/Engine/Plugins/Runtime/GameplayAbilities/Source/GameplayAbilities/Private/Serialization/GameplayEffectContextNetSerializer.cpp
@@ -291,10 +291,11 @@ bool FGameplayEffectContextNetSerializer::IsEqual(FNetSerializationContext& Cont
 		}
 
 		{
-			FNetIsEqualArgs GEIsEqualArgs = {};
+			FNetIsEqualArgs GEIsEqualArgs = Args;
 			GEIsEqualArgs.NetSerializerConfig = &StructNetSerializerConfigForGE;
 			GEIsEqualArgs.Source0 = NetSerializerValuePointer(&Value0.EffectContext);
 			GEIsEqualArgs.Source1 = NetSerializerValuePointer(&Value1.EffectContext);
+
 			if (!StructNetSerializer->IsEqual(Context, GEIsEqualArgs))
 			{
 				return false;
@@ -303,10 +304,11 @@ bool FGameplayEffectContextNetSerializer::IsEqual(FNetSerializationContext& Cont
 
 		if (Value0.ReplicationFlags & EReplicationFlags::ReplicateHitResult)
 		{
-			FNetIsEqualArgs HitResultIsEqualArgs = {};
+			FNetIsEqualArgs HitResultIsEqualArgs = Args;
 			HitResultIsEqualArgs.NetSerializerConfig = &StructNetSerializerConfigForHitResult;
 			HitResultIsEqualArgs.Source0 = NetSerializerValuePointer(&Value0.HitResult);
 			HitResultIsEqualArgs.Source1 = NetSerializerValuePointer(&Value1.HitResult);
+
 			if (!StructNetSerializer->IsEqual(Context, HitResultIsEqualArgs))
 			{
 				return false;
@@ -332,10 +334,11 @@ bool FGameplayEffectContextNetSerializer::IsEqual(FNetSerializationContext& Cont
 		const FHitResult* HitResult1 = SourceValue1.GetHitResult();
 		if ((HitResult0 != nullptr) && (HitResult1 != nullptr))
 		{
-			FNetIsEqualArgs HitEffectIsEqualArgs ={};
+			FNetIsEqualArgs HitEffectIsEqualArgs = Args;
 			HitEffectIsEqualArgs.NetSerializerConfig = &StructNetSerializerConfigForHitResult;
 			HitEffectIsEqualArgs.Source0 = NetSerializerValuePointer(HitResult0);
 			HitEffectIsEqualArgs.Source1 = NetSerializerValuePointer(HitResult1);
+
 			if (!StructNetSerializer->IsEqual(Context, HitEffectIsEqualArgs))
 			{
 				return false;
@@ -355,10 +358,11 @@ bool FGameplayEffectContextNetSerializer::IsEqual(FNetSerializationContext& Cont
 			TempGE0.CopyReplicatedFieldsFrom(SourceValue0);
 			TempGE1.CopyReplicatedFieldsFrom(SourceValue1);
 
-			FNetIsEqualArgs GEIsEqualArgs ={};
+			FNetIsEqualArgs GEIsEqualArgs = Args;
 			GEIsEqualArgs.NetSerializerConfig = &StructNetSerializerConfigForGE;
 			GEIsEqualArgs.Source0 = NetSerializerValuePointer(&TempGE0);
 			GEIsEqualArgs.Source1 = NetSerializerValuePointer(&TempGE1);
+
 			if (!StructNetSerializer->IsEqual(Context, GEIsEqualArgs))
 			{
 				return false;
