[GameplayAbilities]
* Removed includes from header files to speed up compile times

diff --git a/Engine/Plugins/Runtime/GameplayAbilities/Source/GameplayAbilities/Private/Abilities/Async/AbilityAsync_WaitGameplayEffectApplied.cpp b/Engine/Plugins/Runtime/GameplayAbilities/Source/GameplayAbilities/Private/Abilities/Async/AbilityAsync_WaitGameplayEffectApplied.cpp
--- a/Engine/Plugins/Runtime/GameplayAbilities/Source/GameplayAbilities/Private/Abilities/Async/AbilityAsync_WaitGameplayEffectApplied.cpp
+++ b/Engine/Plugins/Runtime/GameplayAbilities/Source/GameplayAbilities/Private/Abilities/Async/AbilityAsync_WaitGameplayEffectApplied.cpp
@@ -3,6 +3,7 @@
 #include "Abilities/Async/AbilityAsync_WaitGameplayEffectApplied.h"
 #include "AbilitySystemGlobals.h"
 #include "AbilitySystemComponent.h"
+#include "AbilitySystemLog.h"
 
 #include UE_INLINE_GENERATED_CPP_BY_NAME(AbilityAsync_WaitGameplayEffectApplied)
 
diff --git a/Engine/Plugins/Runtime/GameplayAbilities/Source/GameplayAbilities/Private/Abilities/GameplayAbility.cpp b/Engine/Plugins/Runtime/GameplayAbilities/Source/GameplayAbilities/Private/Abilities/GameplayAbility.cpp
--- a/Engine/Plugins/Runtime/GameplayAbilities/Source/GameplayAbilities/Private/Abilities/GameplayAbility.cpp
+++ b/Engine/Plugins/Runtime/GameplayAbilities/Source/GameplayAbilities/Private/Abilities/GameplayAbility.cpp
@@ -1,6 +1,7 @@
 // Copyright Epic Games, Inc. All Rights Reserved.
 
 #include "Abilities/GameplayAbility.h"
+#include "AbilitySystemLog.h"
 #include "TimerManager.h"
 #include "Engine/BlueprintGeneratedClass.h"
 #include "Engine/Engine.h"
@@ -20,6 +21,15 @@
 
 #include UE_INLINE_GENERATED_CPP_BY_NAME(GameplayAbility)
 
+#define ENSURE_ABILITY_IS_INSTANTIATED_OR_RETURN(FunctionName, ReturnValue)																				\
+{																																						\
+	if (!ensure(IsInstantiated()))																														\
+	{																																					\
+		ABILITY_LOG(Error, TEXT("%s: " #FunctionName " cannot be called on a non-instanced ability. Check the instancing policy."), *GetPathName());	\
+		return ReturnValue;																																\
+	}																																					\
+}
+
 namespace FAbilitySystemTweaks
 {
 	int ClearAbilityTimers = 1;
@@ -1064,6 +1074,24 @@ UAbilitySystemComponent* UGameplayAbility::GetAbilitySystemComponentFromActorInf
 	return AbilitySystemComponent;
 }
 
+const FGameplayAbilityActorInfo* UGameplayAbility::GetCurrentActorInfo() const
+{
+	ENSURE_ABILITY_IS_INSTANTIATED_OR_RETURN(GetCurrentActorInfo, nullptr);
+	return CurrentActorInfo;
+}
+
+FGameplayAbilityActivationInfo UGameplayAbility::GetCurrentActivationInfo() const
+{
+	ENSURE_ABILITY_IS_INSTANTIATED_OR_RETURN(GetCurrentActivationInfo, FGameplayAbilityActivationInfo());
+	return CurrentActivationInfo;
+}
+
+FGameplayAbilitySpecHandle UGameplayAbility::GetCurrentAbilitySpecHandle() const
+{
+	ENSURE_ABILITY_IS_INSTANTIATED_OR_RETURN(GetCurrentAbilitySpecHandle, FGameplayAbilitySpecHandle());
+	return CurrentSpecHandle;
+}
+
 FGameplayEffectSpecHandle UGameplayAbility::MakeOutgoingGameplayEffectSpec(TSubclassOf<UGameplayEffect> GameplayEffectClass, float Level) const
 {
 	check(CurrentActorInfo && CurrentActorInfo->AbilitySystemComponent.IsValid());
diff --git a/Engine/Plugins/Runtime/GameplayAbilities/Source/GameplayAbilities/Private/Abilities/GameplayAbilityTargetActor.cpp b/Engine/Plugins/Runtime/GameplayAbilities/Source/GameplayAbilities/Private/Abilities/GameplayAbilityTargetActor.cpp
--- a/Engine/Plugins/Runtime/GameplayAbilities/Source/GameplayAbilities/Private/Abilities/GameplayAbilityTargetActor.cpp
+++ b/Engine/Plugins/Runtime/GameplayAbilities/Source/GameplayAbilities/Private/Abilities/GameplayAbilityTargetActor.cpp
@@ -4,6 +4,7 @@
 #include "GameplayAbilitySpec.h"
 #include "GameFramework/PlayerController.h"
 #include "AbilitySystemComponent.h"
+#include "AbilitySystemLog.h"
 #include "Net/UnrealNetwork.h"
 
 #include UE_INLINE_GENERATED_CPP_BY_NAME(GameplayAbilityTargetActor)
diff --git a/Engine/Plugins/Runtime/GameplayAbilities/Source/GameplayAbilities/Private/Abilities/Tasks/AbilityTask.cpp b/Engine/Plugins/Runtime/GameplayAbilities/Source/GameplayAbilities/Private/Abilities/Tasks/AbilityTask.cpp
--- a/Engine/Plugins/Runtime/GameplayAbilities/Source/GameplayAbilities/Private/Abilities/Tasks/AbilityTask.cpp
+++ b/Engine/Plugins/Runtime/GameplayAbilities/Source/GameplayAbilities/Private/Abilities/Tasks/AbilityTask.cpp
@@ -1,6 +1,7 @@
 // Copyright Epic Games, Inc. All Rights Reserved.
 
 #include "Abilities/Tasks/AbilityTask.h"
+#include "AbilitySystemLog.h"
 #include "AbilitySystemComponent.h"
 #include "AbilitySystemStats.h"
 
diff --git a/Engine/Plugins/Runtime/GameplayAbilities/Source/GameplayAbilities/Private/Abilities/Tasks/AbilityTask_ApplyRootMotionConstantForce.cpp b/Engine/Plugins/Runtime/GameplayAbilities/Source/GameplayAbilities/Private/Abilities/Tasks/AbilityTask_ApplyRootMotionConstantForce.cpp
--- a/Engine/Plugins/Runtime/GameplayAbilities/Source/GameplayAbilities/Private/Abilities/Tasks/AbilityTask_ApplyRootMotionConstantForce.cpp
+++ b/Engine/Plugins/Runtime/GameplayAbilities/Source/GameplayAbilities/Private/Abilities/Tasks/AbilityTask_ApplyRootMotionConstantForce.cpp
@@ -5,6 +5,7 @@
 #include "GameFramework/CharacterMovementComponent.h"
 #include "AbilitySystemComponent.h"
 #include "AbilitySystemGlobals.h"
+#include "AbilitySystemLog.h"
 #include "Net/UnrealNetwork.h"
 #include "Engine/World.h"
 
diff --git a/Engine/Plugins/Runtime/GameplayAbilities/Source/GameplayAbilities/Private/Abilities/Tasks/AbilityTask_ApplyRootMotionJumpForce.cpp b/Engine/Plugins/Runtime/GameplayAbilities/Source/GameplayAbilities/Private/Abilities/Tasks/AbilityTask_ApplyRootMotionJumpForce.cpp
--- a/Engine/Plugins/Runtime/GameplayAbilities/Source/GameplayAbilities/Private/Abilities/Tasks/AbilityTask_ApplyRootMotionJumpForce.cpp
+++ b/Engine/Plugins/Runtime/GameplayAbilities/Source/GameplayAbilities/Private/Abilities/Tasks/AbilityTask_ApplyRootMotionJumpForce.cpp
@@ -6,6 +6,7 @@
 #include "GameFramework/CharacterMovementComponent.h"
 #include "AbilitySystemComponent.h"
 #include "AbilitySystemGlobals.h"
+#include "AbilitySystemLog.h"
 #include "Net/UnrealNetwork.h"
 #include "Engine/World.h"
 
diff --git a/Engine/Plugins/Runtime/GameplayAbilities/Source/GameplayAbilities/Private/Abilities/Tasks/AbilityTask_ApplyRootMotionMoveToActorForce.cpp b/Engine/Plugins/Runtime/GameplayAbilities/Source/GameplayAbilities/Private/Abilities/Tasks/AbilityTask_ApplyRootMotionMoveToActorForce.cpp
--- a/Engine/Plugins/Runtime/GameplayAbilities/Source/GameplayAbilities/Private/Abilities/Tasks/AbilityTask_ApplyRootMotionMoveToActorForce.cpp
+++ b/Engine/Plugins/Runtime/GameplayAbilities/Source/GameplayAbilities/Private/Abilities/Tasks/AbilityTask_ApplyRootMotionMoveToActorForce.cpp
@@ -8,6 +8,7 @@
 #include "DrawDebugHelpers.h"
 #include "AbilitySystemComponent.h"
 #include "AbilitySystemGlobals.h"
+#include "AbilitySystemLog.h"
 #include "Net/UnrealNetwork.h"
 #include "Abilities/GameplayAbilityTargetTypes.h"
 #include "Engine/World.h"
diff --git a/Engine/Plugins/Runtime/GameplayAbilities/Source/GameplayAbilities/Private/Abilities/Tasks/AbilityTask_ApplyRootMotionMoveToForce.cpp b/Engine/Plugins/Runtime/GameplayAbilities/Source/GameplayAbilities/Private/Abilities/Tasks/AbilityTask_ApplyRootMotionMoveToForce.cpp
--- a/Engine/Plugins/Runtime/GameplayAbilities/Source/GameplayAbilities/Private/Abilities/Tasks/AbilityTask_ApplyRootMotionMoveToForce.cpp
+++ b/Engine/Plugins/Runtime/GameplayAbilities/Source/GameplayAbilities/Private/Abilities/Tasks/AbilityTask_ApplyRootMotionMoveToForce.cpp
@@ -6,6 +6,7 @@
 #include "GameFramework/CharacterMovementComponent.h"
 #include "AbilitySystemComponent.h"
 #include "AbilitySystemGlobals.h"
+#include "AbilitySystemLog.h"
 #include "Net/UnrealNetwork.h"
 #include "Engine/World.h"
 
diff --git a/Engine/Plugins/Runtime/GameplayAbilities/Source/GameplayAbilities/Private/Abilities/Tasks/AbilityTask_ApplyRootMotionRadialForce.cpp b/Engine/Plugins/Runtime/GameplayAbilities/Source/GameplayAbilities/Private/Abilities/Tasks/AbilityTask_ApplyRootMotionRadialForce.cpp
--- a/Engine/Plugins/Runtime/GameplayAbilities/Source/GameplayAbilities/Private/Abilities/Tasks/AbilityTask_ApplyRootMotionRadialForce.cpp
+++ b/Engine/Plugins/Runtime/GameplayAbilities/Source/GameplayAbilities/Private/Abilities/Tasks/AbilityTask_ApplyRootMotionRadialForce.cpp
@@ -5,6 +5,7 @@
 #include "GameFramework/CharacterMovementComponent.h"
 #include "AbilitySystemComponent.h"
 #include "AbilitySystemGlobals.h"
+#include "AbilitySystemLog.h"
 #include "Net/UnrealNetwork.h"
 #include "Engine/World.h"
 
diff --git a/Engine/Plugins/Runtime/GameplayAbilities/Source/GameplayAbilities/Private/Abilities/Tasks/AbilityTask_PlayMontageAndWait.cpp b/Engine/Plugins/Runtime/GameplayAbilities/Source/GameplayAbilities/Private/Abilities/Tasks/AbilityTask_PlayMontageAndWait.cpp
--- a/Engine/Plugins/Runtime/GameplayAbilities/Source/GameplayAbilities/Private/Abilities/Tasks/AbilityTask_PlayMontageAndWait.cpp
+++ b/Engine/Plugins/Runtime/GameplayAbilities/Source/GameplayAbilities/Private/Abilities/Tasks/AbilityTask_PlayMontageAndWait.cpp
@@ -3,6 +3,7 @@
 #include "Abilities/Tasks/AbilityTask_PlayMontageAndWait.h"
 #include "GameFramework/Character.h"
 #include "AbilitySystemComponent.h"
+#include "AbilitySystemLog.h"
 #include "AbilitySystemGlobals.h"
 
 #include UE_INLINE_GENERATED_CPP_BY_NAME(AbilityTask_PlayMontageAndWait)
diff --git a/Engine/Plugins/Runtime/GameplayAbilities/Source/GameplayAbilities/Private/Abilities/Tasks/AbilityTask_WaitGameplayEffectApplied.cpp b/Engine/Plugins/Runtime/GameplayAbilities/Source/GameplayAbilities/Private/Abilities/Tasks/AbilityTask_WaitGameplayEffectApplied.cpp
--- a/Engine/Plugins/Runtime/GameplayAbilities/Source/GameplayAbilities/Private/Abilities/Tasks/AbilityTask_WaitGameplayEffectApplied.cpp
+++ b/Engine/Plugins/Runtime/GameplayAbilities/Source/GameplayAbilities/Private/Abilities/Tasks/AbilityTask_WaitGameplayEffectApplied.cpp
@@ -3,6 +3,7 @@
 #include "Abilities/Tasks/AbilityTask_WaitGameplayEffectApplied.h"
 #include "AbilitySystemGlobals.h"
 #include "AbilitySystemComponent.h"
+#include "AbilitySystemLog.h"
 
 #include UE_INLINE_GENERATED_CPP_BY_NAME(AbilityTask_WaitGameplayEffectApplied)
 
diff --git a/Engine/Plugins/Runtime/GameplayAbilities/Source/GameplayAbilities/Private/Abilities/Tasks/AbilityTask_WaitVelocityChange.cpp b/Engine/Plugins/Runtime/GameplayAbilities/Source/GameplayAbilities/Private/Abilities/Tasks/AbilityTask_WaitVelocityChange.cpp
--- a/Engine/Plugins/Runtime/GameplayAbilities/Source/GameplayAbilities/Private/Abilities/Tasks/AbilityTask_WaitVelocityChange.cpp
+++ b/Engine/Plugins/Runtime/GameplayAbilities/Source/GameplayAbilities/Private/Abilities/Tasks/AbilityTask_WaitVelocityChange.cpp
@@ -1,6 +1,7 @@
 // Copyright Epic Games, Inc. All Rights Reserved.
 
 #include "Abilities/Tasks/AbilityTask_WaitVelocityChange.h"
+#include "AbilitySystemLog.h"
 #include "GameFramework/MovementComponent.h"
 
 #include UE_INLINE_GENERATED_CPP_BY_NAME(AbilityTask_WaitVelocityChange)
diff --git a/Engine/Plugins/Runtime/GameplayAbilities/Source/GameplayAbilities/Private/AbilitySystemBlueprintLibrary.cpp b/Engine/Plugins/Runtime/GameplayAbilities/Source/GameplayAbilities/Private/AbilitySystemBlueprintLibrary.cpp
--- a/Engine/Plugins/Runtime/GameplayAbilities/Source/GameplayAbilities/Private/AbilitySystemBlueprintLibrary.cpp
+++ b/Engine/Plugins/Runtime/GameplayAbilities/Source/GameplayAbilities/Private/AbilitySystemBlueprintLibrary.cpp
@@ -4,6 +4,7 @@
 #include "GameplayEffectAggregator.h"
 #include "AbilitySystemGlobals.h"
 #include "AbilitySystemComponent.h"
+#include "AbilitySystemLog.h"
 #include "Engine/World.h"
 #include "GameplayEffect.h"
 #include "GameplayEffectUIData.h"
diff --git a/Engine/Plugins/Runtime/GameplayAbilities/Source/GameplayAbilities/Private/AbilitySystemComponent.cpp b/Engine/Plugins/Runtime/GameplayAbilities/Source/GameplayAbilities/Private/AbilitySystemComponent.cpp
--- a/Engine/Plugins/Runtime/GameplayAbilities/Source/GameplayAbilities/Private/AbilitySystemComponent.cpp
+++ b/Engine/Plugins/Runtime/GameplayAbilities/Source/GameplayAbilities/Private/AbilitySystemComponent.cpp
@@ -1,6 +1,7 @@
 // Copyright Epic Games, Inc. All Rights Reserved.
 
 #include "AbilitySystemComponent.h"
+#include "AbilitySystemLog.h"
 #include "UObject/UObjectHash.h"
 #include "UObject/UObjectIterator.h"
 #include "Engine/Canvas.h"
diff --git a/Engine/Plugins/Runtime/GameplayAbilities/Source/GameplayAbilities/Private/AttributeSet.cpp b/Engine/Plugins/Runtime/GameplayAbilities/Source/GameplayAbilities/Private/AttributeSet.cpp
--- a/Engine/Plugins/Runtime/GameplayAbilities/Source/GameplayAbilities/Private/AttributeSet.cpp
+++ b/Engine/Plugins/Runtime/GameplayAbilities/Source/GameplayAbilities/Private/AttributeSet.cpp
@@ -427,6 +427,11 @@ void UAttributeSet::InitFromMetaDataTable(const UDataTable* DataTable)
 	PrintDebug();
 }
 
+AActor* UAttributeSet::GetOwningActor() const
+{
+	return CastChecked<AActor>(GetOuter());
+}
+
 UAbilitySystemComponent* UAttributeSet::GetOwningAbilitySystemComponent() const
 {
 	return UAbilitySystemGlobals::GetAbilitySystemComponentFromActor(GetOwningActor());
diff --git a/Engine/Plugins/Runtime/GameplayAbilities/Source/GameplayAbilities/Private/GameplayAbilityTargetTypes.cpp b/Engine/Plugins/Runtime/GameplayAbilities/Source/GameplayAbilities/Private/GameplayAbilityTargetTypes.cpp
--- a/Engine/Plugins/Runtime/GameplayAbilities/Source/GameplayAbilities/Private/GameplayAbilityTargetTypes.cpp
+++ b/Engine/Plugins/Runtime/GameplayAbilities/Source/GameplayAbilities/Private/GameplayAbilityTargetTypes.cpp
@@ -5,6 +5,7 @@
 #include "AbilitySystemComponent.h"
 #include "AbilitySystemBlueprintLibrary.h"
 #include "AbilitySystemGlobals.h"
+#include "AbilitySystemLog.h"
 #include "Components/MeshComponent.h"
 
 #include UE_INLINE_GENERATED_CPP_BY_NAME(GameplayAbilityTargetTypes)
diff --git a/Engine/Plugins/Runtime/GameplayAbilities/Source/GameplayAbilities/Private/GameplayAbilityTypes.cpp b/Engine/Plugins/Runtime/GameplayAbilities/Source/GameplayAbilities/Private/GameplayAbilityTypes.cpp
--- a/Engine/Plugins/Runtime/GameplayAbilities/Source/GameplayAbilities/Private/GameplayAbilityTypes.cpp
+++ b/Engine/Plugins/Runtime/GameplayAbilities/Source/GameplayAbilities/Private/GameplayAbilityTypes.cpp
@@ -4,6 +4,7 @@
 
 #include "Abilities/GameplayAbility.h"
 #include "AbilitySystemComponent.h"
+#include "AbilitySystemLog.h"
 #include "Animation/AnimInstance.h"
 #include "Animation/AnimMontage.h"
 #include "Components/SkeletalMeshComponent.h"
@@ -143,6 +144,14 @@ bool FGameplayAbilityActorInfo::IsNetAuthority() const
 	return false;
 }
 
+FGameplayAbilityActivationInfo::FGameplayAbilityActivationInfo(AActor* InActor)
+	: bCanBeEndedByOtherInstance(false)	
+{
+	// On Init, we are either Authority or NonAuthority. We haven't been given a PredictionKey and we haven't been confirmed.
+	// NonAuthority essentially means 'I'm not sure what how I'm going to do this yet'.
+	ActivationMode = (InActor->GetLocalRole() == ROLE_Authority ? EGameplayAbilityActivationMode::Authority : EGameplayAbilityActivationMode::NonAuthority);
+}
+
 void FGameplayAbilityActivationInfo::SetPredicting(FPredictionKey PredictionKey)
 {
 	ActivationMode = EGameplayAbilityActivationMode::Predicting;
diff --git a/Engine/Plugins/Runtime/GameplayAbilities/Source/GameplayAbilities/Private/GameplayCueInterface.cpp b/Engine/Plugins/Runtime/GameplayAbilities/Source/GameplayAbilities/Private/GameplayCueInterface.cpp
--- a/Engine/Plugins/Runtime/GameplayAbilities/Source/GameplayAbilities/Private/GameplayCueInterface.cpp
+++ b/Engine/Plugins/Runtime/GameplayAbilities/Source/GameplayAbilities/Private/GameplayCueInterface.cpp
@@ -1,6 +1,7 @@
 // Copyright Epic Games, Inc. All Rights Reserved.
 
 #include "GameplayCueInterface.h"
+#include "AbilitySystemLog.h"
 #include "AbilitySystemStats.h"
 #include "GameplayTagsModule.h"
 #include "AbilitySystemComponent.h"
diff --git a/Engine/Plugins/Runtime/GameplayAbilities/Source/GameplayAbilities/Private/GameplayCueManager.cpp b/Engine/Plugins/Runtime/GameplayAbilities/Source/GameplayAbilities/Private/GameplayCueManager.cpp
--- a/Engine/Plugins/Runtime/GameplayAbilities/Source/GameplayAbilities/Private/GameplayCueManager.cpp
+++ b/Engine/Plugins/Runtime/GameplayAbilities/Source/GameplayAbilities/Private/GameplayCueManager.cpp
@@ -11,6 +11,7 @@
 #include "GameplayTagsManager.h"
 #include "GameplayTagsModule.h"
 #include "AbilitySystemGlobals.h"
+#include "AbilitySystemLog.h"
 #include "AssetRegistry/AssetRegistryModule.h"
 #include "GameplayCueInterface.h"
 #include "GameplayCueSet.h"
diff --git a/Engine/Plugins/Runtime/GameplayAbilities/Source/GameplayAbilities/Private/GameplayCueNotify_Actor.cpp b/Engine/Plugins/Runtime/GameplayAbilities/Source/GameplayAbilities/Private/GameplayCueNotify_Actor.cpp
--- a/Engine/Plugins/Runtime/GameplayAbilities/Source/GameplayAbilities/Private/GameplayCueNotify_Actor.cpp
+++ b/Engine/Plugins/Runtime/GameplayAbilities/Source/GameplayAbilities/Private/GameplayCueNotify_Actor.cpp
@@ -6,6 +6,7 @@
 #include "Components/TimelineComponent.h"
 #include "AbilitySystemStats.h"
 #include "AbilitySystemGlobals.h"
+#include "AbilitySystemLog.h"
 #include "GameplayCueManager.h"
 #include "GameplayCueNotifyTypes.h"
 
diff --git a/Engine/Plugins/Runtime/GameplayAbilities/Source/GameplayAbilities/Private/GameplayCueNotify_Static.cpp b/Engine/Plugins/Runtime/GameplayAbilities/Source/GameplayAbilities/Private/GameplayCueNotify_Static.cpp
--- a/Engine/Plugins/Runtime/GameplayAbilities/Source/GameplayAbilities/Private/GameplayCueNotify_Static.cpp
+++ b/Engine/Plugins/Runtime/GameplayAbilities/Source/GameplayAbilities/Private/GameplayCueNotify_Static.cpp
@@ -5,6 +5,7 @@
 #include "AbilitySystemStats.h"
 #include "AbilitySystemGlobals.h"
 #include "AbilitySystemComponent.h"
+#include "AbilitySystemLog.h"
 #include "GameplayCueManager.h"
 
 #include UE_INLINE_GENERATED_CPP_BY_NAME(GameplayCueNotify_Static)
diff --git a/Engine/Plugins/Runtime/GameplayAbilities/Source/GameplayAbilities/Private/GameplayCueSet.cpp b/Engine/Plugins/Runtime/GameplayAbilities/Source/GameplayAbilities/Private/GameplayCueSet.cpp
--- a/Engine/Plugins/Runtime/GameplayAbilities/Source/GameplayAbilities/Private/GameplayCueSet.cpp
+++ b/Engine/Plugins/Runtime/GameplayAbilities/Source/GameplayAbilities/Private/GameplayCueSet.cpp
@@ -4,6 +4,7 @@
 #include "GameplayTagsManager.h"
 #include "GameplayTagsModule.h"
 #include "AbilitySystemGlobals.h"
+#include "AbilitySystemLog.h"
 #include "GameplayCueNotify_Actor.h"
 #include "GameplayCueNotify_Static.h"
 #include "GameplayCueManager.h"
diff --git a/Engine/Plugins/Runtime/GameplayAbilities/Source/GameplayAbilities/Private/GameplayEffect.cpp b/Engine/Plugins/Runtime/GameplayAbilities/Source/GameplayAbilities/Private/GameplayEffect.cpp
--- a/Engine/Plugins/Runtime/GameplayAbilities/Source/GameplayAbilities/Private/GameplayEffect.cpp
+++ b/Engine/Plugins/Runtime/GameplayAbilities/Source/GameplayAbilities/Private/GameplayEffect.cpp
@@ -3,6 +3,7 @@
 #include "GameplayEffect.h"
 
 #include "TimerManager.h"
+#include "AbilitySystemLog.h"
 #include "GameFramework/GameStateBase.h"
 #include "Engine/ChildConnection.h"
 #include "Engine/NetConnection.h"
diff --git a/Engine/Plugins/Runtime/GameplayAbilities/Source/GameplayAbilities/Private/GameplayEffectAggregator.cpp b/Engine/Plugins/Runtime/GameplayAbilities/Source/GameplayAbilities/Private/GameplayEffectAggregator.cpp
--- a/Engine/Plugins/Runtime/GameplayAbilities/Source/GameplayAbilities/Private/GameplayEffectAggregator.cpp
+++ b/Engine/Plugins/Runtime/GameplayAbilities/Source/GameplayAbilities/Private/GameplayEffectAggregator.cpp
@@ -4,6 +4,7 @@
 #include "UObject/UObjectHash.h"
 #include "UObject/UObjectIterator.h"
 #include "AbilitySystemComponent.h"
+#include "AbilitySystemLog.h"
 #include "AbilitySystemStats.h"
 
 void FAggregatorMod::UpdateQualifies(const FAggregatorEvaluateParameters& Parameters) const
diff --git a/Engine/Plugins/Runtime/GameplayAbilities/Source/GameplayAbilities/Private/GameplayEffectExecutionCalculation.cpp b/Engine/Plugins/Runtime/GameplayAbilities/Source/GameplayAbilities/Private/GameplayEffectExecutionCalculation.cpp
--- a/Engine/Plugins/Runtime/GameplayAbilities/Source/GameplayAbilities/Private/GameplayEffectExecutionCalculation.cpp
+++ b/Engine/Plugins/Runtime/GameplayAbilities/Source/GameplayAbilities/Private/GameplayEffectExecutionCalculation.cpp
@@ -2,6 +2,7 @@
 
 #include "GameplayEffectExecutionCalculation.h"
 #include "AbilitySystemComponent.h"
+#include "AbilitySystemLog.h"
 
 #include UE_INLINE_GENERATED_CPP_BY_NAME(GameplayEffectExecutionCalculation)
 
diff --git a/Engine/Plugins/Runtime/GameplayAbilities/Source/GameplayAbilities/Private/GameplayEffectTypes.cpp b/Engine/Plugins/Runtime/GameplayAbilities/Source/GameplayAbilities/Private/GameplayEffectTypes.cpp
--- a/Engine/Plugins/Runtime/GameplayAbilities/Source/GameplayAbilities/Private/GameplayEffectTypes.cpp
+++ b/Engine/Plugins/Runtime/GameplayAbilities/Source/GameplayAbilities/Private/GameplayEffectTypes.cpp
@@ -1,6 +1,7 @@
 // Copyright Epic Games, Inc. All Rights Reserved.
 
 #include "GameplayEffectTypes.h"
+#include "AbilitySystemLog.h"
 #include "GameFramework/Pawn.h"
 #include "GameplayTagAssetInterface.h"
 #include "GameplayEffect.h"
@@ -258,6 +259,12 @@ bool FGameplayEffectContext::NetSerialize(FArchive& Ar, class UPackageMap* Map,
 	return true;
 }
 
+FString FGameplayEffectContext::ToString() const
+{
+	const AActor* InstigatorPtr = Instigator.Get();
+	return (InstigatorPtr ? InstigatorPtr->GetName() : FString(TEXT("NONE")));
+}
+
 bool FGameplayEffectContext::IsLocallyControlled() const
 {
 	APawn* Pawn = Cast<APawn>(Instigator.Get());
@@ -858,6 +865,12 @@ FGameplayEffectSpecHandle::FGameplayEffectSpecHandle(FGameplayEffectSpec* DataPt
 
 }
 
+bool FGameplayEffectSpecHandle::NetSerialize(FArchive& Ar, class UPackageMap* Map, bool& bOutSuccess)
+{
+	ABILITY_LOG(Fatal, TEXT("FGameplayEffectSpecHandle should not be NetSerialized"));
+	return false;
+}
+
 FGameplayCueParameters::FGameplayCueParameters(const FGameplayEffectSpecForRPC& Spec)
 : NormalizedMagnitude(0.0f)
 , RawMagnitude(0.0f)
@@ -1123,6 +1136,23 @@ const UObject* FGameplayCueParameters::GetSourceObject() const
 	return EffectContext.GetSourceObject();
 }
 
+void FMinimalReplicationTagCountMap::RemoveTag(const FGameplayTag& Tag)
+{
+	MapID++;
+	int32& Count = TagMap.FindOrAdd(Tag);
+	Count--;
+	if (Count == 0)
+	{
+		// Remove from map so that we do not replicate
+		TagMap.Remove(Tag);
+	}
+	else if (Count < 0)
+	{
+		ABILITY_LOG(Error, TEXT("FMinimalReplicationTagCountMap::RemoveTag called on Tag %s and count is now < 0"), *Tag.ToString());
+		Count = 0;
+	}
+}
+
 bool FMinimalReplicationTagCountMap::NetSerialize(FArchive& Ar, class UPackageMap* Map, bool& bOutSuccess)
 {
 	const int32 CountBits = UAbilitySystemGlobals::Get().MinimalReplicationTagCountBits;
diff --git a/Engine/Plugins/Runtime/GameplayAbilities/Source/GameplayAbilities/Private/GameplayModMagnitudeCalculation.cpp b/Engine/Plugins/Runtime/GameplayAbilities/Source/GameplayAbilities/Private/GameplayModMagnitudeCalculation.cpp
--- a/Engine/Plugins/Runtime/GameplayAbilities/Source/GameplayAbilities/Private/GameplayModMagnitudeCalculation.cpp
+++ b/Engine/Plugins/Runtime/GameplayAbilities/Source/GameplayAbilities/Private/GameplayModMagnitudeCalculation.cpp
@@ -1,6 +1,7 @@
 // Copyright Epic Games, Inc. All Rights Reserved.
 
 #include "GameplayModMagnitudeCalculation.h"
+#include "AbilitySystemLog.h"
 
 #include UE_INLINE_GENERATED_CPP_BY_NAME(GameplayModMagnitudeCalculation)
 
diff --git a/Engine/Plugins/Runtime/GameplayAbilities/Source/GameplayAbilities/Private/GameplayPrediction.cpp b/Engine/Plugins/Runtime/GameplayAbilities/Source/GameplayAbilities/Private/GameplayPrediction.cpp
--- a/Engine/Plugins/Runtime/GameplayAbilities/Source/GameplayAbilities/Private/GameplayPrediction.cpp
+++ b/Engine/Plugins/Runtime/GameplayAbilities/Source/GameplayAbilities/Private/GameplayPrediction.cpp
@@ -2,6 +2,7 @@
 
 #include "GameplayPrediction.h"
 #include "AbilitySystemComponent.h"
+#include "AbilitySystemLog.h"
 
 #include UE_INLINE_GENERATED_CPP_BY_NAME(GameplayPrediction)
 
diff --git a/Engine/Plugins/Runtime/GameplayAbilities/Source/GameplayAbilities/Public/Abilities/GameplayAbility.h b/Engine/Plugins/Runtime/GameplayAbilities/Source/GameplayAbilities/Public/Abilities/GameplayAbility.h
--- a/Engine/Plugins/Runtime/GameplayAbilities/Source/GameplayAbilities/Public/Abilities/GameplayAbility.h
+++ b/Engine/Plugins/Runtime/GameplayAbilities/Source/GameplayAbilities/Public/Abilities/GameplayAbility.h
@@ -83,16 +83,6 @@ DECLARE_MULTICAST_DELEGATE_OneParam(FOnGameplayAbilityStateEnded, FName);
 /** Used to delay execution until we leave a critical section */
 DECLARE_DELEGATE(FPostLockDelegate);
 
-
-#define ENSURE_ABILITY_IS_INSTANTIATED_OR_RETURN(FunctionName, ReturnValue)																				\
-{																																						\
-	if (!ensure(IsInstantiated()))																														\
-	{																																					\
-		ABILITY_LOG(Error, TEXT("%s: " #FunctionName " cannot be called on a non-instanced ability. Check the instancing policy."), *GetPathName());	\
-		return ReturnValue;																																\
-	}																																					\
-}
-
 /** Structure that defines how an ability will be triggered by external events */
 USTRUCT()
 struct FAbilityTriggerData
@@ -196,18 +186,10 @@ public:
 	UAbilitySystemComponent* GetAbilitySystemComponentFromActorInfo_Ensured() const;
 
 	/** Gets the current actor info bound to this ability - can only be called on instanced abilities. */
-	const FGameplayAbilityActorInfo* GetCurrentActorInfo() const
-	{
-		ENSURE_ABILITY_IS_INSTANTIATED_OR_RETURN(GetCurrentActorInfo, nullptr);
-		return CurrentActorInfo;
-	}
+	const FGameplayAbilityActorInfo* GetCurrentActorInfo() const;
 
 	/** Gets the current activation info bound to this ability - can only be called on instanced abilities. */
-	FGameplayAbilityActivationInfo GetCurrentActivationInfo() const
-	{
-		ENSURE_ABILITY_IS_INSTANTIATED_OR_RETURN(GetCurrentActivationInfo, FGameplayAbilityActivationInfo());
-		return CurrentActivationInfo;
-	}
+	FGameplayAbilityActivationInfo GetCurrentActivationInfo() const;
 
 	/** Gets the current activation info bound to this ability - can only be called on instanced abilities. */
 	FGameplayAbilityActivationInfo& GetCurrentActivationInfoRef()
@@ -217,11 +199,7 @@ public:
 	}
 
 	/** Gets the current AbilitySpecHandle- can only be called on instanced abilities. */
-	FGameplayAbilitySpecHandle GetCurrentAbilitySpecHandle() const
-	{
-		ENSURE_ABILITY_IS_INSTANTIATED_OR_RETURN(GetCurrentAbilitySpecHandle, FGameplayAbilitySpecHandle());
-		return CurrentSpecHandle;
-	}
+	FGameplayAbilitySpecHandle GetCurrentAbilitySpecHandle() const;
 
 	/** Retrieves the actual AbilitySpec for this ability. Can only be called on instanced abilities. */
 	FGameplayAbilitySpec* GetCurrentAbilitySpec() const;
diff --git a/Engine/Plugins/Runtime/GameplayAbilities/Source/GameplayAbilities/Public/AttributeSet.h b/Engine/Plugins/Runtime/GameplayAbilities/Source/GameplayAbilities/Public/AttributeSet.h
--- a/Engine/Plugins/Runtime/GameplayAbilities/Source/GameplayAbilities/Public/AttributeSet.h
+++ b/Engine/Plugins/Runtime/GameplayAbilities/Source/GameplayAbilities/Public/AttributeSet.h
@@ -5,14 +5,12 @@
 #include "CoreMinimal.h"
 #include "Templates/SubclassOf.h"
 #include "UObject/UnrealType.h"
-#include "GameFramework/Actor.h"
-#include "Engine/CurveTable.h"
 #include "Engine/DataTable.h"
-#include "ScalableFloat.h"
 #include "AttributeSet.generated.h"
 
 class UAbilitySystemComponent;
 class UAttributeSet;
+class UCurveTable;
 struct FGameplayAbilityActorInfo;
 struct FAggregator;
 
@@ -238,7 +236,7 @@ public:
 	virtual void InitFromMetaDataTable(const UDataTable* DataTable);
 
 	/** Gets information about owning actor */
-	FORCEINLINE AActor* GetOwningActor() const { return CastChecked<AActor>(GetOuter()); }
+	AActor* GetOwningActor() const;
 	UAbilitySystemComponent* GetOwningAbilitySystemComponent() const;
 	UAbilitySystemComponent* GetOwningAbilitySystemComponentChecked() const;
 	FGameplayAbilityActorInfo* GetActorInfo() const;
diff --git a/Engine/Plugins/Runtime/GameplayAbilities/Source/GameplayAbilities/Public/GameplayAbilitySpec.h b/Engine/Plugins/Runtime/GameplayAbilities/Source/GameplayAbilities/Public/GameplayAbilitySpec.h
--- a/Engine/Plugins/Runtime/GameplayAbilities/Source/GameplayAbilities/Public/GameplayAbilitySpec.h
+++ b/Engine/Plugins/Runtime/GameplayAbilities/Source/GameplayAbilities/Public/GameplayAbilitySpec.h
@@ -10,6 +10,7 @@
 #include "AttributeSet.h"
 #include "GameplayEffectTypes.h"
 #include "GameplayPrediction.h"
+#include "ScalableFloat.h"
 #include "GameplayAbilitySpec.generated.h"
 
 class UAbilitySystemComponent;
@@ -163,13 +164,7 @@ struct GAMEPLAYABILITIES_API FGameplayAbilityActivationInfo
 	{
 	}
 
-	FGameplayAbilityActivationInfo(AActor* InActor)
-		: bCanBeEndedByOtherInstance(false)	
-	{
-		// On Init, we are either Authority or NonAuthority. We haven't been given a PredictionKey and we haven't been confirmed.
-		// NonAuthority essentially means 'I'm not sure what how I'm going to do this yet'.
-		ActivationMode = (InActor->GetLocalRole() == ROLE_Authority ? EGameplayAbilityActivationMode::Authority : EGameplayAbilityActivationMode::NonAuthority);
-	}
+	FGameplayAbilityActivationInfo(AActor* InActor);
 
 	FGameplayAbilityActivationInfo(EGameplayAbilityActivationMode::Type InType)
 		: ActivationMode(InType)
diff --git a/Engine/Plugins/Runtime/GameplayAbilities/Source/GameplayAbilities/Public/GameplayEffect.h b/Engine/Plugins/Runtime/GameplayAbilities/Source/GameplayAbilities/Public/GameplayEffect.h
--- a/Engine/Plugins/Runtime/GameplayAbilities/Source/GameplayAbilities/Public/GameplayEffect.h
+++ b/Engine/Plugins/Runtime/GameplayAbilities/Source/GameplayAbilities/Public/GameplayEffect.h
@@ -33,6 +33,7 @@ struct FGameplayEffectModCallbackData;
 struct FGameplayEffectSpec;
 struct FAggregatorEvaluateParameters;
 struct FAggregatorMod;
+struct FVisualLogEntry;
 
 /** Enumeration outlining the possible gameplay effect magnitude calculation policies. */
 UENUM()
diff --git a/Engine/Plugins/Runtime/GameplayAbilities/Source/GameplayAbilities/Public/GameplayEffectTypes.h b/Engine/Plugins/Runtime/GameplayAbilities/Source/GameplayAbilities/Public/GameplayEffectTypes.h
--- a/Engine/Plugins/Runtime/GameplayAbilities/Source/GameplayAbilities/Public/GameplayEffectTypes.h
+++ b/Engine/Plugins/Runtime/GameplayAbilities/Source/GameplayAbilities/Public/GameplayEffectTypes.h
@@ -8,10 +8,10 @@
 #include "UObject/Class.h"
 #include "Templates/SubclassOf.h"
 #include "Engine/NetSerialization.h"
+#include "GameFramework/Actor.h"
 #include "GameplayTagContainer.h"
 #include "AttributeSet.h"
 #include "PhysicalMaterials/PhysicalMaterial.h"
-#include "AbilitySystemLog.h"
 #include "GameplayEffectTypes.generated.h"
 
 #define SKILL_SYSTEM_AGGREGATOR_DEBUG 1
@@ -28,6 +28,7 @@ class UGameplayAbility;
 struct FActiveGameplayEffect;
 struct FGameplayEffectModCallbackData;
 struct FGameplayEffectSpec;
+struct FHitResult;
 namespace UE::Net
 {
 	struct FGameplayEffectContextHandleAccessorForNetSerializer;
@@ -489,11 +490,7 @@ struct GAMEPLAYABILITIES_API FGameplayEffectContext
 	}
 
 	/** Returns debug string */
-	virtual FString ToString() const
-	{
-		const AActor* InstigatorPtr = Instigator.Get();
-		return (InstigatorPtr ? InstigatorPtr->GetName() : FString(TEXT("NONE")));
-	}
+	virtual FString ToString() const;
 
 	/** Returns the actual struct used for serialization, subclasses must override this! */
 	virtual UScriptStruct* GetScriptStruct() const
@@ -1593,11 +1590,7 @@ struct GAMEPLAYABILITIES_API FGameplayEffectSpecHandle
 		return Data.IsValid();
 	}
 
-	bool NetSerialize(FArchive& Ar, class UPackageMap* Map, bool& bOutSuccess)
-	{
-		ABILITY_LOG(Fatal, TEXT("FGameplayEffectSpecHandle should not be NetSerialized"));
-		return false;
-	}
+	bool NetSerialize(FArchive& Ar, class UPackageMap* Map, bool& bOutSuccess);
 
 	/** Comparison operator */
 	bool operator==(FGameplayEffectSpecHandle const& Other) const
@@ -1645,22 +1638,7 @@ struct GAMEPLAYABILITIES_API FMinimalReplicationTagCountMap
 		TagMap.FindOrAdd(Tag)++;
 	}
 
-	void RemoveTag(const FGameplayTag& Tag)
-	{
-		MapID++;
-		int32& Count = TagMap.FindOrAdd(Tag);
-		Count--;
-		if (Count == 0)
-		{
-			// Remove from map so that we do not replicate
-			TagMap.Remove(Tag);
-		}
-		else if (Count < 0)
-		{
-			ABILITY_LOG(Error, TEXT("FMinimalReplicationTagCountMap::RemoveTag called on Tag %s and count is now < 0"), *Tag.ToString());
-			Count = 0;
-		}
-	}
+	void RemoveTag(const FGameplayTag& Tag);
 
 	void AddTags(const FGameplayTagContainer& Container)
 	{
