Fixed non unity compile errors

diff --git a/Engine/Plugins/AI/MLAdapter/Source/MLAdapter/Private/Sensors/MLAdapterSensor_Attribute.cpp b/Engine/Plugins/AI/MLAdapter/Source/MLAdapter/Private/Sensors/MLAdapterSensor_Attribute.cpp
--- a/Engine/Plugins/AI/MLAdapter/Source/MLAdapter/Private/Sensors/MLAdapterSensor_Attribute.cpp
+++ b/Engine/Plugins/AI/MLAdapter/Source/MLAdapter/Private/Sensors/MLAdapterSensor_Attribute.cpp
@@ -2,6 +2,7 @@
 
 #include "Sensors/MLAdapterSensor_Attribute.h"
 #include "AttributeSet.h"
+#include "GameFramework/Actor.h"
 
 
 UMLAdapterSensor_Attribute::UMLAdapterSensor_Attribute(const FObjectInitializer& ObjectInitializer)
diff --git a/Engine/Plugins/Runtime/GameplayAbilities/Source/GameplayAbilities/Private/AbilitySystemGlobals.cpp b/Engine/Plugins/Runtime/GameplayAbilities/Source/GameplayAbilities/Private/AbilitySystemGlobals.cpp
--- a/Engine/Plugins/Runtime/GameplayAbilities/Source/GameplayAbilities/Private/AbilitySystemGlobals.cpp
+++ b/Engine/Plugins/Runtime/GameplayAbilities/Source/GameplayAbilities/Private/AbilitySystemGlobals.cpp
@@ -6,6 +6,7 @@
 #include "GameplayCueInterface.h"
 #include "AbilitySystemComponent.h"
 #include "AbilitySystemInterface.h"
+#include "AbilitySystemLog.h"
 #include "GameFramework/PlayerController.h"
 #include "GameplayCueManager.h"
 #include "GameplayTagResponseTable.h"
