Fixed a prediction issue when using AbilityTask_NetworkSyncPoint. Abilities executing gameplay cues on a local predicted ability were not ever playing the gameplay cue locally.

Repro:
1) Create a burst gameplay cue that does some simple particle effect and gameplay cue tag to associate with it
2) Create a Local Predicted/Instanced Per Actor GA that when it activates will start a Delay Node (2s was my test) then does a Wait Net Sync (Wait Both) after the sync call ExecuteGameplayCueOnOwner for the gameplay cue tag
3) Grant then activate the GA on the client

If running multiple clients all simulated clients see the burst cue fine, but the local client won't

Reason and then fix:
In this code the ScopePrediction was being made properly and passing from client to server so all predicted paths properly ignored execution assuming the client will/already handle those executions. But, because of this node being "latent" as soon as it finishes executing waiting for the delegate to fire it lost the prediction key made during local activation, and thus wouldn't perform any local predicted executions. By moving the prediction key to only when the server event fires and then just before broadcasting the delegate for the action to resume the ability we properly setup a ScopedPredictionKey allowing local clients to properly run predicted executions.

#abilitysystem
#rnx

diff --git a/Engine/Plugins/Runtime/GameplayAbilities/Source/GameplayAbilities/Private/Abilities/Tasks/AbilityTask_NetworkSyncPoint.cpp b/Engine/Plugins/Runtime/GameplayAbilities/Source/GameplayAbilities/Private/Abilities/Tasks/AbilityTask_NetworkSyncPoint.cpp
--- a/Engine/Plugins/Runtime/GameplayAbilities/Source/GameplayAbilities/Private/Abilities/Tasks/AbilityTask_NetworkSyncPoint.cpp
+++ b/Engine/Plugins/Runtime/GameplayAbilities/Source/GameplayAbilities/Private/Abilities/Tasks/AbilityTask_NetworkSyncPoint.cpp
@@ -29,27 +29,29 @@ UAbilityTask_NetworkSyncPoint* UAbilityTask_NetworkSyncPoint::WaitNetSync(class
 
 void UAbilityTask_NetworkSyncPoint::Activate()
 {
-	FScopedPredictionWindow ScopedPrediction(AbilitySystemComponent.Get(), IsPredictingClient());
-
 	if (AbilitySystemComponent.IsValid())
 	{
 		if (IsPredictingClient())
 		{
-			if (SyncType != EAbilityTaskNetSyncType::OnlyServerWait )
+			if (SyncType != EAbilityTaskNetSyncType::OnlyServerWait)
 			{
 				// As long as we are waiting (!= OnlyServerWait), listen for the GenericSignalFromServer event
 				ReplicatedEventToListenFor = EAbilityGenericReplicatedEvent::GenericSignalFromServer;
 			}
 			if (SyncType != EAbilityTaskNetSyncType::OnlyClientWait)
 			{
+				// @note: When on a predicting client we need to make sure we send off the scoped prediction key to the server so things like gameplay
+				// cue's that replicate down will continue to not execute on a client that is predicting abilities.
+				FScopedPredictionWindow ScopedPrediction(AbilitySystemComponent.Get(), IsPredictingClient());
+
 				// As long as the server is waiting (!= OnlyClientWait), send the Server and RPC for this signal
 				AbilitySystemComponent->ServerSetReplicatedEvent(EAbilityGenericReplicatedEvent::GenericSignalFromClient, GetAbilitySpecHandle(), GetActivationPredictionKey(), AbilitySystemComponent->ScopedPredictionKey);
 			}
-			
+
 		}
 		else if (IsForRemoteClient())
 		{
-			if (SyncType != EAbilityTaskNetSyncType::OnlyClientWait )
+			if (SyncType != EAbilityTaskNetSyncType::OnlyClientWait)
 			{
 				// As long as we are waiting (!= OnlyClientWait), listen for the GenericSignalFromClient event
 				ReplicatedEventToListenFor = EAbilityGenericReplicatedEvent::GenericSignalFromClient;
@@ -77,6 +79,12 @@ void UAbilityTask_NetworkSyncPoint::SyncFinished()
 {
 	if (IsValid(this))
 	{
+		// @note: When we finish the task to let the ability resume execution we need to ensure predicting clients have a valid
+		// scoped prediction key so they can continue to execute things like gameplay cues. This is because the server will pass
+		// along a client prediction key through its replication chain ensuring predicted abilities don't execute predicted
+		// features (i.e. gameplay cues)
+		FScopedPredictionWindow ScopedPrediction(AbilitySystemComponent.Get(), IsPredictingClient());
+
 		if (ShouldBroadcastAbilityTaskDelegates())
 		{
 			OnSync.Broadcast();
