Iris
- Tagged not replicated properties as NotReplicated to be able to use default struct netserialization

diff --git a/Engine/Plugins/Runtime/GameplayAbilities/Source/GameplayAbilities/Public/Abilities/GameplayAbilityTargetTypes.h b/Engine/Plugins/Runtime/GameplayAbilities/Source/GameplayAbilities/Public/Abilities/GameplayAbilityTargetTypes.h
--- a/Engine/Plugins/Runtime/GameplayAbilities/Source/GameplayAbilities/Public/Abilities/GameplayAbilityTargetTypes.h
+++ b/Engine/Plugins/Runtime/GameplayAbilities/Source/GameplayAbilities/Public/Abilities/GameplayAbilityTargetTypes.h
@@ -610,7 +610,7 @@ struct GAMEPLAYABILITIES_API FGameplayAbilityTargetData_SingleTargetHit : public
 	UPROPERTY()
 	FHitResult	HitResult;
 
-	UPROPERTY()
+	UPROPERTY(NotReplicated)
 	bool bHitReplaced = false;
 
 	bool NetSerialize(FArchive& Ar, class UPackageMap* Map, bool& bOutSuccess);
diff --git a/Engine/Source/Runtime/Engine/Classes/GameFramework/UpdateLevelVisibilityLevelInfo.h b/Engine/Source/Runtime/Engine/Classes/GameFramework/UpdateLevelVisibilityLevelInfo.h
--- a/Engine/Source/Runtime/Engine/Classes/GameFramework/UpdateLevelVisibilityLevelInfo.h
+++ b/Engine/Source/Runtime/Engine/Classes/GameFramework/UpdateLevelVisibilityLevelInfo.h
@@ -91,6 +91,7 @@ struct ENGINE_API FUpdateLevelVisibilityLevelInfo
 	uint32 bTryMakeVisible : 1;
 
 	/** Skip connection close if level can't be found (not net serialized) */
+	UPROPERTY(NotReplicated)
 	uint32 bSkipCloseOnError : 1;
 
 	bool NetSerialize(FArchive& Ar, UPackageMap* PackageMap, bool& bOutSuccess);
