Extract GameplayAbilitySecHandle and ActiveGameplayEffectHandle in their own file

diff --git a/Engine/Plugins/Runtime/GameplayAbilities/Source/GameplayAbilities/Private/ActiveGameplayEffectHandle.cpp b/Engine/Plugins/Runtime/GameplayAbilities/Source/GameplayAbilities/Private/ActiveGameplayEffectHandle.cpp
new file mode 100644
--- /dev/null
+++ b/Engine/Plugins/Runtime/GameplayAbilities/Source/GameplayAbilities/Private/ActiveGameplayEffectHandle.cpp
@@ -0,0 +1,53 @@
+// Copyright Epic Games, Inc. All Rights Reserved.
+
+#include "ActiveGameplayEffectHandle.h"
+#include "AbilitySystemComponent.h"
+
+namespace GlobalActiveGameplayEffectHandles
+{
+	static TMap<FActiveGameplayEffectHandle, TWeakObjectPtr<UAbilitySystemComponent>>	Map;
+}
+
+void FActiveGameplayEffectHandle::ResetGlobalHandleMap()
+{
+	GlobalActiveGameplayEffectHandles::Map.Reset();
+}
+
+FActiveGameplayEffectHandle FActiveGameplayEffectHandle::GenerateNewHandle(UAbilitySystemComponent* OwningComponent)
+{
+	static int32 GHandleID=0;
+	FActiveGameplayEffectHandle NewHandle(GHandleID++);
+
+	TWeakObjectPtr<UAbilitySystemComponent> WeakPtr(OwningComponent);
+
+	GlobalActiveGameplayEffectHandles::Map.Add(NewHandle, WeakPtr);
+
+	return NewHandle;
+}
+
+UAbilitySystemComponent* FActiveGameplayEffectHandle::GetOwningAbilitySystemComponent()
+{
+	TWeakObjectPtr<UAbilitySystemComponent>* Ptr = GlobalActiveGameplayEffectHandles::Map.Find(*this);
+	if (Ptr)
+	{
+		return Ptr->Get();
+	}
+
+	return nullptr;	
+}
+
+const UAbilitySystemComponent* FActiveGameplayEffectHandle::GetOwningAbilitySystemComponent() const
+{
+	TWeakObjectPtr<UAbilitySystemComponent>* Ptr = GlobalActiveGameplayEffectHandles::Map.Find(*this);
+	if (Ptr)
+	{
+		return Ptr->Get();
+	}
+
+	return nullptr;
+}
+
+void FActiveGameplayEffectHandle::RemoveFromGlobalMap()
+{
+	GlobalActiveGameplayEffectHandles::Map.Remove(*this);
+}
diff --git a/Engine/Plugins/Runtime/GameplayAbilities/Source/GameplayAbilities/Private/GameplayAbilitySpecHandle.cpp b/Engine/Plugins/Runtime/GameplayAbilities/Source/GameplayAbilities/Private/GameplayAbilitySpecHandle.cpp
new file mode 100644
--- /dev/null
+++ b/Engine/Plugins/Runtime/GameplayAbilities/Source/GameplayAbilities/Private/GameplayAbilitySpecHandle.cpp
@@ -0,0 +1,12 @@
+// Copyright Epic Games, Inc. All Rights Reserved.
+
+#include "GameplayAbilitySpecHandle.h"
+
+// ----------------------------------------------------
+
+void FGameplayAbilitySpecHandle::GenerateNewHandle()
+{
+	// Must be in C++ to avoid duplicate statics accross execution units
+	static int32 GHandle = 1;
+	Handle = GHandle++;
+}
diff --git a/Engine/Plugins/Runtime/GameplayAbilities/Source/GameplayAbilities/Private/GameplayAbilityTypes.cpp b/Engine/Plugins/Runtime/GameplayAbilities/Source/GameplayAbilities/Private/GameplayAbilityTypes.cpp
--- a/Engine/Plugins/Runtime/GameplayAbilities/Source/GameplayAbilities/Private/GameplayAbilityTypes.cpp
+++ b/Engine/Plugins/Runtime/GameplayAbilities/Source/GameplayAbilities/Private/GameplayAbilityTypes.cpp
@@ -289,14 +289,6 @@ FGameplayAbilitySpec::FGameplayAbilitySpec(FGameplayAbilitySpecDef& InDef, int32
 	Level = InDef.LevelScalableFloat.GetValueAtLevel(InGameplayEffectLevel, &ContextString);
 }
 
-// ----------------------------------------------------
-
-void FGameplayAbilitySpecHandle::GenerateNewHandle()
-{
-	// Must be in C++ to avoid duplicate statics accross execution units
-	static int32 GHandle = 1;
-	Handle = GHandle++;
-}
 
 // ----------------------------------------------------
 
diff --git a/Engine/Plugins/Runtime/GameplayAbilities/Source/GameplayAbilities/Private/GameplayEffect.cpp b/Engine/Plugins/Runtime/GameplayAbilities/Source/GameplayAbilities/Private/GameplayEffect.cpp
--- a/Engine/Plugins/Runtime/GameplayAbilities/Source/GameplayAbilities/Private/GameplayEffect.cpp
+++ b/Engine/Plugins/Runtime/GameplayAbilities/Source/GameplayAbilities/Private/GameplayEffect.cpp
@@ -4765,60 +4765,7 @@ void FActiveGameplayEffectsContainer::CloneFrom(const FActiveGameplayEffectsCont
 	}
 }
 
-// --------------------------------------------------------------------------------------------------------------------------------------------------------
-//
-//	Misc
-//
-// --------------------------------------------------------------------------------------------------------------------------------------------------------
 
-namespace GlobalActiveGameplayEffectHandles
-{
-	static TMap<FActiveGameplayEffectHandle, TWeakObjectPtr<UAbilitySystemComponent>>	Map;
-}
-
-void FActiveGameplayEffectHandle::ResetGlobalHandleMap()
-{
-	GlobalActiveGameplayEffectHandles::Map.Reset();
-}
-
-FActiveGameplayEffectHandle FActiveGameplayEffectHandle::GenerateNewHandle(UAbilitySystemComponent* OwningComponent)
-{
-	static int32 GHandleID=0;
-	FActiveGameplayEffectHandle NewHandle(GHandleID++);
-
-	TWeakObjectPtr<UAbilitySystemComponent> WeakPtr(OwningComponent);
-
-	GlobalActiveGameplayEffectHandles::Map.Add(NewHandle, WeakPtr);
-
-	return NewHandle;
-}
-
-UAbilitySystemComponent* FActiveGameplayEffectHandle::GetOwningAbilitySystemComponent()
-{
-	TWeakObjectPtr<UAbilitySystemComponent>* Ptr = GlobalActiveGameplayEffectHandles::Map.Find(*this);
-	if (Ptr)
-	{
-		return Ptr->Get();
-	}
-
-	return nullptr;	
-}
-
-const UAbilitySystemComponent* FActiveGameplayEffectHandle::GetOwningAbilitySystemComponent() const
-{
-	TWeakObjectPtr<UAbilitySystemComponent>* Ptr = GlobalActiveGameplayEffectHandles::Map.Find(*this);
-	if (Ptr)
-	{
-		return Ptr->Get();
-	}
-
-	return nullptr;
-}
-
-void FActiveGameplayEffectHandle::RemoveFromGlobalMap()
-{
-	GlobalActiveGameplayEffectHandles::Map.Remove(*this);
-}
 
 // -----------------------------------------------------------------
 
diff --git a/Engine/Plugins/Runtime/GameplayAbilities/Source/GameplayAbilities/Public/ActiveGameplayEffectHandle.h b/Engine/Plugins/Runtime/GameplayAbilities/Source/GameplayAbilities/Public/ActiveGameplayEffectHandle.h
new file mode 100644
--- /dev/null
+++ b/Engine/Plugins/Runtime/GameplayAbilities/Source/GameplayAbilities/Public/ActiveGameplayEffectHandle.h
@@ -0,0 +1,88 @@
+// Copyright Epic Games, Inc. All Rights Reserved.
+
+#pragma once
+
+#include "ActiveGameplayEffectHandle.generated.h"
+
+class UAbilitySystemComponent;
+
+/**
+ * This handle is required for things outside of FActiveGameplayEffectsContainer to refer to a specific active GameplayEffect
+ *	For example if a skill needs to create an active effect and then destroy that specific effect that it created, it has to do so
+ *	through a handle. a pointer or index into the active list is not sufficient.
+ */
+USTRUCT(BlueprintType)
+struct GAMEPLAYABILITIES_API FActiveGameplayEffectHandle
+{
+	GENERATED_USTRUCT_BODY()
+
+	FActiveGameplayEffectHandle()
+		: Handle(INDEX_NONE),
+		bPassedFiltersAndWasExecuted(false)
+	{
+	}
+
+	FActiveGameplayEffectHandle(int32 InHandle)
+		: Handle(InHandle),
+		bPassedFiltersAndWasExecuted(true)
+	{
+	}
+
+	/** True if this is tracking an active ongoing gameplay effect */
+	bool IsValid() const
+	{
+		return Handle != INDEX_NONE;
+	}
+
+	/** True if this applied a gameplay effect. This can be true when it is not valid if it applied an instant effect */
+	bool WasSuccessfullyApplied() const
+	{
+		return bPassedFiltersAndWasExecuted;
+	}
+
+	/** Creates a new handle, will be set to successfully applied */
+	static FActiveGameplayEffectHandle GenerateNewHandle(UAbilitySystemComponent* OwningComponent);
+
+	/** Resets the map that supports GetOwningAbilitySystemComponent */
+	static void ResetGlobalHandleMap();
+
+	/** Returns the ability system component that created this handle */
+	UAbilitySystemComponent* GetOwningAbilitySystemComponent();
+	const UAbilitySystemComponent* GetOwningAbilitySystemComponent() const;
+
+	/** Remove this from the GetOwningAbilitySystemComponent map */
+	void RemoveFromGlobalMap();
+
+	bool operator==(const FActiveGameplayEffectHandle& Other) const
+	{
+		return Handle == Other.Handle;
+	}
+
+	bool operator!=(const FActiveGameplayEffectHandle& Other) const
+	{
+		return Handle != Other.Handle;
+	}
+
+	friend uint32 GetTypeHash(const FActiveGameplayEffectHandle& InHandle)
+	{
+		return InHandle.Handle;
+	}
+
+	FString ToString() const
+	{
+		return FString::Printf(TEXT("%d"), Handle);
+	}
+
+	void Invalidate()
+	{
+		Handle = INDEX_NONE;
+	}
+
+private:
+
+	UPROPERTY()
+	int32 Handle;
+
+	UPROPERTY()
+	bool bPassedFiltersAndWasExecuted;
+};
\ No newline at end of file
diff --git a/Engine/Plugins/Runtime/GameplayAbilities/Source/GameplayAbilities/Public/GameplayAbilitySpec.h b/Engine/Plugins/Runtime/GameplayAbilities/Source/GameplayAbilities/Public/GameplayAbilitySpec.h
--- a/Engine/Plugins/Runtime/GameplayAbilities/Source/GameplayAbilities/Public/GameplayAbilitySpec.h
+++ b/Engine/Plugins/Runtime/GameplayAbilities/Source/GameplayAbilities/Public/GameplayAbilitySpec.h
@@ -8,6 +8,7 @@
 #include "Templates/SubclassOf.h"
 #include "Engine/NetSerialization.h"
 #include "AttributeSet.h"
+#include "GameplayAbilitySpecHandle.h"
 #include "GameplayEffectTypes.h"
 #include "GameplayPrediction.h"
 #include "ScalableFloat.h"
@@ -16,56 +17,6 @@
 class UAbilitySystemComponent;
 class UGameplayAbility;
 
-/**
- *	This file exists in addition so that GameplayEffect.h can use FGameplayAbilitySpec without having to include GameplayAbilityTypes.h which has depancies on
- *	GameplayEffect.h
- */
-
-/** Handle that points to a specific granted ability. These are globally unique */
-USTRUCT(BlueprintType)
-struct FGameplayAbilitySpecHandle
-{
-	GENERATED_USTRUCT_BODY()
-
-	FGameplayAbilitySpecHandle()
-		: Handle(INDEX_NONE)
-	{
-	}
-
-	/** True if GenerateNewHandle was called on this handle */
-	bool IsValid() const
-	{
-		return Handle != INDEX_NONE;
-	}
-
-	/** Sets this to a valid handle */
-	void GenerateNewHandle();
-
-	bool operator==(const FGameplayAbilitySpecHandle& Other) const
-	{
-		return Handle == Other.Handle;
-	}
-
-	bool operator!=(const FGameplayAbilitySpecHandle& Other) const
-	{
-		return Handle != Other.Handle;
-	}
-
-	friend uint32 GetTypeHash(const FGameplayAbilitySpecHandle& SpecHandle)
-	{
-		return ::GetTypeHash(SpecHandle.Handle);
-	}
-
-	FString ToString() const
-	{
-		return IsValid() ? FString::FromInt(Handle) : TEXT("Invalid");
-	}
-
-private:
-
-	UPROPERTY()
-	int32 Handle;
-};
 
 /** Describes the status of activating this ability, this is updated as prediction is handled */
 UENUM(BlueprintType)
diff --git a/Engine/Plugins/Runtime/GameplayAbilities/Source/GameplayAbilities/Public/GameplayAbilitySpecHandle.h b/Engine/Plugins/Runtime/GameplayAbilities/Source/GameplayAbilities/Public/GameplayAbilitySpecHandle.h
new file mode 100644
--- /dev/null
+++ b/Engine/Plugins/Runtime/GameplayAbilities/Source/GameplayAbilities/Public/GameplayAbilitySpecHandle.h
@@ -0,0 +1,57 @@
+// Copyright Epic Games, Inc. All Rights Reserved.
+
+#pragma once
+
+#include "GameplayAbilitySpecHandle.generated.h"
+
+
+/**
+ *	This file exists in addition so that GameplayEffect.h can use FGameplayAbilitySpec without having to include GameplayAbilityTypes.h which has depancies on
+ *	GameplayEffect.h
+ */
+
+/** Handle that points to a specific granted ability. These are globally unique */
+USTRUCT(BlueprintType)
+struct FGameplayAbilitySpecHandle
+{
+	GENERATED_USTRUCT_BODY()
+
+	FGameplayAbilitySpecHandle()
+		: Handle(INDEX_NONE)
+	{
+	}
+
+	/** True if GenerateNewHandle was called on this handle */
+	bool IsValid() const
+	{
+		return Handle != INDEX_NONE;
+	}
+
+	/** Sets this to a valid handle */
+	void GenerateNewHandle();
+
+	bool operator==(const FGameplayAbilitySpecHandle& Other) const
+	{
+		return Handle == Other.Handle;
+	}
+
+	bool operator!=(const FGameplayAbilitySpecHandle& Other) const
+	{
+		return Handle != Other.Handle;
+	}
+
+	friend uint32 GetTypeHash(const FGameplayAbilitySpecHandle& SpecHandle)
+	{
+		return ::GetTypeHash(SpecHandle.Handle);
+	}
+
+	FString ToString() const
+	{
+		return IsValid() ? FString::FromInt(Handle) : TEXT("Invalid");
+	}
+
+private:
+
+	UPROPERTY()
+	int32 Handle;
+};
\ No newline at end of file
diff --git a/Engine/Plugins/Runtime/GameplayAbilities/Source/GameplayAbilities/Public/GameplayEffectTypes.h b/Engine/Plugins/Runtime/GameplayAbilities/Source/GameplayAbilities/Public/GameplayEffectTypes.h
--- a/Engine/Plugins/Runtime/GameplayAbilities/Source/GameplayAbilities/Public/GameplayEffectTypes.h
+++ b/Engine/Plugins/Runtime/GameplayAbilities/Source/GameplayAbilities/Public/GameplayEffectTypes.h
@@ -2,6 +2,7 @@
 
 #pragma once
 
+#include "ActiveGameplayEffectHandle.h"
 #include "CoreMinimal.h"
 #include "UObject/ObjectMacros.h"
 #include "UObject/Object.h"
@@ -163,88 +164,6 @@ enum class EGameplayEffectStackingType : uint8
 };
 
 
-/**
- * This handle is required for things outside of FActiveGameplayEffectsContainer to refer to a specific active GameplayEffect
- *	For example if a skill needs to create an active effect and then destroy that specific effect that it created, it has to do so
- *	through a handle. a pointer or index into the active list is not sufficient.
- */
-USTRUCT(BlueprintType)
-struct GAMEPLAYABILITIES_API FActiveGameplayEffectHandle
-{
-	GENERATED_USTRUCT_BODY()
-
-	FActiveGameplayEffectHandle()
-		: Handle(INDEX_NONE),
-		bPassedFiltersAndWasExecuted(false)
-	{
-
-	}
-
-	FActiveGameplayEffectHandle(int32 InHandle)
-		: Handle(InHandle),
-		bPassedFiltersAndWasExecuted(true)
-	{
-
-	}
-
-	/** True if this is tracking an active ongoing gameplay effect */
-	bool IsValid() const
-	{
-		return Handle != INDEX_NONE;
-	}
-
-	/** True if this applied a gameplay effect. This can be true when it is not valid if it applied an instant effect */
-	bool WasSuccessfullyApplied() const
-	{
-		return bPassedFiltersAndWasExecuted;
-	}
-
-	/** Creates a new handle, will be set to successfully applied */
-	static FActiveGameplayEffectHandle GenerateNewHandle(UAbilitySystemComponent* OwningComponent);
-
-	/** Resets the map that supports GetOwningAbilitySystemComponent */
-	static void ResetGlobalHandleMap();
-
-	/** Returns the ability system component that created this handle */
-	UAbilitySystemComponent* GetOwningAbilitySystemComponent();
-	const UAbilitySystemComponent* GetOwningAbilitySystemComponent() const;
-
-	/** Remove this from the GetOwningAbilitySystemComponent map */
-	void RemoveFromGlobalMap();
-
-	bool operator==(const FActiveGameplayEffectHandle& Other) const
-	{
-		return Handle == Other.Handle;
-	}
-
-	bool operator!=(const FActiveGameplayEffectHandle& Other) const
-	{
-		return Handle != Other.Handle;
-	}
-
-	friend uint32 GetTypeHash(const FActiveGameplayEffectHandle& InHandle)
-	{
-		return InHandle.Handle;
-	}
-
-	FString ToString() const
-	{
-		return FString::Printf(TEXT("%d"), Handle);
-	}
-
-	void Invalidate()
-	{
-		Handle = INDEX_NONE;
-	}
-
-private:
-
-	UPROPERTY()
-	int32 Handle;
-
-	UPROPERTY()
-	bool bPassedFiltersAndWasExecuted;
-};
 
 
 /** Data that describes what happened in an attribute modification. This is passed to ability set callbacks */
