Iris
- Added method to validate traits when setting up a forwarding serializer
- Added validation for FGameplayEffectContextNetserializer and FGameplayAbilityRepAnimMontageNetSerializer
- Added missing bHasDynamicState trait from GameplayEffectContextNetserializer
- Fixed spelling of UE_NET_IMPLEMENT_NETSERIALIZER_REGISTRY_DELEGATES

diff --git a/Engine/Plugins/Runtime/GameplayAbilities/Source/GameplayAbilities/Private/Serialization/GameplayAbilityRepAnimMontageNetSerializer.cpp b/Engine/Plugins/Runtime/GameplayAbilities/Source/GameplayAbilities/Private/Serialization/GameplayAbilityRepAnimMontageNetSerializer.cpp
--- a/Engine/Plugins/Runtime/GameplayAbilities/Source/GameplayAbilities/Private/Serialization/GameplayAbilityRepAnimMontageNetSerializer.cpp
+++ b/Engine/Plugins/Runtime/GameplayAbilities/Source/GameplayAbilities/Private/Serialization/GameplayAbilityRepAnimMontageNetSerializer.cpp
@@ -336,7 +336,7 @@ bool FGameplayAbilityRepAnimMontageNetSerializer::IsEqual(FNetSerializationConte
 			return false;
 		}
 
-		if (Value0.ReplicationFlags != Value1.ReplicationFlags || Value0.Position != Value1.ReplicationFlags || Value0.SectionIdToPlay != Value1.SectionIdToPlay)
+		if (Value0.ReplicationFlags != Value1.ReplicationFlags || Value0.Position != Value1.Position || Value0.SectionIdToPlay != Value1.SectionIdToPlay)
 		{
 			return false;
 		}
@@ -447,6 +447,9 @@ void FGameplayAbilityRepAnimMontageNetSerializer::FNetSerializerRegistryDelegate
 	const FReplicationStateDescriptor* Descriptor = StructNetSerializerConfigForBase.StateDescriptor.GetReference();
 	check(Descriptor != nullptr);
 
+	// Verify traits
+	ValidateForwardingNetSerializerTraits(&UE_NET_GET_SERIALIZER(FGameplayAbilityRepAnimMontageNetSerializer), Descriptor->Traits);
+	
 	// Validate our assumptions regarding quantized state size and alignment.
 	constexpr SIZE_T OffsetOfGameplayAbilityRepAnimMontage = offsetof(FQuantizedType, GameplayAbilityRepAnimMontage);
 	if ((sizeof(FQuantizedType::GameplayAbilityRepAnimMontage) < Descriptor->InternalSize) || (((OffsetOfGameplayAbilityRepAnimMontage/Descriptor->InternalAlignment)*Descriptor->InternalAlignment) != OffsetOfGameplayAbilityRepAnimMontage))
diff --git a/Engine/Plugins/Runtime/GameplayAbilities/Source/GameplayAbilities/Private/Serialization/GameplayEffectContextNetSerializer.cpp b/Engine/Plugins/Runtime/GameplayAbilities/Source/GameplayAbilities/Private/Serialization/GameplayEffectContextNetSerializer.cpp
--- a/Engine/Plugins/Runtime/GameplayAbilities/Source/GameplayAbilities/Private/Serialization/GameplayEffectContextNetSerializer.cpp
+++ b/Engine/Plugins/Runtime/GameplayAbilities/Source/GameplayAbilities/Private/Serialization/GameplayEffectContextNetSerializer.cpp
@@ -7,6 +7,7 @@
 
 #include "GameplayEffectTypes.h"
 #include "Net/Core/NetBitArray.h"
+#include "Net/Core/Trace/NetTrace.h"
 #include "Iris/Serialization/NetBitStreamReader.h"
 #include "Iris/Serialization/NetBitStreamWriter.h"
 #include "Iris/Serialization/NetSerializerDelegates.h"
@@ -20,7 +21,6 @@ uint16 FGameplayEffectContextAccessorForNetSerializer::PropertyToMemberIndex[FGa
 namespace UE::Net
 {
 
-// 
 struct FGameplayEffectContextNetSerializer
 {
 	// Version
@@ -29,6 +29,7 @@ struct FGameplayEffectContextNetSerializer
 	// Traits
 	static constexpr bool bIsForwardingSerializer = true; // Triggers asserts if a function is missing
 	static constexpr bool bHasCustomNetReference = true;
+	static constexpr bool bHasDynamicState = true;
 
 	// Types
 	enum EReplicationFlags : uint32
@@ -96,7 +97,8 @@ const FGameplayEffectContextNetSerializer::ConfigType FGameplayEffectContextNetS
 FGameplayEffectContextNetSerializer::FNetSerializerRegistryDelegates FGameplayEffectContextNetSerializer::NetSerializerRegistryDelegates;
 FStructNetSerializerConfig FGameplayEffectContextNetSerializer::StructNetSerializerConfigForGE;
 FStructNetSerializerConfig FGameplayEffectContextNetSerializer::StructNetSerializerConfigForHitResult;
-const FNetSerializer* FGameplayEffectContextNetSerializer::StructNetSerializer = &UE_NET_GET_SERIALIZER(FStructNetSerializer);
+ const FNetSerializer* FGameplayEffectContextNetSerializer::StructNetSerializer = &UE_NET_GET_SERIALIZER(FStructNetSerializer);
+
 EReplicationStateTraits FGameplayEffectContextNetSerializer::GEStateTraits;
 EReplicationStateTraits FGameplayEffectContextNetSerializer::HitResultStateTraits;
 
@@ -115,6 +117,7 @@ void FGameplayEffectContextNetSerializer::Serialize(FNetSerializationContext& Co
 		const FReplicationStateDescriptor* Descriptor = StructNetSerializerConfigForGE.StateDescriptor.GetReference();
 		const FReplicationStateMemberDescriptor* MemberDescriptors = Descriptor->MemberDescriptors;
 		const FReplicationStateMemberSerializerDescriptor* MemberSerializerDescriptors = Descriptor->MemberSerializerDescriptors;
+		const FReplicationStateMemberDebugDescriptor* MemberDebugDescriptors = Descriptor->MemberDebugDescriptors;
 		
 		uint32 MemberMaskStorage = ~0U;
 		FNetBitArrayView MemberMask(&MemberMaskStorage, FGameplayEffectContextAccessorForNetSerializer::EPropertyName::PropertyName_ReplicatedPropertyCount, FNetBitArrayView::NoResetNoValidate);
@@ -131,6 +134,9 @@ void FGameplayEffectContextNetSerializer::Serialize(FNetSerializationContext& Co
 			const FReplicationStateMemberDescriptor& MemberDescriptor = MemberDescriptors[MemberIndex];
 			const FReplicationStateMemberSerializerDescriptor& MemberSerializerDescriptor = MemberSerializerDescriptors[MemberIndex];
 
+			UE_NET_TRACE_DYNAMIC_NAME_SCOPE(MemberDebugDescriptors[MemberIndex].DebugName, *Context.GetBitStreamWriter(), Context.GetTraceCollector(), ENetTraceVerbosity::Verbose);
+			UE_NET_TRACE_DYNAMIC_NAME_SCOPE(MemberSerializerDescriptor.Serializer->Name, *Context.GetBitStreamWriter(), Context.GetTraceCollector(), ENetTraceVerbosity::VeryVerbose);
+
 			FNetSerializeArgs MemberSerializeArgs;
 			MemberSerializeArgs.NetSerializerConfig = MemberSerializerDescriptor.SerializerConfig;
 			MemberSerializeArgs.Source = reinterpret_cast<NetSerializerValuePointer>(&Value.EffectContext) + MemberDescriptor.InternalMemberOffset;
@@ -162,6 +168,7 @@ void FGameplayEffectContextNetSerializer::Deserialize(FNetSerializationContext&
 		const FReplicationStateDescriptor* Descriptor = StructNetSerializerConfigForGE.StateDescriptor.GetReference();
 		const FReplicationStateMemberDescriptor* MemberDescriptors = Descriptor->MemberDescriptors;
 		const FReplicationStateMemberSerializerDescriptor* MemberSerializerDescriptors = Descriptor->MemberSerializerDescriptors;
+		const FReplicationStateMemberDebugDescriptor* MemberDebugDescriptors = Descriptor->MemberDebugDescriptors;
 		
 		uint32 MemberMaskStorage = ~0U;
 		FNetBitArrayView MemberMask(&MemberMaskStorage, FGameplayEffectContextAccessorForNetSerializer::EPropertyName::PropertyName_ReplicatedPropertyCount, FNetBitArrayView::NoResetNoValidate);
@@ -178,6 +185,9 @@ void FGameplayEffectContextNetSerializer::Deserialize(FNetSerializationContext&
 			const FReplicationStateMemberDescriptor& MemberDescriptor = MemberDescriptors[MemberIndex];
 			const FReplicationStateMemberSerializerDescriptor& MemberSerializerDescriptor = MemberSerializerDescriptors[MemberIndex];
 
+			UE_NET_TRACE_DYNAMIC_NAME_SCOPE(MemberDebugDescriptors[MemberIndex].DebugName, *Context.GetBitStreamReader(), Context.GetTraceCollector(), ENetTraceVerbosity::Verbose);
+			UE_NET_TRACE_DYNAMIC_NAME_SCOPE(MemberSerializerDescriptor.Serializer->Name, *Context.GetBitStreamReader(), Context.GetTraceCollector(), ENetTraceVerbosity::VeryVerbose);
+
 			FNetDeserializeArgs MemberDeserializeArgs;
 			MemberDeserializeArgs.NetSerializerConfig = MemberSerializerDescriptor.SerializerConfig;
 			MemberDeserializeArgs.Target = reinterpret_cast<NetSerializerValuePointer>(&Target.EffectContext) + MemberDescriptor.InternalMemberOffset;
@@ -559,6 +569,9 @@ void FGameplayEffectContextNetSerializer::FNetSerializerRegistryDelegates::OnPos
 			LowLevelFatalError(TEXT("FQuantizedType::HitResult has size %u but requires size %u and alignment %u."), uint32(sizeof(FQuantizedType::HitResult)), uint32(Descriptor->InternalSize), uint32(Descriptor->InternalAlignment));
 		}
 	}
+
+	// Verify traits
+	ValidateForwardingNetSerializerTraits(&UE_NET_GET_SERIALIZER(FGameplayEffectContextNetSerializer), GEStateTraits | HitResultStateTraits);
 }
 
 }
diff --git a/Engine/Source/Runtime/Experimental/Iris/Core/Private/Iris/ReplicationState/PropertyNetSerializerInfoRegistry.cpp b/Engine/Source/Runtime/Experimental/Iris/Core/Private/Iris/ReplicationState/PropertyNetSerializerInfoRegistry.cpp
--- a/Engine/Source/Runtime/Experimental/Iris/Core/Private/Iris/ReplicationState/PropertyNetSerializerInfoRegistry.cpp
+++ b/Engine/Source/Runtime/Experimental/Iris/Core/Private/Iris/ReplicationState/PropertyNetSerializerInfoRegistry.cpp
@@ -2,6 +2,7 @@
 
 #include "Iris/ReplicationState/PropertyNetSerializerInfoRegistry.h"
 #include "Iris/ReplicationState/EnumPropertyNetSerializerInfo.h"
+#include "Iris/ReplicationState/ReplicationStateDescriptor.h"
 #include "Templates/Sorting.h"
 #include "UObject/UnrealType.h"
 #include "Iris/Serialization/NetSerializers.h"
@@ -200,4 +201,21 @@ const FPropertyNetSerializerInfo* FPropertyNetSerializerInfoRegistry::GetNopNetS
 	return &Private::NopNetSerializerInfo;
 }
 
+void ValidateForwardingNetSerializerTraits(const FNetSerializer* Serializer, EReplicationStateTraits UsedReplicationStateTraits)
+{
+	const ENetSerializerTraits SerializerTraits = Serializer->Traits;
+	if (EnumHasAnyFlags(UsedReplicationStateTraits, EReplicationStateTraits::HasDynamicState) && !EnumHasAnyFlags(SerializerTraits, ENetSerializerTraits::HasDynamicState))
+	{
+		LowLevelFatalError(TEXT("FNetSerializer: %s is using serializer(s) that HasDynamicState without setting trait: static constexpr bool bHasDynamicState = true; in the serializer declaration."), Serializer->Name);
+	}
+	if (EnumHasAnyFlags(UsedReplicationStateTraits, EReplicationStateTraits::HasObjectReference) && !EnumHasAnyFlags(SerializerTraits, ENetSerializerTraits::HasCustomNetReference))
+	{
+		LowLevelFatalError(TEXT("FNetSerializer: %s is using serializer(s) that has trait HasCustomNetReference without setting trait: static constexpr bool bHasCustomNetReference = true; in the serializer declaration."), Serializer->Name);
+	}
+	if (EnumHasAnyFlags(UsedReplicationStateTraits, EReplicationStateTraits::HasConnectionSpecificSerialization) && !EnumHasAnyFlags(SerializerTraits, ENetSerializerTraits::HasConnectionSpecificSerialization))
+	{
+		LowLevelFatalError(TEXT("FNetSerializer: %s is using serializer(s) that has HasConnectionSpecificSerialization without setting trait: static constexpr bool bHasConnectionSpecificSerialization = true; in the serializer declaration."), Serializer->Name);
+	}
+}
+
 }
diff --git a/Engine/Source/Runtime/Experimental/Iris/Core/Public/Iris/ReplicationState/PropertyNetSerializerInfoRegistry.h b/Engine/Source/Runtime/Experimental/Iris/Core/Public/Iris/ReplicationState/PropertyNetSerializerInfoRegistry.h
--- a/Engine/Source/Runtime/Experimental/Iris/Core/Public/Iris/ReplicationState/PropertyNetSerializerInfoRegistry.h
+++ b/Engine/Source/Runtime/Experimental/Iris/Core/Public/Iris/ReplicationState/PropertyNetSerializerInfoRegistry.h
@@ -13,6 +13,13 @@
 namespace UE::Net
 {
 
+enum class EReplicationStateTraits : uint32;
+
+}
+
+namespace UE::Net
+{
+
 /**
  * Currently we require each supported type to register FPropertyNetSerializerInfo 
  * It provides information on what NetSerializer to use for which property and how to build the required NetSerializer config which is used when we build the dynamic descriptor
@@ -120,6 +127,9 @@ private:
 	const FName PropertyFName;
 };
 
+// Issue fatal error if matching trait found in UsedReplicationStateTraits is not set for the Serializer.
+void IRISCORE_API ValidateForwardingNetSerializerTraits(const FNetSerializer* Serializer, EReplicationStateTraits UsedReplicationStateTraits);
+
 }
 
 // Only needed if we want to export PropertyNetSerializerInfo, this goes in the header if we need to export it
@@ -155,7 +165,7 @@ UE::Net::FPropertyNetSerializerInfoRegistry::Unregister(&GetPropertyNetSerialize
 #endif
 
 // Implement minimal required delegates for a NetSerializer
-#define UE_NET_IMPLEMENT_NETSERIALZIER_REGISTRY_DELEGATES(Name) \
+#define UE_NET_IMPLEMENT_NETSERIALIZER_REGISTRY_DELEGATES(Name) \
 struct F##Name##NetSerializerRegistryDelegates : protected FNetSerializerRegistryDelegates \
 { \
 	~F##Name##NetSerializerRegistryDelegates() { UE_NET_UNREGISTER_NETSERIALIZER_INFO(PropertyNetSerializerRegistry_NAME_##Name); } \
@@ -167,4 +177,5 @@ static F##Name##NetSerializerRegistryDelegates Name##NetSerializerRegistryDelega
 #define UE_NET_IMPLEMENT_FORWARDING_NETSERIALIZER_AND_REGISTRY_DELEGATES(Name, SerializerName) \
 	static const FName PropertyNetSerializerRegistry_NAME_##Name( PREPROCESSOR_TO_STRING(Name) ); \
 	UE_NET_IMPLEMENT_NAMED_STRUCT_NETSERIALIZER_INFO(PropertyNetSerializerRegistry_NAME_##Name, SerializerName); \
-	UE_NET_IMPLEMENT_NETSERIALZIER_REGISTRY_DELEGATES(Name)
+	UE_NET_IMPLEMENT_NETSERIALIZER_REGISTRY_DELEGATES(Name)
+
