Ability System Replays: Fixed bug with non-Position Replicated replay recording where abilities that jumped around an animation would not record properly.

[RN]

diff --git a/Engine/Plugins/Runtime/GameplayAbilities/Source/GameplayAbilities/Private/AbilitySystemComponent.cpp b/Engine/Plugins/Runtime/GameplayAbilities/Source/GameplayAbilities/Private/AbilitySystemComponent.cpp
--- a/Engine/Plugins/Runtime/GameplayAbilities/Source/GameplayAbilities/Private/AbilitySystemComponent.cpp
+++ b/Engine/Plugins/Runtime/GameplayAbilities/Source/GameplayAbilities/Private/AbilitySystemComponent.cpp
@@ -1760,6 +1760,18 @@ bool UAbilitySystemComponent::IsOwnerActorAuthoritative() const
 	return !bCachedIsNetSimulated;
 }
 
+bool UAbilitySystemComponent::ShouldRecordMontageReplication() const
+{
+	// Returns true IF the owner is authoritative OR the world is recording a replay.
+	if (IsOwnerActorAuthoritative())
+	{
+		return true;
+	}
+
+	const UWorld* World = GetWorld();
+	return World && World->IsRecordingReplay();
+}
+
 void UAbilitySystemComponent::OnAttributeAggregatorDirty(FAggregator* Aggregator, FGameplayAttribute Attribute, bool bFromRecursiveCall)
 {
 	ActiveGameplayEffects.OnAttributeAggregatorDirty(Aggregator, Attribute, bFromRecursiveCall);
diff --git a/Engine/Plugins/Runtime/GameplayAbilities/Source/GameplayAbilities/Private/AbilitySystemComponent_Abilities.cpp b/Engine/Plugins/Runtime/GameplayAbilities/Source/GameplayAbilities/Private/AbilitySystemComponent_Abilities.cpp
--- a/Engine/Plugins/Runtime/GameplayAbilities/Source/GameplayAbilities/Private/AbilitySystemComponent_Abilities.cpp
+++ b/Engine/Plugins/Runtime/GameplayAbilities/Source/GameplayAbilities/Private/AbilitySystemComponent_Abilities.cpp
@@ -2727,8 +2727,10 @@ float UAbilitySystemComponent::PlayMontage(UGameplayAbility* InAnimatingAbility,
 				AnimInstance->Montage_JumpToSection(StartSectionName, NewAnimMontage);
 			}
 
-			// Replicate to non owners
-			if (IsOwnerActorAuthoritative())
+			// Replicate for non-owners and for replay recordings
+			// The data we set from GetRepAnimMontageInfo_Mutable() is used both by the server to replicate to clients and by clients to record replays.
+			// We need to set this data for recording clients because there exists network configurations where an abilities montage data will not replicate to some clients (for example: if the client is an autonomous proxy.)
+			if (ShouldRecordMontageReplication())
 			{
 				FGameplayAbilityRepAnimMontage& MutableRepAnimMontageInfo = GetRepAnimMontageInfo_Mutable();
 
@@ -2745,8 +2747,12 @@ float UAbilitySystemComponent::PlayMontage(UGameplayAbility* InAnimatingAbility,
 
 				// Update parameters that change during Montage life time.
 				AnimMontage_UpdateReplicatedData();
+			}
 
-				// Force net update on our avatar actor
+			// Replicate to non-owners
+			if (IsOwnerActorAuthoritative())
+			{
+				// Force net update on our avatar actor.
 				if (AbilityActorInfo->AvatarActor != nullptr)
 				{
 					AbilityActorInfo->AvatarActor->ForceNetUpdate();
@@ -2785,20 +2791,20 @@ float UAbilitySystemComponent::PlayMontageSimulated(UAnimMontage* NewAnimMontage
 
 void UAbilitySystemComponent::AnimMontage_UpdateReplicatedData()
 {
-	check(IsOwnerActorAuthoritative());
+	check(ShouldRecordMontageReplication());
 
 	AnimMontage_UpdateReplicatedData(GetRepAnimMontageInfo_Mutable());
 }
 
 void UAbilitySystemComponent::AnimMontage_UpdateReplicatedData(FGameplayAbilityRepAnimMontage& OutRepAnimMontageInfo)
 {
-	UAnimInstance* AnimInstance = AbilityActorInfo.IsValid() ? AbilityActorInfo->GetAnimInstance() : nullptr;
+	const UAnimInstance* AnimInstance = AbilityActorInfo.IsValid() ? AbilityActorInfo->GetAnimInstance() : nullptr;
 	if (AnimInstance && LocalAnimMontageInfo.AnimMontage)
 	{
 		OutRepAnimMontageInfo.AnimMontage = LocalAnimMontageInfo.AnimMontage;
 
 		// Compressed Flags
-		bool bIsStopped = AnimInstance->Montage_GetIsStopped(LocalAnimMontageInfo.AnimMontage);
+		const bool bIsStopped = AnimInstance->Montage_GetIsStopped(LocalAnimMontageInfo.AnimMontage);
 
 		if (!bIsStopped)
 		{
@@ -3053,7 +3059,10 @@ void UAbilitySystemComponent::CurrentMontageJumpToSection(FName SectionName)
 	if ((SectionName != NAME_None) && AnimInstance && LocalAnimMontageInfo.AnimMontage)
 	{
 		AnimInstance->Montage_JumpToSection(SectionName, LocalAnimMontageInfo.AnimMontage);
-		if (IsOwnerActorAuthoritative())
+
+		// This data is needed for replication on the server and recording replays on clients.
+		// We need to set GetRepAnimMontageInfo_Mutable on replay recording clients because this data is NOT replicated to all clients (for example, it is NOT replicated to autonomous proxy clients.)
+		if (ShouldRecordMontageReplication())
 		{
 			FGameplayAbilityRepAnimMontage& MutableRepAnimMontageInfo = GetRepAnimMontageInfo_Mutable();
 
@@ -3066,7 +3075,9 @@ void UAbilitySystemComponent::CurrentMontageJumpToSection(FName SectionName)
 
 			AnimMontage_UpdateReplicatedData();
 		}
-		else
+		
+		// If we are NOT the authority, then let the server handling jumping the montage.
+		if (!IsOwnerActorAuthoritative())
 		{
 			ServerCurrentMontageJumpToSectionName(LocalAnimMontageInfo.AnimMontage, SectionName);
 		}
diff --git a/Engine/Plugins/Runtime/GameplayAbilities/Source/GameplayAbilities/Public/AbilitySystemComponent.h b/Engine/Plugins/Runtime/GameplayAbilities/Source/GameplayAbilities/Public/AbilitySystemComponent.h
--- a/Engine/Plugins/Runtime/GameplayAbilities/Source/GameplayAbilities/Public/AbilitySystemComponent.h
+++ b/Engine/Plugins/Runtime/GameplayAbilities/Source/GameplayAbilities/Public/AbilitySystemComponent.h
@@ -265,6 +265,9 @@ class GAMEPLAYABILITIES_API UAbilitySystemComponent : public UGameplayTasksCompo
 	/** Returns true if this component's actor has authority */
 	virtual bool IsOwnerActorAuthoritative() const;
 
+	/** Returns true if this component should record montage replication info. */
+	virtual bool ShouldRecordMontageReplication() const;
+
 	/** Replicate that an ability has ended/canceled, to the client or server as appropriate */
 	virtual void ReplicateEndOrCancelAbility(FGameplayAbilitySpecHandle Handle, FGameplayAbilityActivationInfo ActivationInfo, UGameplayAbility* Ability, bool bWasCanceled);
 
