Fixed non-unity/pch errors

diff --git a/Engine/Plugins/Runtime/GameplayAbilities/Source/GameplayAbilities/Private/Sequencer/MovieSceneGameplayCueSections.cpp b/Engine/Plugins/Runtime/GameplayAbilities/Source/GameplayAbilities/Private/Sequencer/MovieSceneGameplayCueSections.cpp
--- a/Engine/Plugins/Runtime/GameplayAbilities/Source/GameplayAbilities/Private/Sequencer/MovieSceneGameplayCueSections.cpp
+++ b/Engine/Plugins/Runtime/GameplayAbilities/Source/GameplayAbilities/Private/Sequencer/MovieSceneGameplayCueSections.cpp
@@ -131,7 +131,7 @@ void ExecuteGameplayCueEvent(IMovieScenePlayer* Player, const UE::MovieScene::FE
 
 			if (GameplayCueKey.bAttachToBinding)
 			{
-				if (CueParameters.TargetAttachComponent != nullptr)
+				if (CueParameters.TargetAttachComponent.Get() != nullptr)
 				{
 					FTransform Transform = FTransform::Identity;
 
diff --git a/Engine/Source/Runtime/ApplicationCore/Public/IOS/IOSApplication.h b/Engine/Source/Runtime/ApplicationCore/Public/IOS/IOSApplication.h
--- a/Engine/Source/Runtime/ApplicationCore/Public/IOS/IOSApplication.h
+++ b/Engine/Source/Runtime/ApplicationCore/Public/IOS/IOSApplication.h
@@ -2,6 +2,7 @@
 
 #pragma once
 
+#include "GenericPlatform/Accessibility/GenericAccessibleInterfaces.h"
 #include "GenericPlatform/GenericApplication.h"
 
 class FIOSWindow;
diff --git a/Engine/Source/Runtime/GeometryCore/Public/Implicit/SparseNarrowBandMeshSDF.h b/Engine/Source/Runtime/GeometryCore/Public/Implicit/SparseNarrowBandMeshSDF.h
--- a/Engine/Source/Runtime/GeometryCore/Public/Implicit/SparseNarrowBandMeshSDF.h
+++ b/Engine/Source/Runtime/GeometryCore/Public/Implicit/SparseNarrowBandMeshSDF.h
@@ -988,7 +988,7 @@ private:
 							if (IsInBlock(LocalCoords))
 							{
 								const int32 local_index = FloatBlockData3Type::ToLinear(LocalCoords.X, LocalCoords.Y, LocalCoords.Z);
-								if (VistedBlock.BitArray[local_index] == false)
+								if ((bool)VistedBlock.BitArray[local_index] == false)
 								{
 									VistedBlock.BitArray[local_index] = true;
 
