Add nullptr check in UAbilitySystemBlueprintLibrary::MakeSpecHandle

diff --git a/Engine/Plugins/Runtime/GameplayAbilities/Source/GameplayAbilities/Private/AbilitySystemBlueprintLibrary.cpp b/Engine/Plugins/Runtime/GameplayAbilities/Source/GameplayAbilities/Private/AbilitySystemBlueprintLibrary.cpp
--- a/Engine/Plugins/Runtime/GameplayAbilities/Source/GameplayAbilities/Private/AbilitySystemBlueprintLibrary.cpp
+++ b/Engine/Plugins/Runtime/GameplayAbilities/Source/GameplayAbilities/Private/AbilitySystemBlueprintLibrary.cpp
@@ -256,9 +256,15 @@ FGameplayTargetDataFilterHandle UAbilitySystemBlueprintLibrary::MakeFilterHandle
 
 FGameplayEffectSpecHandle UAbilitySystemBlueprintLibrary::MakeSpecHandle(UGameplayEffect* InGameplayEffect, AActor* InInstigator, AActor* InEffectCauser, float InLevel)
 {
-	FGameplayEffectContext* EffectContext = UAbilitySystemGlobals::Get().AllocGameplayEffectContext();
-	EffectContext->AddInstigator(InInstigator, InEffectCauser);
-	return FGameplayEffectSpecHandle(new FGameplayEffectSpec(InGameplayEffect, FGameplayEffectContextHandle(EffectContext), InLevel));
+	if (InGameplayEffect)
+	{
+		FGameplayEffectContext* EffectContext = UAbilitySystemGlobals::Get().AllocGameplayEffectContext();
+		EffectContext->AddInstigator(InInstigator, InEffectCauser);
+		return FGameplayEffectSpecHandle(new FGameplayEffectSpec(InGameplayEffect, FGameplayEffectContextHandle(EffectContext), InLevel));
+	}
+	
+	ABILITY_LOG(Warning, TEXT("%s was called with an invalid GameplayEffect object!"), *FString(__FUNCTION__));
+	return FGameplayEffectSpecHandle();
 }
 
 FGameplayEffectSpecHandle UAbilitySystemBlueprintLibrary::CloneSpecHandle(AActor* InNewInstigator, AActor* InEffectCauser, FGameplayEffectSpecHandle GameplayEffectSpecHandle_Clone)
