Fix edge cases of memory stomp for WaitGameplayEvent chaining usages
- Broadcast could cause GenericGameplayEventCallbacks whole map to be reallocated, if we triggered UAbilityTask_WaitGameplayEvent or UAbilityAsync_WaitGameplayEvent's Activate during the broadcast. Thus, the delegate itself would be reallocated, causing This pointer or const ref of invocation list to be invalid.

diff --git a/Engine/Plugins/Runtime/GameplayAbilities/Source/GameplayAbilities/Private/AbilitySystemComponent_Abilities.cpp b/Engine/Plugins/Runtime/GameplayAbilities/Source/GameplayAbilities/Private/AbilitySystemComponent_Abilities.cpp
--- a/Engine/Plugins/Runtime/GameplayAbilities/Source/GameplayAbilities/Private/AbilitySystemComponent_Abilities.cpp
+++ b/Engine/Plugins/Runtime/GameplayAbilities/Source/GameplayAbilities/Private/AbilitySystemComponent_Abilities.cpp
@@ -2258,7 +2258,9 @@ int32 UAbilitySystemComponent::HandleGameplayEvent(FGameplayTag EventTag, const
 
 	if (FGameplayEventMulticastDelegate* Delegate = GenericGameplayEventCallbacks.Find(EventTag))
 	{
-		Delegate->Broadcast(Payload);
+		// Make a copy before broadcasting to prevent memory stomping
+		FGameplayEventMulticastDelegate DelegateCopy = *Delegate;
+		DelegateCopy.Broadcast(Payload);
 	}
 
 	// Make a copy in case it changes due to callbacks
