FFieldPath now respects CoreRedirects.
FFieldPathProperty applies redirects on ImportText_Internal, automatically redirecting FFieldPath saved in string format for blueprint graph pins.
FGameplayAttribute applies redirects after deserializing unredirected value and fixes up derived properties, for UObject based assets with FGameplayAttribute properties.
UAttributeSet subclass renaming and adding ClassRedirect works now. Previously, FGameplayAttribute pins would show correct value in blueprint graphs but use incorrect value at runtime.
FGameplayAttribute renaming and adding PropertyRedirect works now. Previously, FGameplayAttribute fields in blueprint assets would break (null out), in addition to blueprint graphs using undirected value at runtime.

diff --git a/Engine/Plugins/Runtime/GameplayAbilities/Source/GameplayAbilities/Private/AttributeSet.cpp b/Engine/Plugins/Runtime/GameplayAbilities/Source/GameplayAbilities/Private/AttributeSet.cpp
--- a/Engine/Plugins/Runtime/GameplayAbilities/Source/GameplayAbilities/Private/AttributeSet.cpp
+++ b/Engine/Plugins/Runtime/GameplayAbilities/Source/GameplayAbilities/Private/AttributeSet.cpp
@@ -246,21 +246,49 @@ void FGameplayAttribute::PostSerialize(const FArchive& Ar)
 {
 	if (Ar.IsLoading() && Ar.IsPersistent() && !Ar.HasAnyPortFlags(PPF_Duplicate | PPF_DuplicateForPIE))
 	{
+		// Once struct is loaded, check if redirectors apply to the imported attribute field path
+		const FString PathName = Attribute.ToString();
+		const FString RedirectedPathName = FFieldPathProperty::RedirectFieldPathName(PathName);
+		if (!RedirectedPathName.Equals(PathName))
+		{
+			// If the path got redirected, attempt to resolve the new property
+			FString NewAttributeOwner;
+			FString NewAttributeName;
+			if (RedirectedPathName.Split(":", &NewAttributeOwner, &NewAttributeName))
+			{
+				// Update attribute's field path (may or may not resolve)
+				const UStruct* NewClass = FindObject<UStruct>(nullptr, *NewAttributeOwner);
+				Attribute = FindFProperty<FProperty>(NewClass, *NewAttributeName);
+
+				// Verbose log any applied redirectors
+				FUObjectSerializeContext* LoadContext = const_cast<FArchive*>(&Ar)->GetSerializeContext();
+				const FString AssetName = (LoadContext && LoadContext->SerializedObject) ? LoadContext->SerializedObject->GetPathName() : TEXT("Unknown Object");
+				ABILITY_LOG(Verbose, TEXT("FGameplayAttribute::PostSerialize redirected an attribute '%s' -> '%s'. (Asset: %s)"), *PathName, *RedirectedPathName, *AssetName);
+			}
+		}
+
+		// The attribute reference is serialized in two ways:
+		// - 'Attribute' contains the full path, ex: /Script/GameModule.GameAttributeSet:AttrName
+		// - 'AttributeOwner' and 'AttributeName' are cached references to the attribute set UClass and the attribute's name
+		// We want the data to stay in sync, with 'Attribute' having priority as source of truth. In both cases, derive one from
+		// the other to keep them in sync.
 		if (Attribute.Get())
 		{
+			// Ensure owner and name are in sync with field path
 			AttributeOwner = Attribute->GetOwnerStruct();
 			Attribute->GetName(AttributeName);
 		}
 		else if (!AttributeName.IsEmpty() && AttributeOwner != nullptr)
 		{
+			// Attempt to resolve field path from owner and attribute name
 			Attribute = FindFProperty<FProperty>(AttributeOwner, *AttributeName);
 
+			// Log warning if attribute failed to resolve while name + owner were non-null
 			if (!Attribute.Get())
 			{
 				FUObjectSerializeContext* LoadContext = const_cast<FArchive*>(&Ar)->GetSerializeContext();
-				FString AssetName = (LoadContext && LoadContext->SerializedObject) ? LoadContext->SerializedObject->GetPathName() : TEXT("Unknown Object");
-
-				FString OwnerName = AttributeOwner ? AttributeOwner->GetName() : TEXT("NONE");
+				const FString AssetName = (LoadContext && LoadContext->SerializedObject) ? LoadContext->SerializedObject->GetPathName() : TEXT("Unknown Object");
+				const FString OwnerName = AttributeOwner ? AttributeOwner->GetName() : TEXT("NONE");
 				ABILITY_LOG(Warning, TEXT("FGameplayAttribute::PostSerialize called on an invalid attribute with owner %s and name %s. (Asset: %s)"), *OwnerName, *AttributeName, *AssetName);
 			}
 		}
diff --git a/Engine/Plugins/Runtime/GameplayAbilities/Source/GameplayAbilitiesEditor/Private/SGameplayAttributeGraphPin.cpp b/Engine/Plugins/Runtime/GameplayAbilities/Source/GameplayAbilitiesEditor/Private/SGameplayAttributeGraphPin.cpp
--- a/Engine/Plugins/Runtime/GameplayAbilities/Source/GameplayAbilitiesEditor/Private/SGameplayAttributeGraphPin.cpp
+++ b/Engine/Plugins/Runtime/GameplayAbilities/Source/GameplayAbilitiesEditor/Private/SGameplayAttributeGraphPin.cpp
@@ -19,95 +19,18 @@ void SGameplayAttributeGraphPin::Construct( const FArguments& InArgs, UEdGraphPi
 }
 
 TSharedRef<SWidget>	SGameplayAttributeGraphPin::GetDefaultValueWidget()
-{
-	// Parse out current default value	
-	FString DefaultString = GraphPinObj->GetDefaultAsString();
+{	
 	FGameplayAttribute DefaultAttribute;
 
-	// check for redirectors
-	FString TempString;
-	FString WorkingString;
-	if (DefaultString.Split(TEXT(","), nullptr, &WorkingString, ESearchCase::CaseSensitive))
-	{
-		if (WorkingString.Split(TEXT(","), &TempString, nullptr, ESearchCase::CaseSensitive))
-		{
-			WorkingString = MoveTemp(TempString);
-			FString AttributeNameString;
-			WorkingString.Split(TEXT(":"), &TempString, &AttributeNameString, ESearchCase::CaseSensitive);
-
-			WorkingString = MoveTemp(TempString);
-			FString ClassNameString;
-			WorkingString.Split(TEXT("."), &TempString, &ClassNameString, ESearchCase::CaseSensitive);
-
-			WorkingString = MoveTemp(TempString);
-			FString PackageNameString;
-			WorkingString.Split(TEXT("Attribute="), nullptr, &PackageNameString);
-
-			// class redirector check
-			const FCoreRedirect* ClassValueRedirect = nullptr;
-			FCoreRedirectObjectName OldClassName(FName(*ClassNameString, FNAME_Find), NAME_None, FName(PackageNameString, FNAME_Find));
-			FCoreRedirectObjectName NewClassName;
-			bool bFoundClassRedirector = FCoreRedirects::RedirectNameAndValues(ECoreRedirectFlags::Type_Class, OldClassName, NewClassName, &ClassValueRedirect);
-
-			// property redirector check
-			const FCoreRedirect* PropertyRedirect = nullptr;
-			FCoreRedirectObjectName OldPropertyName(FName(*AttributeNameString, FNAME_Find), FName(*ClassNameString, FNAME_Find), FName(PackageNameString, FNAME_Find));
-			FCoreRedirectObjectName NewPropertyName;
-			bool bFoundPropertyRedirector = FCoreRedirects::RedirectNameAndValues(ECoreRedirectFlags::Type_Property, OldPropertyName, NewPropertyName, &PropertyRedirect);
-
-			if (bFoundClassRedirector || bFoundPropertyRedirector)
-			{
-				// we found a redirector
-				// now we need to find the matching property for the new attribute
-
-				bool bFoundMatch = false;
-
-				for (TObjectIterator<UClass> ClassIt; ClassIt; ++ClassIt)
-				{
-					UClass* Class = *ClassIt;
-					if ((Class->IsChildOf(UAttributeSet::StaticClass()) || Class->IsChildOf(UAbilitySystemComponent::StaticClass())) && !Class->ClassGeneratedBy)
-					{
-						if (Class->GetFName() == NewClassName.ObjectName)
-						{
-							FName PropertyNameToFind(AttributeNameString, FNAME_Find);
-
-							if (bFoundPropertyRedirector)
-							{
-								PropertyNameToFind = NewPropertyName.ObjectName;
-							}
-
-							for (TFieldIterator<FProperty> PropertyIt(Class, EFieldIteratorFlags::ExcludeSuper); PropertyIt; ++PropertyIt)
-							{
-								FProperty* Property = *PropertyIt;
-								if (Property->GetFName() == PropertyNameToFind)
-								{
-									FGameplayAttribute Attribute;
-									Attribute.SetUProperty(Property);
-									DefaultString.Reset();
-									FGameplayAttribute::StaticStruct()->ExportText(DefaultString, &Attribute, &Attribute, nullptr, EPropertyPortFlags::PPF_SerializedAsImportText, nullptr);
-									bFoundMatch = true;
-									break;
-								}
-							}
-						}
-					}
-
-					if (bFoundMatch)
-					{
-						break;
-					}
-				}
-			}
-		}
-	}
-
-	UScriptStruct* PinLiteralStructType = FGameplayAttribute::StaticStruct();
+	// Parse out current default value
+	FString DefaultString = GraphPinObj->GetDefaultAsString();
 	if (!DefaultString.IsEmpty())
 	{
+		UScriptStruct* PinLiteralStructType = FGameplayAttribute::StaticStruct();
 		PinLiteralStructType->ImportText(*DefaultString, &DefaultAttribute, nullptr, EPropertyPortFlags::PPF_SerializedAsImportText, GError, PinLiteralStructType->GetName(), true);
 	}
 
-	//Create widget
+	// Create widget
 	return SNew(SVerticalBox)
 		+SVerticalBox::Slot()
 		.AutoHeight()
diff --git a/Engine/Source/Runtime/CoreUObject/Private/UObject/FieldPathProperty.cpp b/Engine/Source/Runtime/CoreUObject/Private/UObject/FieldPathProperty.cpp
--- a/Engine/Source/Runtime/CoreUObject/Private/UObject/FieldPathProperty.cpp
+++ b/Engine/Source/Runtime/CoreUObject/Private/UObject/FieldPathProperty.cpp
@@ -7,6 +7,7 @@
 #include "UObject/LinkerLoad.h"
 #include "UObject/ConstructorHelpers.h"
 #include "Misc/Parse.h"
+#include "UObject/CoreRedirects.h"
 #include "UObject/PropertyHelper.h"
 
 IMPLEMENT_FIELD(FFieldPathProperty)
@@ -162,7 +163,14 @@ const TCHAR* FFieldPathProperty::ImportText_Internal( const TCHAR* Buffer, void*
 			}
 			PathName = MoveTemp(UnquotedPathName);
 		}
+
+		// Attempt to apply core redirects to imported path name
+		PathName = RedirectFieldPathName(PathName);
+
+		// Resolve FFieldPath from imported and fixed up path name
 		ImportedPath.Generate(*PathName);
+
+		// Set the FFieldPath value within container
 		if (PropertyPointerType == EPropertyPointerType::Container && HasSetter())
 		{
 			SetValue_InContainer(ContainerOrPropertyPtr, ImportedPath);
@@ -213,3 +221,45 @@ bool FFieldPathProperty::SupportsNetSharedSerialization() const
 {
 	return false;
 }
+
+FString FFieldPathProperty::RedirectFieldPathName(const FString& InPathName)
+{
+	// Apply redirectors to imported PathName value. First deconstruct it into package, class and field names.
+	// PathName format: /Script/MyGameModule.MyClass:MyField
+	FString OldPackageName;
+	FString OldClassName;
+	FString OldFieldName;
+	FString Tmp;
+	if (InPathName.Split(TEXT("."), &OldPackageName, &Tmp) && Tmp.Split(TEXT(":"), &OldClassName, &OldFieldName))
+	{
+		// Check for full property path redirect
+		{
+			const FCoreRedirectObjectName OldRedirectName = FCoreRedirectObjectName(InPathName);
+			constexpr ECoreRedirectFlags RedirectFlags = ECoreRedirectFlags::Type_Property; 
+			FCoreRedirectObjectName NewObjectName;
+			const FCoreRedirect* FoundValueRedirect = nullptr;
+			if (FCoreRedirects::RedirectNameAndValues(RedirectFlags, OldRedirectName, NewObjectName, &FoundValueRedirect))
+			{
+				const FString NewPathName = FString::Printf(TEXT("%s.%s:%s"), *NewObjectName.PackageName.ToString(), *NewObjectName.OuterName.ToString(), *NewObjectName.ObjectName.ToString());
+				UE_LOG(LogCoreRedirects, Verbose, TEXT("FFieldPathProperty: Redirected '%s' -> '%s'"), *InPathName, *NewPathName);
+				return NewPathName;
+			}
+		}
+		
+		// Check for outer-only redirect
+		{
+			const FCoreRedirectObjectName OldRedirectName(*OldClassName, NAME_None, *OldPackageName);
+			constexpr ECoreRedirectFlags RedirectFlags = ECoreRedirectFlags::Type_Package | ECoreRedirectFlags::Type_Class | ECoreRedirectFlags::Type_Struct;
+			FCoreRedirectObjectName NewObjectName;
+			const FCoreRedirect* FoundValueRedirect = nullptr;
+			if (FCoreRedirects::RedirectNameAndValues(RedirectFlags, OldRedirectName, NewObjectName, &FoundValueRedirect))
+			{
+				const FString NewPathName = FString::Printf(TEXT("%s.%s:%s"), *NewObjectName.PackageName.ToString(), *NewObjectName.ObjectName.ToString(), *OldFieldName);
+				UE_LOG(LogCoreRedirects, Verbose, TEXT("FFieldPathProperty: Redirected '%s' -> '%s'"), *InPathName, *NewPathName);
+				return NewPathName;
+			}
+		}
+	}
+	
+	return InPathName;
+}
diff --git a/Engine/Source/Runtime/CoreUObject/Public/UObject/FieldPathProperty.h b/Engine/Source/Runtime/CoreUObject/Public/UObject/FieldPathProperty.h
--- a/Engine/Source/Runtime/CoreUObject/Public/UObject/FieldPathProperty.h
+++ b/Engine/Source/Runtime/CoreUObject/Public/UObject/FieldPathProperty.h
@@ -79,4 +79,6 @@ PRAGMA_ENABLE_DEPRECATION_WARNINGS
 	virtual void EmitReferenceInfo(UClass& OwnerClass, int32 BaseOffset, TArray<const FStructProperty*>& EncounteredStructProps, FGCStackSizeHelper& StackSizeHelper) override;
 	virtual bool SupportsNetSharedSerialization() const override;
 	// End of FProperty interface
+
+	static FString RedirectFieldPathName(const FString& InPathName);
 };
\ No newline at end of file
