Fixed issue with GameplayEffectContext replicating reference to dynamic object that never will be replicated which prevents GameplayEffect from calling Owner->HandleDeferredGameplayCues(this) as bHasMoreUnmappedReferences would always be true.

Engine changes

diff --git a/Engine/Plugins/Runtime/GameplayAbilities/Source/GameplayAbilities/Private/GameplayEffectTypes.cpp b/Engine/Plugins/Runtime/GameplayAbilities/Source/GameplayAbilities/Private/GameplayEffectTypes.cpp
--- a/Engine/Plugins/Runtime/GameplayAbilities/Source/GameplayAbilities/Private/GameplayEffectTypes.cpp
+++ b/Engine/Plugins/Runtime/GameplayAbilities/Source/GameplayAbilities/Private/GameplayEffectTypes.cpp
@@ -22,6 +22,18 @@ const FName FGameplayModEvaluationChannelSettings::ForceHideMetadataKey(TEXT("Fo
 const FString FGameplayModEvaluationChannelSettings::ForceHideMetadataEnabledValue(TEXT("True"));
 #endif // #if WITH_EDITORONLY_DATA
 
+#if !UE_BUILD_SHIPPING
+namespace UE::Private
+{
+static bool bWarnIfTryingToReplicateNotSupportedActorReference = false;
+static FAutoConsoleVariableRef CVarWarnIfTryingToReplicateNotSupportedActorReference(
+	TEXT("GameplayEffectContext.WarnIfTryingToReplicateNotSupportedActorReference"),
+	bWarnIfTryingToReplicateNotSupportedActorReference,
+	TEXT("If set to true a warning will be issued if we are trying to replicate a reference to a not supported actor as part of a GameplayEffectContext."),
+	ECVF_Default);
+}
+#endif
+
 FGameplayModEvaluationChannelSettings::FGameplayModEvaluationChannelSettings()
 {
 	static const UEnum* EvalChannelEnum = nullptr;
@@ -115,10 +127,36 @@ FString FGameplayEffectAttributeCaptureDefinition::ToSimpleString() const
 //
 // --------------------------------------------------------------------------------------------------------------------------------------------------------
 
+bool FGameplayEffectContext::CanActorReferenceBeReplicated(const AActor* Actor)
+{
+	// We always support replication of null references and stably named actors
+	if (!Actor || Actor->IsFullNameStableForNetworking())
+	{
+		return true;
+	}
+
+	// If we get here this is a dynamic object and we only want to replicate the reference if the actor is set to replicate, otherwise the resolve on the client will constantly fail
+	const bool bIsSupportedForNetWorking = Actor->IsSupportedForNetworking();
+	const bool bCanDynamicReferenceBeReplicated = bIsSupportedForNetWorking && Actor->GetIsReplicated();
+
+#if !UE_BUILD_SHIPPING
+	// Optionally trigger warning if we are trying to replicate a reference to an object that never will be resolvable on receiving end
+	if (UE::Private::bWarnIfTryingToReplicateNotSupportedActorReference && (!bCanDynamicReferenceBeReplicated && bIsSupportedForNetWorking))
+	{
+		ABILITY_LOG(Warning, TEXT("Attempted to replicate a reference to dynamically spawned object that is set to not replicate %s."), *(Actor->GetName()));
+	}
+#endif
+
+	return bCanDynamicReferenceBeReplicated;
+}
+
 void FGameplayEffectContext::AddInstigator(class AActor *InInstigator, class AActor *InEffectCauser)
 {
 	Instigator = InInstigator;
-	EffectCauser = InEffectCauser;
+	bReplicateInstigator = CanActorReferenceBeReplicated(InInstigator);
+
+	SetEffectCauser(InEffectCauser);
+
 	InstigatorAbilitySystemComponent = NULL;
 
 	// Cache off the AbilitySystemComponent.
@@ -177,11 +215,11 @@ bool FGameplayEffectContext::NetSerialize(FArchive& Ar, class UPackageMap* Map,
 	uint8 RepBits = 0;
 	if (Ar.IsSaving())
 	{
-		if (Instigator.IsValid() )
+		if (bReplicateInstigator && Instigator.IsValid())
 		{
 			RepBits |= 1 << 0;
 		}
-		if (EffectCauser.IsValid() )
+		if (bReplicateEffectCauser && EffectCauser.IsValid() )
 		{
 			RepBits |= 1 << 1;
 		}
diff --git a/Engine/Plugins/Runtime/GameplayAbilities/Source/GameplayAbilities/Private/Serialization/GameplayEffectContextNetSerializer.cpp b/Engine/Plugins/Runtime/GameplayAbilities/Source/GameplayAbilities/Private/Serialization/GameplayEffectContextNetSerializer.cpp
--- a/Engine/Plugins/Runtime/GameplayAbilities/Source/GameplayAbilities/Private/Serialization/GameplayEffectContextNetSerializer.cpp
+++ b/Engine/Plugins/Runtime/GameplayAbilities/Source/GameplayAbilities/Private/Serialization/GameplayEffectContextNetSerializer.cpp
@@ -37,7 +37,9 @@ struct FGameplayEffectContextNetSerializer
 	{
 		ReplicateWorldOrigin = 1U,
 		ReplicateSourceObject = ReplicateWorldOrigin << 1U,
-		ReplicateHitResult =  ReplicateSourceObject << 1U,
+		ReplicateInstigator = ReplicateSourceObject << 1U,
+		ReplicateEffectCauser = ReplicateInstigator << 1U,
+		ReplicateHitResult =  ReplicateEffectCauser << 1U,
 	};
 
 	struct FQuantizedType
@@ -110,7 +112,7 @@ void FGameplayEffectContextNetSerializer::Serialize(FNetSerializationContext& Co
 	FNetBitStreamWriter* Writer = Context.GetBitStreamWriter();
 
 	const uint32 ReplicationFlags = Value.ReplicationFlags;
-	Writer->WriteBits(ReplicationFlags, 3U);
+	Writer->WriteBits(ReplicationFlags, 5U);
 
 	// We need to manually serialize the properties as there are some replicated properties that don't need to be replicated,
 	// as is the case with some of the bools that are covered by the ReplicationFlags.
@@ -122,6 +124,8 @@ void FGameplayEffectContextNetSerializer::Serialize(FNetSerializationContext& Co
 		
 		uint32 MemberMaskStorage = ~0U;
 		FNetBitArrayView MemberMask(&MemberMaskStorage, FGameplayEffectContextAccessorForNetSerializer::EPropertyName::PropertyName_ReplicatedPropertyCount, FNetBitArrayView::NoResetNoValidate);
+		MemberMask.SetBitValue(FGameplayEffectContextAccessorForNetSerializer::EPropertyName::PropertyName_Instigator, (ReplicationFlags & EReplicationFlags::ReplicateInstigator) != 0);
+		MemberMask.SetBitValue(FGameplayEffectContextAccessorForNetSerializer::EPropertyName::PropertyName_EffectCauser, (ReplicationFlags & EReplicationFlags::ReplicateEffectCauser) != 0);
 		MemberMask.SetBitValue(FGameplayEffectContextAccessorForNetSerializer::EPropertyName::PropertyName_SourceObject, (ReplicationFlags & EReplicationFlags::ReplicateSourceObject) != 0);
 		MemberMask.SetBitValue(FGameplayEffectContextAccessorForNetSerializer::EPropertyName::PropertyName_WorldOrigin, (ReplicationFlags & EReplicationFlags::ReplicateWorldOrigin) != 0);
 		for (uint32 PropertyIt = 0, PropertyEndIt = FGameplayEffectContextAccessorForNetSerializer::EPropertyName::PropertyName_ReplicatedPropertyCount; PropertyIt != PropertyEndIt; ++PropertyIt)
@@ -160,7 +164,7 @@ void FGameplayEffectContextNetSerializer::Deserialize(FNetSerializationContext&
 
 	FNetBitStreamReader* Reader = Context.GetBitStreamReader();
 
-	const uint32 ReplicationFlags = Reader->ReadBits(3U);
+	const uint32 ReplicationFlags = Reader->ReadBits(5U);
 	Target.ReplicationFlags = ReplicationFlags;
 
 	// We need to manually deserialize the properties as there are some replicated properties that don't need to be replicated,
@@ -173,6 +177,8 @@ void FGameplayEffectContextNetSerializer::Deserialize(FNetSerializationContext&
 		
 		uint32 MemberMaskStorage = ~0U;
 		FNetBitArrayView MemberMask(&MemberMaskStorage, FGameplayEffectContextAccessorForNetSerializer::EPropertyName::PropertyName_ReplicatedPropertyCount, FNetBitArrayView::NoResetNoValidate);
+		MemberMask.SetBitValue(FGameplayEffectContextAccessorForNetSerializer::EPropertyName::PropertyName_Instigator, (ReplicationFlags & EReplicationFlags::ReplicateInstigator) != 0);
+		MemberMask.SetBitValue(FGameplayEffectContextAccessorForNetSerializer::EPropertyName::PropertyName_EffectCauser, (ReplicationFlags & EReplicationFlags::ReplicateEffectCauser) != 0);
 		MemberMask.SetBitValue(FGameplayEffectContextAccessorForNetSerializer::EPropertyName::PropertyName_SourceObject, (ReplicationFlags & EReplicationFlags::ReplicateSourceObject) != 0);
 		MemberMask.SetBitValue(FGameplayEffectContextAccessorForNetSerializer::EPropertyName::PropertyName_WorldOrigin, (ReplicationFlags & EReplicationFlags::ReplicateWorldOrigin) != 0);
 		for (uint32 PropertyIt = 0, PropertyEndIt = FGameplayEffectContextAccessorForNetSerializer::EPropertyName::PropertyName_ReplicatedPropertyCount; PropertyIt != PropertyEndIt; ++PropertyIt)
@@ -226,6 +232,9 @@ void FGameplayEffectContextNetSerializer::Quantize(FNetSerializationContext& Con
 	TempGE.CopyReplicatedFieldsFrom(SourceValue);
 
 	const FHitResult* HitResult = SourceValue.GetHitResult();
+	
+	ReplicationFlags |= (TempGE.ShouldReplicateInstigator() ? EReplicationFlags::ReplicateInstigator : 0);
+	ReplicationFlags |= (TempGE.ShouldReplicateEffectCauser() ? EReplicationFlags::ReplicateEffectCauser : 0);
 	ReplicationFlags |= (TempGE.ShouldReplicateWorldOrigin() ? EReplicationFlags::ReplicateWorldOrigin : 0);
 	ReplicationFlags |= (TempGE.ShouldReplicateSourceObject() ? EReplicationFlags::ReplicateSourceObject : 0);
 	ReplicationFlags |= (HitResult != nullptr ? EReplicationFlags::ReplicateHitResult : 0);
@@ -268,6 +277,8 @@ void FGameplayEffectContextNetSerializer::Dequantize(FNetSerializationContext& C
 		GEDequantizeArgs.Source = NetSerializerValuePointer(&SourceValue.EffectContext);
 		GEDequantizeArgs.Target = NetSerializerValuePointer(&TempGE);
 		StructNetSerializer->Dequantize(Context, GEDequantizeArgs);
+		TempGE.SetShouldReplicateInstigator((ReplicationFlags & EReplicationFlags::ReplicateInstigator) != 0);
+		TempGE.SetShouldReplicateEffectCauser((ReplicationFlags & EReplicationFlags::ReplicateEffectCauser) != 0);
 		TempGE.SetShouldReplicateWorldOrigin((ReplicationFlags & EReplicationFlags::ReplicateWorldOrigin) != 0);
 		TempGE.SetShouldReplicateSourceObject((ReplicationFlags & EReplicationFlags::ReplicateSourceObject) != 0);
 	}
@@ -584,14 +595,25 @@ void FGameplayEffectContextAccessorForNetSerializer::CopyReplicatedFieldsFrom(co
 	this->AbilityCDO = GE.AbilityCDO;
 	this->AbilityLevel = GE.AbilityLevel;
 	this->Actors = GE.Actors;
-	this->EffectCauser = GE.EffectCauser;
-	this->Instigator = GE.Instigator;
+	if (GE.bReplicateEffectCauser)
+	{
+		this->EffectCauser = GE.EffectCauser;
+	}
+
+	if (GE.bReplicateInstigator)
+	{
+		this->Instigator = GE.Instigator;
+	}
+
 	this->bHasWorldOrigin = GE.bHasWorldOrigin;
 	if (GE.bHasWorldOrigin)
 	{
 		this->WorldOrigin = GE.WorldOrigin;
 	}
 	this->bReplicateSourceObject = GE.bReplicateSourceObject;
+	this->bReplicateInstigator = GE.bReplicateInstigator;
+	this->bReplicateEffectCauser = GE.bReplicateEffectCauser;
+
 	if (GE.bReplicateSourceObject)
 	{
 		this->SourceObject = GE.SourceObject;
diff --git a/Engine/Plugins/Runtime/GameplayAbilities/Source/GameplayAbilities/Private/Serialization/InternalGameplayEffectContextNetSerializer.h b/Engine/Plugins/Runtime/GameplayAbilities/Source/GameplayAbilities/Private/Serialization/InternalGameplayEffectContextNetSerializer.h
--- a/Engine/Plugins/Runtime/GameplayAbilities/Source/GameplayAbilities/Private/Serialization/InternalGameplayEffectContextNetSerializer.h
+++ b/Engine/Plugins/Runtime/GameplayAbilities/Source/GameplayAbilities/Private/Serialization/InternalGameplayEffectContextNetSerializer.h
@@ -14,6 +14,12 @@ public:
 	void CopyReplicatedFieldsFrom(const FGameplayEffectContextAccessorForNetSerializer& GE);	
 	void AssignReplicatedFieldsTo(FGameplayEffectContextAccessorForNetSerializer& GE) const;
 
+	void SetShouldReplicateInstigator(bool bShouldReplicateInstigator) { bReplicateInstigator = bShouldReplicateInstigator; }
+	bool ShouldReplicateInstigator() const { return bReplicateInstigator; }
+
+	void SetShouldReplicateEffectCauser(bool bShouldReplicateEffectCauser) { bReplicateEffectCauser = bShouldReplicateEffectCauser; }
+	bool ShouldReplicateEffectCauser() const { return bReplicateEffectCauser; }
+
 	void SetShouldReplicateWorldOrigin(bool bShouldReplicateWorldOrigin) { bHasWorldOrigin = bShouldReplicateWorldOrigin; }
 	bool ShouldReplicateWorldOrigin() const { return bHasWorldOrigin; }
 
diff --git a/Engine/Plugins/Runtime/GameplayAbilities/Source/GameplayAbilities/Public/GameplayEffectTypes.h b/Engine/Plugins/Runtime/GameplayAbilities/Source/GameplayAbilities/Public/GameplayEffectTypes.h
--- a/Engine/Plugins/Runtime/GameplayAbilities/Source/GameplayAbilities/Public/GameplayEffectTypes.h
+++ b/Engine/Plugins/Runtime/GameplayAbilities/Source/GameplayAbilities/Public/GameplayEffectTypes.h
@@ -283,6 +283,8 @@ struct GAMEPLAYABILITIES_API FGameplayEffectContext
 	, WorldOrigin(ForceInitToZero)
 	, bHasWorldOrigin(false)
 	, bReplicateSourceObject(false)
+	, bReplicateInstigator(false)
+	, bReplicateEffectCauser(false)
 	{
 	}
 
@@ -291,6 +293,8 @@ struct GAMEPLAYABILITIES_API FGameplayEffectContext
 	, WorldOrigin(ForceInitToZero)
 	, bHasWorldOrigin(false)
 	, bReplicateSourceObject(false)
+	, bReplicateInstigator(false)
+	, bReplicateEffectCauser(false)
 	{
 		FGameplayEffectContext::AddInstigator(InInstigator, InEffectCauser);
 	}
@@ -342,6 +346,7 @@ struct GAMEPLAYABILITIES_API FGameplayEffectContext
 	void SetEffectCauser(AActor* InEffectCauser)
 	{
 		EffectCauser = InEffectCauser;
+		bReplicateEffectCauser = CanActorReferenceBeReplicated(InEffectCauser);
 	}
 
 	/** Should always return the original instigator that started the whole chain. Subclasses can override what this does */
@@ -440,6 +445,8 @@ struct GAMEPLAYABILITIES_API FGameplayEffectContext
 	virtual bool NetSerialize(FArchive& Ar, class UPackageMap* Map, bool& bOutSuccess);
 
 protected:
+	static bool CanActorReferenceBeReplicated(const AActor* Actor);
+
 	// The object pointers here have to be weak because contexts aren't necessarily tracked by GC in all cases
 
 	/** Instigator actor, the actor that owns the ability system component */
@@ -485,8 +492,16 @@ protected:
 	uint8 bHasWorldOrigin:1;
 
 	/** True if the SourceObject can be replicated. This bool is not replicated itself. */
-	UPROPERTY()
+	UPROPERTY(NotReplicated)
 	uint8 bReplicateSourceObject:1;
+	
+	/** True if the Instigator can be replicated. This bool is not replicated itself. */
+	UPROPERTY(NotReplicated)	
+	uint8 bReplicateInstigator:1;
+
+	/** True if the Instigator can be replicated. This bool is not replicated itself. */
+	UPROPERTY(NotReplicated)	
+	uint8 bReplicateEffectCauser:1;
 };
 
 template<>
@@ -566,7 +581,7 @@ struct GAMEPLAYABILITIES_API FGameplayEffectContextHandle
 		}
 	}
 
-	/** Sets Abilit instance and CDO parameters on context */
+	/** Sets Ability instance and CDO parameters on context */
 	void SetAbility(const UGameplayAbility* InGameplayAbility)
 	{
 		if (IsValid())
diff --git a/Engine/Source/Runtime/Net/Core/Classes/Net/Serialization/FastArraySerializer.h b/Engine/Source/Runtime/Net/Core/Classes/Net/Serialization/FastArraySerializer.h
--- a/Engine/Source/Runtime/Net/Core/Classes/Net/Serialization/FastArraySerializer.h
+++ b/Engine/Source/Runtime/Net/Core/Classes/Net/Serialization/FastArraySerializer.h
@@ -1226,8 +1226,6 @@ bool FFastArraySerializer::FastArrayDeltaSerialize(TArray<Type> &Items, FNetDelt
 			}
 		}
 
-		Helper.CallPostReplicatedReceiveOrNot(Items.Num());
-
 		return true;
 	}
 
@@ -1352,9 +1350,8 @@ bool FFastArraySerializer::FastArrayDeltaSerialize(TArray<Type> &Items, FNetDelt
 		if (Parms.bOutSomeObjectsWereMapped)
 		{
 			ArraySerializer.PostReplicatedChange(ChangedIndices, Items.Num());
+			Helper.CallPostReplicatedReceiveOrNot(Items.Num());
 		}
-		Helper.CallPostReplicatedReceiveOrNot(Items.Num());
-
 		return true;
 	}
 
@@ -1673,8 +1670,6 @@ bool FFastArraySerializer::FastArrayDeltaSerialize_DeltaSerializeStructs(TArray<
 	{
 		UE_LOG(LogNetFastTArray, Log, TEXT("FastArrayDeltaSerialize_DeltaSerializeStruct for %s. %s. GatherGuidReferences"), *InnerStruct->GetName(), *InnerStruct->GetOwnerStruct()->GetName());
 		Parms.NetSerializeCB->GatherGuidReferencesForFastArray(DeltaSerializeParams);
-		Helper.CallPostReplicatedReceiveOrNot(Items.Num());
-
 		return true;
 	}
 	else if (Parms.MoveGuidToUnmapped)
@@ -1695,8 +1690,8 @@ bool FFastArraySerializer::FastArrayDeltaSerialize_DeltaSerializeStructs(TArray<
 		if (Parms.bOutSomeObjectsWereMapped)
 		{
 			ArraySerializer.PostReplicatedChange(ChangedIndices, Items.Num());
+			Helper.CallPostReplicatedReceiveOrNot(Items.Num());
 		}
-		Helper.CallPostReplicatedReceiveOrNot(Items.Num());
 
 		return true;
 	}
