Add Device properties to Gameplay Cue Notify Burst effects

diff --git a/Engine/Plugins/Runtime/GameplayAbilities/Source/GameplayAbilities/Private/GameplayCueNotifyTypes.cpp b/Engine/Plugins/Runtime/GameplayAbilities/Source/GameplayAbilities/Private/GameplayCueNotifyTypes.cpp
--- a/Engine/Plugins/Runtime/GameplayAbilities/Source/GameplayAbilities/Private/GameplayCueNotifyTypes.cpp
+++ b/Engine/Plugins/Runtime/GameplayAbilities/Source/GameplayAbilities/Private/GameplayCueNotifyTypes.cpp
@@ -18,6 +18,8 @@
 #include "GameFramework/PlayerController.h"
 #include "GameFramework/PlayerState.h"
 #include "GameFramework/Character.h"
+#include "GameFramework/InputDeviceSubsystem.h"
+#include "GameFramework/InputDeviceProperties.h"
 #include "Sound/SoundWaveProcedural.h"
 
 #include UE_INLINE_GENERATED_CPP_BY_NAME(GameplayCueNotifyTypes)
@@ -846,6 +848,53 @@ bool FGameplayCueNotify_ForceFeedbackInfo::PlayForceFeedback(const FGameplayCueN
 	return false;
 }
 
+bool FGameplayCueNotify_InputDevicePropertyInfo::SetDeviceProperties(const FGameplayCueNotify_SpawnContext& SpawnContext, FGameplayCueNotify_SpawnResult& OutSpawnResult) const
+{
+	if (DeviceProperties.IsEmpty())
+	{
+		return false;
+	}
+	
+	if (APlayerController* TargetPC = SpawnContext.FindLocalPlayerController(EGameplayCueNotify_LocallyControlledSource::TargetActor))
+	{
+		// Apply any device properties from this gameplay cue
+		if (UInputDeviceSubsystem* System = UInputDeviceSubsystem::Get())
+		{
+			FSetDevicePropertyParams Params = {};
+			Params.UserId = TargetPC->GetPlatformUserId();
+			Params.bRemoveAfterEvaluationTime = true;
+			if (ensure(Params.UserId.IsValid()))
+			{
+				for (TObjectPtr<UInputDeviceProperty> Prop : DeviceProperties)
+				{
+					Params.DeviceProperty = Prop;
+					FInputDevicePropertyHandle Handle = System->SetDeviceProperty(Params);
+					ensure(Handle.IsValid());
+				}
+				return true;
+			}
+		}
+	}
+
+	return false;
+}
+
+void FGameplayCueNotify_InputDevicePropertyInfo::ValidateBurstAssets(UObject* ContainingAsset, const FString& Context, TArray<FText>& ValidationErrors) const
+{
+#if WITH_EDITORONLY_DATA
+	for (const TObjectPtr<UInputDeviceProperty> Prop : DeviceProperties)
+	{
+		if (!Prop)
+		{
+			ValidationErrors.Add(FText::Format(
+				LOCTEXT("ForceFeedback_ShouldNotLoop", "There is a null device property used in slot [{0}] for asset [{1}]."),
+				FText::AsCultureInvariant(Context),
+				FText::AsCultureInvariant(ContainingAsset->GetPathName())));
+		}
+	}
+#endif
+}
+
 void FGameplayCueNotify_ForceFeedbackInfo::ValidateBurstAssets(UObject* ContainingAsset, const FString& Context, TArray<FText>& ValidationErrors) const
 {
 #if WITH_EDITORONLY_DATA
@@ -968,6 +1017,7 @@ void FGameplayCueNotify_BurstEffects::ExecuteEffects(const FGameplayCueNotify_Sp
 	BurstCameraShake.PlayCameraShake(SpawnContext, OutSpawnResult);
 	BurstCameraLensEffect.PlayCameraLensEffect(SpawnContext, OutSpawnResult);
 	BurstForceFeedback.PlayForceFeedback(SpawnContext, OutSpawnResult);
+	BurstDevicePropertyEffect.SetDeviceProperties(SpawnContext, OutSpawnResult);
 	BurstDecal.SpawnDecal(SpawnContext, OutSpawnResult);
 }
 
@@ -986,6 +1036,7 @@ void FGameplayCueNotify_BurstEffects::ValidateAssociatedAssets(UObject* Containi
 	BurstCameraShake.ValidateBurstAssets(ContainingAsset, Context + TEXT(".BurstCameraShake"), ValidationErrors);
 	BurstCameraLensEffect.ValidateBurstAssets(ContainingAsset, Context + TEXT(".BurstCameraLensEffect"), ValidationErrors);
 	BurstForceFeedback.ValidateBurstAssets(ContainingAsset, Context + TEXT(".BurstForceFeedback"), ValidationErrors);
+	BurstDevicePropertyEffect.ValidateBurstAssets(ContainingAsset, Context + TEXT(".BurstDevicePropertyEffect"), ValidationErrors);
 }
 
 
@@ -1017,6 +1068,7 @@ void FGameplayCueNotify_LoopingEffects::StartEffects(const FGameplayCueNotify_Sp
 	LoopingCameraShake.PlayCameraShake(SpawnContext, OutSpawnResult);
 	LoopingCameraLensEffect.PlayCameraLensEffect(SpawnContext, OutSpawnResult);
 	LoopingForceFeedback.PlayForceFeedback(SpawnContext, OutSpawnResult);
+	LoopingInputDevicePropertyEffect.SetDeviceProperties(SpawnContext, OutSpawnResult);
 }
 
 void FGameplayCueNotify_LoopingEffects::StopEffects(FGameplayCueNotify_SpawnResult& SpawnResult) const
diff --git a/Engine/Plugins/Runtime/GameplayAbilities/Source/GameplayAbilities/Public/GameplayCueNotifyTypes.h b/Engine/Plugins/Runtime/GameplayAbilities/Source/GameplayAbilities/Public/GameplayCueNotifyTypes.h
--- a/Engine/Plugins/Runtime/GameplayAbilities/Source/GameplayAbilities/Public/GameplayCueNotifyTypes.h
+++ b/Engine/Plugins/Runtime/GameplayAbilities/Source/GameplayAbilities/Public/GameplayCueNotifyTypes.h
@@ -17,6 +17,7 @@ class USoundBase;
 class UCameraShakeBase;
 class ICameraLensEffectInterface;
 class UForceFeedbackEffect;
+class UInputDeviceProperty;
 class UForceFeedbackAttenuation;
 class UForceFeedbackComponent;
 class UDecalComponent;
@@ -611,6 +612,27 @@ public:
 };
 
 
+/**
+ * FGameplayCueNotify_InputDevicePropertyInfo
+ *
+ * Properties that specify how to set input device properties during a gameplay cue notify
+ */
+USTRUCT(BlueprintType)
+struct FGameplayCueNotify_InputDevicePropertyInfo
+{
+	GENERATED_BODY()
+
+	/** Set the device properties on specified on this struct on the Input Device Subsystem. */
+	bool SetDeviceProperties(const FGameplayCueNotify_SpawnContext& SpawnContext, FGameplayCueNotify_SpawnResult& OutSpawnResult) const;
+
+	/** Validate that the device properties in this effect are usable */
+	void ValidateBurstAssets(UObject* ContainingAsset, const FString& Context, TArray<FText>& ValidationErrors) const;
+	
+	/** Input Device properties to apply */
+	UPROPERTY(EditAnywhere, Instanced, BlueprintReadOnly, Category = GameplayCueNotify)
+	TArray<TObjectPtr<UInputDeviceProperty>> DeviceProperties;
+};
+
 /**
  * FGameplayCueNotify_DecalInfo
  *
@@ -707,6 +729,10 @@ protected:
 	UPROPERTY(EditDefaultsOnly, BlueprintReadOnly, Category = GameplayCueNotify)
 	FGameplayCueNotify_ForceFeedbackInfo BurstForceFeedback;
 
+	// Input device properties to be applied on gameplay cue execution
+	UPROPERTY(EditAnywhere, BlueprintReadOnly, Category = GameplayCueNotify)
+	FGameplayCueNotify_InputDevicePropertyInfo BurstDevicePropertyEffect;
+
 	// Decal to be spawned on gameplay cue execution.  Actor should have fade out time or override should be set so it will clean up properly.
 	UPROPERTY(EditDefaultsOnly, BlueprintReadOnly, Category = GameplayCueNotify)
 	FGameplayCueNotify_DecalInfo BurstDecal;
@@ -753,4 +779,8 @@ protected:
 	// Force feedback to be played on gameplay cue loop start.
 	UPROPERTY(EditDefaultsOnly, BlueprintReadOnly, Category = GameplayCueNotify)
 	FGameplayCueNotify_ForceFeedbackInfo LoopingForceFeedback;
+
+	// Input device properties to be applied on gameplay cue loop start.
+	UPROPERTY(EditAnywhere, BlueprintReadOnly, Category = GameplayCueNotify)
+	FGameplayCueNotify_InputDevicePropertyInfo LoopingInputDevicePropertyEffect;
 };
