Extract GameplayAbilityRepAnimMontage

diff --git a/Engine/Plugins/Runtime/GameplayAbilities/Source/GameplayAbilities/Private/AbilitySystemReplicationProxyInterface.cpp b/Engine/Plugins/Runtime/GameplayAbilities/Source/GameplayAbilities/Private/AbilitySystemReplicationProxyInterface.cpp
new file mode 100644
--- /dev/null
+++ b/Engine/Plugins/Runtime/GameplayAbilities/Source/GameplayAbilities/Private/AbilitySystemReplicationProxyInterface.cpp
@@ -0,0 +1,59 @@
+// Copyright Epic Games, Inc. All Rights Reserved.
+
+#pragma once
+
+#include "AbilitySystemReplicationProxyInterface.h"
+
+#include "GameplayPrediction.h"
+#include "GameplayEffect.h"
+
+
+void IAbilitySystemReplicationProxyInterface::Call_InvokeGameplayCueExecuted_FromSpec(const FGameplayEffectSpecForRPC Spec, FPredictionKey PredictionKey)
+{
+	NetMulticast_InvokeGameplayCueExecuted_FromSpec(Spec, PredictionKey);
+}
+
+void IAbilitySystemReplicationProxyInterface::Call_InvokeGameplayCueExecuted(const FGameplayTag GameplayCueTag, FPredictionKey PredictionKey, FGameplayEffectContextHandle EffectContext)
+{
+	NetMulticast_InvokeGameplayCueExecuted(GameplayCueTag, PredictionKey, EffectContext);
+}
+
+void IAbilitySystemReplicationProxyInterface::Call_InvokeGameplayCuesExecuted(const FGameplayTagContainer GameplayCueTags, FPredictionKey PredictionKey, FGameplayEffectContextHandle EffectContext)
+{
+	NetMulticast_InvokeGameplayCuesExecuted(GameplayCueTags, PredictionKey, EffectContext);
+}
+
+void IAbilitySystemReplicationProxyInterface::Call_InvokeGameplayCueExecuted_WithParams(const FGameplayTag GameplayCueTag, FPredictionKey PredictionKey, FGameplayCueParameters GameplayCueParameters)
+{
+	NetMulticast_InvokeGameplayCueExecuted_WithParams(GameplayCueTag, PredictionKey, GameplayCueParameters);
+}
+
+void IAbilitySystemReplicationProxyInterface::Call_InvokeGameplayCuesExecuted_WithParams(const FGameplayTagContainer GameplayCueTags, FPredictionKey PredictionKey, FGameplayCueParameters GameplayCueParameters)
+{
+	NetMulticast_InvokeGameplayCuesExecuted_WithParams(GameplayCueTags, PredictionKey, GameplayCueParameters);
+}
+
+void IAbilitySystemReplicationProxyInterface::Call_InvokeGameplayCueAdded(const FGameplayTag GameplayCueTag, FPredictionKey PredictionKey, FGameplayEffectContextHandle EffectContext)
+{
+	NetMulticast_InvokeGameplayCueAdded(GameplayCueTag, PredictionKey, EffectContext);
+}
+
+void IAbilitySystemReplicationProxyInterface::Call_InvokeGameplayCueAdded_WithParams(const FGameplayTag GameplayCueTag, FPredictionKey PredictionKey, FGameplayCueParameters Parameters)
+{
+	NetMulticast_InvokeGameplayCueAdded_WithParams(GameplayCueTag, PredictionKey, Parameters);
+}
+
+void IAbilitySystemReplicationProxyInterface::Call_InvokeGameplayCueAddedAndWhileActive_FromSpec(const FGameplayEffectSpecForRPC& Spec, FPredictionKey PredictionKey)
+{
+	NetMulticast_InvokeGameplayCueAddedAndWhileActive_FromSpec(Spec, PredictionKey);
+}
+
+void IAbilitySystemReplicationProxyInterface::Call_InvokeGameplayCueAddedAndWhileActive_WithParams(const FGameplayTag GameplayCueTag, FPredictionKey PredictionKey, FGameplayCueParameters GameplayCueParameters)
+{
+	NetMulticast_InvokeGameplayCueAddedAndWhileActive_WithParams(GameplayCueTag, PredictionKey, GameplayCueParameters);
+}
+
+void IAbilitySystemReplicationProxyInterface::Call_InvokeGameplayCuesAddedAndWhileActive_WithParams(const FGameplayTagContainer GameplayCueTags, FPredictionKey PredictionKey, FGameplayCueParameters GameplayCueParameters)
+{
+	NetMulticast_InvokeGameplayCuesAddedAndWhileActive_WithParams(GameplayCueTags, PredictionKey, GameplayCueParameters);
+}
\ No newline at end of file
diff --git a/Engine/Plugins/Runtime/GameplayAbilities/Source/GameplayAbilities/Private/GameplayAbilityRepAnimMontage.cpp b/Engine/Plugins/Runtime/GameplayAbilities/Source/GameplayAbilities/Private/GameplayAbilityRepAnimMontage.cpp
new file mode 100644
--- /dev/null
+++ b/Engine/Plugins/Runtime/GameplayAbilities/Source/GameplayAbilities/Private/GameplayAbilityRepAnimMontage.cpp
@@ -0,0 +1,83 @@
+
+// Copyright Epic Games, Inc. All Rights Reserved.
+
+#include "Abilities/GameplayAbilityRepAnimMontage.h"
+
+#include UE_INLINE_GENERATED_CPP_BY_NAME(GameplayAbilityRepAnimMontage)
+
+bool FGameplayAbilityRepAnimMontage::NetSerialize(FArchive& Ar, class UPackageMap* Map, bool& bOutSuccess)
+{
+	uint8 RepPosition = bRepPosition;
+	Ar.SerializeBits(&RepPosition, 1);
+	if (RepPosition)
+	{
+		bRepPosition = true;
+
+		// when rep'ing position, we don't want to skip correction
+		// and we don't need to force the section id to play
+		SectionIdToPlay = 0;
+		SkipPositionCorrection = false;
+
+		// @note: section frames have such a high amount of precision they use, when
+		// removing some of the position precision and packing it into a uint32 caused
+		// issues where ability code would pick the end of a previous section instead of
+		// the start of a new section. For now serializing the full position again.
+		Ar << Position;
+	}
+	else
+	{
+		bRepPosition = false;
+
+		// when rep'ing the section to play id, we want to skip
+		// correction, and don't want a position
+		SkipPositionCorrection = true;
+		Position = 0.0f;
+		Ar.SerializeBits(&SectionIdToPlay, 7);
+	}
+
+	uint8 bIsStopped = IsStopped;
+	Ar.SerializeBits(&bIsStopped, 1);
+	IsStopped = bIsStopped & 1;
+
+	if (Ar.EngineNetVer() < HISTORY_MONTAGE_PLAY_INST_ID_SERIALIZATION)
+	{
+		uint8 bForcePlayBit = 0;
+		Ar.SerializeBits(&bForcePlayBit, 1);
+
+		if (Ar.IsLoading())
+		{
+			// Emulate behavior of ForcePlayBit via PlayInstanceId when reading from older-revision data
+			PlayInstanceId = (bForcePlayBit ? 1 : 0);
+		}
+	}	
+
+	uint8 bSkipPositionCorrection = SkipPositionCorrection;
+	Ar.SerializeBits(&bSkipPositionCorrection, 1);
+	SkipPositionCorrection = bSkipPositionCorrection & 1;
+
+	uint8 SkipPlayRate = bSkipPlayRate;
+	Ar.SerializeBits(&SkipPlayRate, 1);
+	bSkipPlayRate = SkipPlayRate & 1;
+
+	Ar << AnimMontage;
+	Ar << PlayRate;
+	Ar << BlendTime;
+	Ar << NextSectionID;
+	if (Ar.EngineNetVer() >= HISTORY_MONTAGE_PLAY_INST_ID_SERIALIZATION)
+	{
+		Ar << PlayInstanceId;
+	}
+	PredictionKey.NetSerialize(Ar, Map, bOutSuccess);
+
+	bOutSuccess = true;
+	return true;
+}
+
+void FGameplayAbilityRepAnimMontage::SetRepAnimPositionMethod(ERepAnimPositionMethod InMethod)
+{
+	switch (InMethod)
+	{
+	case ERepAnimPositionMethod::Position: bRepPosition = true; break;
+	case ERepAnimPositionMethod::CurrentSectionId: bRepPosition = false; break;
+	}
+}
diff --git a/Engine/Plugins/Runtime/GameplayAbilities/Source/GameplayAbilities/Private/GameplayAbilityTypes.cpp b/Engine/Plugins/Runtime/GameplayAbilities/Source/GameplayAbilities/Private/GameplayAbilityTypes.cpp
--- a/Engine/Plugins/Runtime/GameplayAbilities/Source/GameplayAbilities/Private/GameplayAbilityTypes.cpp
+++ b/Engine/Plugins/Runtime/GameplayAbilities/Source/GameplayAbilities/Private/GameplayAbilityTypes.cpp
@@ -406,80 +406,3 @@ void FGameplayAbilityReplicatedDataContainer::PrintDebug()
 	ABILITY_LOG(Warning, TEXT("============================="));
 }
 
-bool FGameplayAbilityRepAnimMontage::NetSerialize(FArchive& Ar, class UPackageMap* Map, bool& bOutSuccess)
-{
-	uint8 RepPosition = bRepPosition;
-	Ar.SerializeBits(&RepPosition, 1);
-	if (RepPosition)
-	{
-		bRepPosition = true;
-
-		// when rep'ing position, we don't want to skip correction
-		// and we don't need to force the section id to play
-		SectionIdToPlay = 0;
-		SkipPositionCorrection = false;
-
-		// @note: section frames have such a high amount of precision they use, when
-		// removing some of the position precision and packing it into a uint32 caused
-		// issues where ability code would pick the end of a previous section instead of
-		// the start of a new section. For now serializing the full position again.
-		Ar << Position;
-	}
-	else
-	{
-		bRepPosition = false;
-
-		// when rep'ing the section to play id, we want to skip
-		// correction, and don't want a position
-		SkipPositionCorrection = true;
-		Position = 0.0f;
-		Ar.SerializeBits(&SectionIdToPlay, 7);
-	}
-
-	uint8 bIsStopped = IsStopped;
-	Ar.SerializeBits(&bIsStopped, 1);
-	IsStopped = bIsStopped & 1;
-
-	if (Ar.EngineNetVer() < HISTORY_MONTAGE_PLAY_INST_ID_SERIALIZATION)
-	{
-		uint8 bForcePlayBit = 0;
-		Ar.SerializeBits(&bForcePlayBit, 1);
-
-		if (Ar.IsLoading())
-		{
-			// Emulate behavior of ForcePlayBit via PlayInstanceId when reading from older-revision data
-			PlayInstanceId = (bForcePlayBit ? 1 : 0);
-		}
-	}	
-
-	uint8 bSkipPositionCorrection = SkipPositionCorrection;
-	Ar.SerializeBits(&bSkipPositionCorrection, 1);
-	SkipPositionCorrection = bSkipPositionCorrection & 1;
-
-	uint8 SkipPlayRate = bSkipPlayRate;
-	Ar.SerializeBits(&SkipPlayRate, 1);
-	bSkipPlayRate = SkipPlayRate & 1;
-
-	Ar << AnimMontage;
-	Ar << PlayRate;
-	Ar << BlendTime;
-	Ar << NextSectionID;
-	if (Ar.EngineNetVer() >= HISTORY_MONTAGE_PLAY_INST_ID_SERIALIZATION)
-	{
-		Ar << PlayInstanceId;
-	}
-	PredictionKey.NetSerialize(Ar, Map, bOutSuccess);
-
-	bOutSuccess = true;
-	return true;
-}
-
-void FGameplayAbilityRepAnimMontage::SetRepAnimPositionMethod(ERepAnimPositionMethod InMethod)
-{
-	switch (InMethod)
-	{
-	case ERepAnimPositionMethod::Position: bRepPosition = true; break;
-	case ERepAnimPositionMethod::CurrentSectionId: bRepPosition = false; break;
-	}
-}
-
diff --git a/Engine/Plugins/Runtime/GameplayAbilities/Source/GameplayAbilities/Public/Abilities/GameplayAbilityRepAnimMontage.h b/Engine/Plugins/Runtime/GameplayAbilities/Source/GameplayAbilities/Public/Abilities/GameplayAbilityRepAnimMontage.h
new file mode 100644
--- /dev/null
+++ b/Engine/Plugins/Runtime/GameplayAbilities/Source/GameplayAbilities/Public/Abilities/GameplayAbilityRepAnimMontage.h
@@ -0,0 +1,98 @@
+// Copyright Epic Games, Inc. All Rights Reserved.
+
+#pragma once
+
+#include "GameplayPrediction.h"
+#include "GameplayAbilityRepAnimMontage.generated.h"
+
+class UAnimMontage;
+
+/** Enum used by the Ability Rep Anim Montage struct to rep the quantized position or the current section id */
+UENUM()
+enum class ERepAnimPositionMethod
+{
+	Position = 0,			// reps the position in the montage to keep the client in sync (heavier, quantized, more precise)
+	CurrentSectionId = 1,	// reps the current section id we want to play on the client (compact, less precise)
+};
+
+/** Data about montages that is replicated to simulated clients */
+USTRUCT()
+struct GAMEPLAYABILITIES_API FGameplayAbilityRepAnimMontage
+{
+	GENERATED_USTRUCT_BODY()
+
+	/** AnimMontage ref */
+	UPROPERTY()
+	TObjectPtr<UAnimMontage> AnimMontage;
+
+	/** Play Rate */
+	UPROPERTY()
+	float PlayRate;
+
+	/** Montage position */
+	UPROPERTY(NotReplicated)
+	float Position;
+
+	/** Montage current blend time */
+	UPROPERTY()
+	float BlendTime;
+
+	/** NextSectionID */
+	UPROPERTY()
+	uint8 NextSectionID;
+
+	/** ID incremented every time a montage is played, used to trigger replication when the same montage is played multiple times. This ID wraps around when it reaches its max value. */
+	UPROPERTY()
+	uint8 PlayInstanceId;
+
+	/** flag indicating we should serialize the position or the current section id */
+	UPROPERTY()
+	uint8 bRepPosition : 1;
+
+	/** Bit set when montage has been stopped. */
+	UPROPERTY()
+	uint8 IsStopped : 1;
+	
+	/** Stops montage position from replicating at all to save bandwidth */
+	UPROPERTY()
+	uint8 SkipPositionCorrection : 1;
+
+	/** Stops PlayRate from replicating to save bandwidth. PlayRate will be assumed to be 1.f. */
+	UPROPERTY()
+	uint8 bSkipPlayRate : 1;
+
+	UPROPERTY()
+	FPredictionKey PredictionKey;
+
+	/** The current section Id used by the montage. Will only be valid if bRepPosition is false */
+	UPROPERTY(NotReplicated)
+	uint8 SectionIdToPlay;
+
+	FGameplayAbilityRepAnimMontage()
+	: AnimMontage(nullptr),
+	PlayRate(0.f),
+	Position(0.f),
+	BlendTime(0.f),
+	NextSectionID(0),
+	PlayInstanceId(0),
+	bRepPosition(true),
+	IsStopped(true),
+	SkipPositionCorrection(false),
+	bSkipPlayRate(false),
+	SectionIdToPlay(0)
+	{
+	}
+
+	bool NetSerialize(FArchive& Ar, class UPackageMap* Map, bool& bOutSuccess);
+
+	void SetRepAnimPositionMethod(ERepAnimPositionMethod InMethod);
+};
+
+template<>
+struct TStructOpsTypeTraits<FGameplayAbilityRepAnimMontage> : public TStructOpsTypeTraitsBase2<FGameplayAbilityRepAnimMontage>
+{
+	enum
+	{
+		WithNetSerializer = true,
+	};
+};
diff --git a/Engine/Plugins/Runtime/GameplayAbilities/Source/GameplayAbilities/Public/Abilities/GameplayAbilityTypes.h b/Engine/Plugins/Runtime/GameplayAbilities/Source/GameplayAbilities/Public/Abilities/GameplayAbilityTypes.h
--- a/Engine/Plugins/Runtime/GameplayAbilities/Source/GameplayAbilities/Public/Abilities/GameplayAbilityTypes.h
+++ b/Engine/Plugins/Runtime/GameplayAbilities/Source/GameplayAbilities/Public/Abilities/GameplayAbilityTypes.h
@@ -16,6 +16,9 @@
 #include "Animation/AnimInstance.h"
 #include "Components/SkeletalMeshComponent.h"
 #endif
+#if UE_ENABLE_INCLUDE_ORDER_DEPRECATED_IN_5_2
+#include "Abilities/GameplayAbilityRepAnimMontage.h"
+#endif
 #include "Abilities/GameplayAbilityTargetTypes.h"
 #include "GameplayAbilityTypes.generated.h"
 
@@ -196,95 +199,6 @@ struct GAMEPLAYABILITIES_API FGameplayAbilityActorInfo
 	virtual void ClearActorInfo();
 };
 
-/** Enum used by the Ability Rep Anim Montage struct to rep the quantized position or the current section id */
-UENUM()
-enum class ERepAnimPositionMethod
-{
-	Position = 0,			// reps the position in the montage to keep the client in sync (heavier, quantized, more precise)
-	CurrentSectionId = 1,	// reps the current section id we want to play on the client (compact, less precise)
-};
-
-/** Data about montages that is replicated to simulated clients */
-USTRUCT()
-struct GAMEPLAYABILITIES_API FGameplayAbilityRepAnimMontage
-{
-	GENERATED_USTRUCT_BODY()
-
-	/** AnimMontage ref */
-	UPROPERTY()
-	TObjectPtr<UAnimMontage> AnimMontage;
-
-	/** Play Rate */
-	UPROPERTY()
-	float PlayRate;
-
-	/** Montage position */
-	UPROPERTY(NotReplicated)
-	float Position;
-
-	/** Montage current blend time */
-	UPROPERTY()
-	float BlendTime;
-
-	/** NextSectionID */
-	UPROPERTY()
-	uint8 NextSectionID;
-
-	/** ID incremented every time a montage is played, used to trigger replication when the same montage is played multiple times. This ID wraps around when it reaches its max value. */
-	UPROPERTY()
-	uint8 PlayInstanceId;
-
-	/** flag indicating we should serialize the position or the current section id */
-	UPROPERTY()
-	uint8 bRepPosition : 1;
-
-	/** Bit set when montage has been stopped. */
-	UPROPERTY()
-	uint8 IsStopped : 1;
-	
-	/** Stops montage position from replicating at all to save bandwidth */
-	UPROPERTY()
-	uint8 SkipPositionCorrection : 1;
-
-	/** Stops PlayRate from replicating to save bandwidth. PlayRate will be assumed to be 1.f. */
-	UPROPERTY()
-	uint8 bSkipPlayRate : 1;
-
-	UPROPERTY()
-	FPredictionKey PredictionKey;
-
-	/** The current section Id used by the montage. Will only be valid if bRepPosition is false */
-	UPROPERTY(NotReplicated)
-	uint8 SectionIdToPlay;
-
-	FGameplayAbilityRepAnimMontage()
-	: AnimMontage(nullptr),
-	PlayRate(0.f),
-	Position(0.f),
-	BlendTime(0.f),
-	NextSectionID(0),
-	PlayInstanceId(0),
-	bRepPosition(true),
-	IsStopped(true),
-	SkipPositionCorrection(false),
-	bSkipPlayRate(false),
-	SectionIdToPlay(0)
-	{
-	}
-
-	bool NetSerialize(FArchive& Ar, class UPackageMap* Map, bool& bOutSuccess);
-
-	void SetRepAnimPositionMethod(ERepAnimPositionMethod InMethod);
-};
-
-template<>
-struct TStructOpsTypeTraits<FGameplayAbilityRepAnimMontage> : public TStructOpsTypeTraitsBase2<FGameplayAbilityRepAnimMontage>
-{
-	enum
-	{
-		WithNetSerializer = true,
-	};
-};
 
 
 /** Data about montages that were played locally (all montages in case of server. predictive montages in case of client). Never replicated directly. */
diff --git a/Engine/Plugins/Runtime/GameplayAbilities/Source/GameplayAbilities/Public/AbilitySystemComponent.h b/Engine/Plugins/Runtime/GameplayAbilities/Source/GameplayAbilities/Public/AbilitySystemComponent.h
--- a/Engine/Plugins/Runtime/GameplayAbilities/Source/GameplayAbilities/Public/AbilitySystemComponent.h
+++ b/Engine/Plugins/Runtime/GameplayAbilities/Source/GameplayAbilities/Public/AbilitySystemComponent.h
@@ -7,23 +7,27 @@
 #include "Templates/SubclassOf.h"
 #include "Engine/NetSerialization.h"
 #include "Engine/EngineTypes.h"
+#include "Engine/TimerHandle.h"
 #include "GameplayTagContainer.h"
 #include "AttributeSet.h"
 #include "EngineDefines.h"
-#if UE_ENABLE_INCLUDE_ORDER_DEPRECATED_IN_5_2
-#include "GameplayEffectTypes.h"
-#endif
 #include "GameplayPrediction.h"
 #include "GameplayCueInterface.h"
 #include "GameplayTagAssetInterface.h"
 #include "GameplayAbilitySpec.h"
 #include "GameplayEffect.h"
-#include "Abilities/GameplayAbilityTypes.h"
 #include "GameplayTasksComponent.h"
+#include "Abilities/GameplayAbilityRepAnimMontage.h"
 #include "Abilities/GameplayAbilityTargetTypes.h"
 #include "Abilities/GameplayAbility.h"
 #include "AbilitySystemReplicationProxyInterface.h"
 #include "Net/Core/PushModel/PushModel.h"
+
+#if UE_ENABLE_INCLUDE_ORDER_DEPRECATED_IN_5_2
+#include "Abilities/GameplayAbilityTypes.h"
+#include "GameplayEffectTypes.h"
+#endif
+
 #include "AbilitySystemComponent.generated.h"
 
 class AGameplayAbilityTargetActor;
diff --git a/Engine/Plugins/Runtime/GameplayAbilities/Source/GameplayAbilities/Public/AbilitySystemReplicationProxyInterface.h b/Engine/Plugins/Runtime/GameplayAbilities/Source/GameplayAbilities/Public/AbilitySystemReplicationProxyInterface.h
--- a/Engine/Plugins/Runtime/GameplayAbilities/Source/GameplayAbilities/Public/AbilitySystemReplicationProxyInterface.h
+++ b/Engine/Plugins/Runtime/GameplayAbilities/Source/GameplayAbilities/Public/AbilitySystemReplicationProxyInterface.h
@@ -5,11 +5,19 @@
 #include "CoreMinimal.h"
 #include "UObject/ObjectMacros.h"
 #include "UObject/Interface.h"
+#if UE_ENABLE_INCLUDE_ORDER_DEPRECATED_IN_5_1
 #include "GameplayPrediction.h"
 #include "GameplayEffect.h"
+#endif
 #include "AbilitySystemReplicationProxyInterface.generated.h"
 
 class UAbilitySystemComponent;
+struct FGameplayCueParameters;
+struct FGameplayEffectContextHandle;
+struct FGameplayEffectSpecForRPC;
+struct FGameplayTag;
+struct FGameplayTagContainer;
+struct FPredictionKey;
 
 /** Interface for actors that act like an ability system component for replication. This can be used to group together different components for replication */
 UINTERFACE(MinimalAPI, meta = (CannotImplementInterfaceInBlueprint))
@@ -29,55 +37,16 @@ class GAMEPLAYABILITIES_API IAbilitySystemReplicationProxyInterface
 	//	Call Proxies: these can be optionally implemented to intercept GameplayCue events. By default, they forward to the multicasts.
 	// --------------------------------------------------------------------------------------------
 
-	virtual void Call_InvokeGameplayCueExecuted_FromSpec(const FGameplayEffectSpecForRPC Spec, FPredictionKey PredictionKey)
-	{
-		NetMulticast_InvokeGameplayCueExecuted_FromSpec(Spec, PredictionKey);
-	}
-
-	virtual void Call_InvokeGameplayCueExecuted(const FGameplayTag GameplayCueTag, FPredictionKey PredictionKey, FGameplayEffectContextHandle EffectContext)
-	{
-		NetMulticast_InvokeGameplayCueExecuted(GameplayCueTag, PredictionKey, EffectContext);
-	}
-
-	virtual void Call_InvokeGameplayCuesExecuted(const FGameplayTagContainer GameplayCueTags, FPredictionKey PredictionKey, FGameplayEffectContextHandle EffectContext)
-	{
-		NetMulticast_InvokeGameplayCuesExecuted(GameplayCueTags, PredictionKey, EffectContext);
-	}
-	
-	virtual void Call_InvokeGameplayCueExecuted_WithParams(const FGameplayTag GameplayCueTag, FPredictionKey PredictionKey, FGameplayCueParameters GameplayCueParameters)
-	{
-		NetMulticast_InvokeGameplayCueExecuted_WithParams(GameplayCueTag, PredictionKey, GameplayCueParameters);
-	}
-	
-	virtual void Call_InvokeGameplayCuesExecuted_WithParams(const FGameplayTagContainer GameplayCueTags, FPredictionKey PredictionKey, FGameplayCueParameters GameplayCueParameters)
-	{
-		NetMulticast_InvokeGameplayCuesExecuted_WithParams(GameplayCueTags, PredictionKey, GameplayCueParameters);
-	}
-	
-	virtual void Call_InvokeGameplayCueAdded(const FGameplayTag GameplayCueTag, FPredictionKey PredictionKey, FGameplayEffectContextHandle EffectContext)
-	{
-		NetMulticast_InvokeGameplayCueAdded(GameplayCueTag, PredictionKey, EffectContext);
-	}
-	
-	virtual void Call_InvokeGameplayCueAdded_WithParams(const FGameplayTag GameplayCueTag, FPredictionKey PredictionKey, FGameplayCueParameters Parameters)
-	{
-		NetMulticast_InvokeGameplayCueAdded_WithParams(GameplayCueTag, PredictionKey, Parameters);
-	}
-	
-	virtual void Call_InvokeGameplayCueAddedAndWhileActive_FromSpec(const FGameplayEffectSpecForRPC& Spec, FPredictionKey PredictionKey)
-	{
-		NetMulticast_InvokeGameplayCueAddedAndWhileActive_FromSpec(Spec, PredictionKey);
-	}
-	
-	virtual void Call_InvokeGameplayCueAddedAndWhileActive_WithParams(const FGameplayTag GameplayCueTag, FPredictionKey PredictionKey, FGameplayCueParameters GameplayCueParameters)
-	{
-		NetMulticast_InvokeGameplayCueAddedAndWhileActive_WithParams(GameplayCueTag, PredictionKey, GameplayCueParameters);
-	}
-	
-	virtual void Call_InvokeGameplayCuesAddedAndWhileActive_WithParams(const FGameplayTagContainer GameplayCueTags, FPredictionKey PredictionKey, FGameplayCueParameters GameplayCueParameters)
-	{
-		NetMulticast_InvokeGameplayCuesAddedAndWhileActive_WithParams(GameplayCueTags, PredictionKey, GameplayCueParameters);
-	}
+	virtual void Call_InvokeGameplayCueExecuted_FromSpec(const FGameplayEffectSpecForRPC Spec, FPredictionKey PredictionKey);
+	virtual void Call_InvokeGameplayCueExecuted(const FGameplayTag GameplayCueTag, FPredictionKey PredictionKey, FGameplayEffectContextHandle EffectContext);
+	virtual void Call_InvokeGameplayCuesExecuted(const FGameplayTagContainer GameplayCueTags, FPredictionKey PredictionKey, FGameplayEffectContextHandle EffectContext);
+	virtual void Call_InvokeGameplayCueExecuted_WithParams(const FGameplayTag GameplayCueTag, FPredictionKey PredictionKey, FGameplayCueParameters GameplayCueParameters);
+	virtual void Call_InvokeGameplayCuesExecuted_WithParams(const FGameplayTagContainer GameplayCueTags, FPredictionKey PredictionKey, FGameplayCueParameters GameplayCueParameters);
+	virtual void Call_InvokeGameplayCueAdded(const FGameplayTag GameplayCueTag, FPredictionKey PredictionKey, FGameplayEffectContextHandle EffectContext);
+	virtual void Call_InvokeGameplayCueAdded_WithParams(const FGameplayTag GameplayCueTag, FPredictionKey PredictionKey, FGameplayCueParameters Parameters);
+	virtual void Call_InvokeGameplayCueAddedAndWhileActive_FromSpec(const FGameplayEffectSpecForRPC& Spec, FPredictionKey PredictionKey);
+	virtual void Call_InvokeGameplayCueAddedAndWhileActive_WithParams(const FGameplayTag GameplayCueTag, FPredictionKey PredictionKey, FGameplayCueParameters GameplayCueParameters);
+	virtual void Call_InvokeGameplayCuesAddedAndWhileActive_WithParams(const FGameplayTagContainer GameplayCueTags, FPredictionKey PredictionKey, FGameplayCueParameters GameplayCueParameters);
 
 	// ----------------------------------------------
 	//	Multicast Proxies: these should be implemented as unreliable NetMulticast functions
diff --git a/Engine/Plugins/Runtime/GameplayAbilities/Source/GameplayAbilities/Public/GameplayEffectTypes.h b/Engine/Plugins/Runtime/GameplayAbilities/Source/GameplayAbilities/Public/GameplayEffectTypes.h
--- a/Engine/Plugins/Runtime/GameplayAbilities/Source/GameplayAbilities/Public/GameplayEffectTypes.h
+++ b/Engine/Plugins/Runtime/GameplayAbilities/Source/GameplayAbilities/Public/GameplayEffectTypes.h
@@ -100,7 +100,7 @@ UENUM(BlueprintType)
 namespace EGameplayModOp
 {
 	/** Defines the ways that mods will modify attributes. Numeric ones operate on the existing value, override ignores it */
-	enum Type
+	enum Type : int
 	{		
 		/** Numeric. */
 		Additive = 0		UMETA(DisplayName="Add"),
@@ -976,7 +976,7 @@ UENUM(BlueprintType)
 namespace EGameplayCueEvent
 {
 	/** Indicates what type of action happened to a specific gameplay cue tag. Sometimes you will get multiple events at once */
-	enum Type
+	enum Type : int
 	{
 		/** Called when a GameplayCue with duration is first activated, this will only be called if the client witnessed the activation */
 		OnActive,
@@ -1043,7 +1043,7 @@ UENUM(BlueprintType)
 namespace EGameplayTagEventType
 {
 	/** Rather a tag was added or removed, used in callbacks */
-	enum Type
+	enum Type : int
 	{		
 		/** Event only happens when tag is new or completely removed */
 		NewOrRemoved,
