Fix duplicate FText key
#rnx

diff --git a/Engine/Plugins/Runtime/GameplayAbilities/Source/GameplayAbilities/Private/GameplayCueNotifyTypes.cpp b/Engine/Plugins/Runtime/GameplayAbilities/Source/GameplayAbilities/Private/GameplayCueNotifyTypes.cpp
--- a/Engine/Plugins/Runtime/GameplayAbilities/Source/GameplayAbilities/Private/GameplayCueNotifyTypes.cpp
+++ b/Engine/Plugins/Runtime/GameplayAbilities/Source/GameplayAbilities/Private/GameplayCueNotifyTypes.cpp
@@ -887,7 +887,7 @@ void FGameplayCueNotify_InputDevicePropertyInfo::ValidateBurstAssets(UObject* Co
 		if (!Prop)
 		{
 			ValidationErrors.Add(FText::Format(
-				LOCTEXT("ForceFeedback_ShouldNotLoop", "There is a null device property used in slot [{0}] for asset [{1}]."),
+				LOCTEXT("InputDeviceProperty_ShouldNotBeNull", "There is a null device property used in slot [{0}] for asset [{1}]."),
 				FText::AsCultureInvariant(Context),
 				FText::AsCultureInvariant(ContainingAsset->GetPathName())));
 		}
