PR #9829: Fix typo in IsValid(Object) API upgrade (Contributed by [REDACTED_PERSON])

diff --git a/Engine/Plugins/Runtime/GameplayAbilities/Source/GameplayAbilities/Private/GameplayCueManager.cpp b/Engine/Plugins/Runtime/GameplayAbilities/Source/GameplayAbilities/Private/GameplayCueManager.cpp
--- a/Engine/Plugins/Runtime/GameplayAbilities/Source/GameplayAbilities/Private/GameplayCueManager.cpp
+++ b/Engine/Plugins/Runtime/GameplayAbilities/Source/GameplayAbilities/Private/GameplayCueManager.cpp
@@ -1569,7 +1569,7 @@ void UGameplayCueManager::UpdatePreallocation(UWorld* World)
 		AGameplayCueNotify_Actor* PrespawnedInstance = Cast<AGameplayCueNotify_Actor>(World->SpawnActor(CDO->GetClass()));
 		if (ensureMsgf(PrespawnedInstance, TEXT("Failed to prespawn GC notify for: %s"), *GetNameSafe(CDO)))
 		{
-			ensureMsgf(IsValid(PrespawnedInstance) == false, TEXT("Newly spawned GC is PendingKILL: %s"), *GetNameSafe(CDO));
+			ensureMsgf(IsValid(PrespawnedInstance) == true, TEXT("Newly spawned GC is invalid: %s"), *GetNameSafe(CDO));
 
 			if (LogGameplayCueActorSpawning)
 			{
