Fixed Bug# UE-172641 - AbilityTask_NetworkSyncPoint can cause duplicate Gameplay Effect applications

Fixing an issue where it was possible for locally predicted GE's to be added w/ a different prediction key than the server causing it to fail to clean up.

diff --git a/Engine/Plugins/Runtime/GameplayAbilities/Source/GameplayAbilities/Private/Abilities/Tasks/AbilityTask_NetworkSyncPoint.cpp b/Engine/Plugins/Runtime/GameplayAbilities/Source/GameplayAbilities/Private/Abilities/Tasks/AbilityTask_NetworkSyncPoint.cpp
--- a/Engine/Plugins/Runtime/GameplayAbilities/Source/GameplayAbilities/Private/Abilities/Tasks/AbilityTask_NetworkSyncPoint.cpp
+++ b/Engine/Plugins/Runtime/GameplayAbilities/Source/GameplayAbilities/Private/Abilities/Tasks/AbilityTask_NetworkSyncPoint.cpp
@@ -29,29 +29,27 @@ UAbilityTask_NetworkSyncPoint* UAbilityTask_NetworkSyncPoint::WaitNetSync(class
 
 void UAbilityTask_NetworkSyncPoint::Activate()
 {
+	FScopedPredictionWindow ScopedPrediction(AbilitySystemComponent.Get(), IsPredictingClient());
+
 	if (AbilitySystemComponent.IsValid())
 	{
 		if (IsPredictingClient())
 		{
-			if (SyncType != EAbilityTaskNetSyncType::OnlyServerWait)
+			if (SyncType != EAbilityTaskNetSyncType::OnlyServerWait )
 			{
 				// As long as we are waiting (!= OnlyServerWait), listen for the GenericSignalFromServer event
 				ReplicatedEventToListenFor = EAbilityGenericReplicatedEvent::GenericSignalFromServer;
 			}
 			if (SyncType != EAbilityTaskNetSyncType::OnlyClientWait)
 			{
-				// @note: When on a predicting client we need to make sure we send off the scoped prediction key to the server so things like gameplay
-				// cue's that replicate down will continue to not execute on a client that is predicting abilities.
-				FScopedPredictionWindow ScopedPrediction(AbilitySystemComponent.Get(), IsPredictingClient());
-
 				// As long as the server is waiting (!= OnlyClientWait), send the Server and RPC for this signal
 				AbilitySystemComponent->ServerSetReplicatedEvent(EAbilityGenericReplicatedEvent::GenericSignalFromClient, GetAbilitySpecHandle(), GetActivationPredictionKey(), AbilitySystemComponent->ScopedPredictionKey);
 			}
-
+			
 		}
 		else if (IsForRemoteClient())
 		{
-			if (SyncType != EAbilityTaskNetSyncType::OnlyClientWait)
+			if (SyncType != EAbilityTaskNetSyncType::OnlyClientWait )
 			{
 				// As long as we are waiting (!= OnlyClientWait), listen for the GenericSignalFromClient event
 				ReplicatedEventToListenFor = EAbilityGenericReplicatedEvent::GenericSignalFromClient;
@@ -79,12 +77,6 @@ void UAbilityTask_NetworkSyncPoint::SyncFinished()
 {
 	if (IsValid(this))
 	{
-		// @note: When we finish the task to let the ability resume execution we need to ensure predicting clients have a valid
-		// scoped prediction key so they can continue to execute things like gameplay cues. This is because the server will pass
-		// along a client prediction key through its replication chain ensuring predicted abilities don't execute predicted
-		// features (i.e. gameplay cues)
-		FScopedPredictionWindow ScopedPrediction(AbilitySystemComponent.Get(), IsPredictingClient());
-
 		if (ShouldBroadcastAbilityTaskDelegates())
 		{
 			OnSync.Broadcast();
