Improved function desriptions for ability async types, to avoid confusion with their ability task-derived equivalents and recommend them for ability graphs.

#rnx

diff --git a/Engine/Plugins/Runtime/GameplayAbilities/Source/GameplayAbilities/Public/Abilities/Async/AbilityAsync_WaitAttributeChanged.h b/Engine/Plugins/Runtime/GameplayAbilities/Source/GameplayAbilities/Public/Abilities/Async/AbilityAsync_WaitAttributeChanged.h
--- a/Engine/Plugins/Runtime/GameplayAbilities/Source/GameplayAbilities/Public/Abilities/Async/AbilityAsync_WaitAttributeChanged.h
+++ b/Engine/Plugins/Runtime/GameplayAbilities/Source/GameplayAbilities/Public/Abilities/Async/AbilityAsync_WaitAttributeChanged.h
@@ -16,6 +16,11 @@ class GAMEPLAYABILITIES_API UAbilityAsync_WaitAttributeChanged : public UAbility
 
 public:
 
+	/**
+	 * Wait until the specified gameplay attribute is changed on a target ability system component
+	 * It will keep listening as long as OnlyTriggerOnce = false
+	 * If used in an ability graph, this async action will wait even after activation ends. It's recommended to use WaitForAttributeChange instead.
+	 */
 	UFUNCTION(BlueprintCallable, Category = "Ability|Tasks", meta = (DefaultToSelf = "TargetActor", BlueprintInternalUseOnly = "TRUE"))
 	static UAbilityAsync_WaitAttributeChanged* WaitForAttributeChanged(AActor* TargetActor, FGameplayAttribute Attribute, bool OnlyTriggerOnce = false);
 
diff --git a/Engine/Plugins/Runtime/GameplayAbilities/Source/GameplayAbilities/Public/Abilities/Async/AbilityAsync_WaitGameplayEffectApplied.h b/Engine/Plugins/Runtime/GameplayAbilities/Source/GameplayAbilities/Public/Abilities/Async/AbilityAsync_WaitGameplayEffectApplied.h
--- a/Engine/Plugins/Runtime/GameplayAbilities/Source/GameplayAbilities/Public/Abilities/Async/AbilityAsync_WaitGameplayEffectApplied.h
+++ b/Engine/Plugins/Runtime/GameplayAbilities/Source/GameplayAbilities/Public/Abilities/Async/AbilityAsync_WaitGameplayEffectApplied.h
@@ -25,6 +25,7 @@ public:
 	/**
 	 * Wait until a GameplayEffect is applied to a target actor
 	 * If TriggerOnce is true, this action will only activate one time. Otherwise it will return every time a GE is applied that meets the requirements over the life of the ability
+	 * If used in an ability graph, this async action will wait even after activation ends. It's recommended to use WaitGameplayEffectApplied instead.
 	 */
 	UFUNCTION(BlueprintCallable, Category = "Ability|Async", meta = (DefaultToSelf = "TargetActor", BlueprintInternalUseOnly = "TRUE"))
 	static UAbilityAsync_WaitGameplayEffectApplied* WaitGameplayEffectAppliedToActor(AActor* TargetActor, const FGameplayTargetDataFilterHandle SourceFilter, FGameplayTagRequirements SourceTagRequirements, FGameplayTagRequirements TargetTagRequirements, bool TriggerOnce = false, bool ListenForPeriodicEffect = false);
diff --git a/Engine/Plugins/Runtime/GameplayAbilities/Source/GameplayAbilities/Public/Abilities/Async/AbilityAsync_WaitGameplayEvent.h b/Engine/Plugins/Runtime/GameplayAbilities/Source/GameplayAbilities/Public/Abilities/Async/AbilityAsync_WaitGameplayEvent.h
--- a/Engine/Plugins/Runtime/GameplayAbilities/Source/GameplayAbilities/Public/Abilities/Async/AbilityAsync_WaitGameplayEvent.h
+++ b/Engine/Plugins/Runtime/GameplayAbilities/Source/GameplayAbilities/Public/Abilities/Async/AbilityAsync_WaitGameplayEvent.h
@@ -19,6 +19,7 @@ public:
 	 * Wait until the specified gameplay tag event is triggered on a target ability system component
 	 * It will keep listening as long as OnlyTriggerOnce = false
 	 * If OnlyMatchExact = false it will trigger for nested tags
+	 * If used in an ability graph, this async action will wait even after activation ends. It's recommended to use WaitGameplayEvent instead.
 	 */
 	UFUNCTION(BlueprintCallable, Category = "Ability|Async", meta = (DefaultToSelf = "TargetActor", BlueprintInternalUseOnly = "TRUE"))
 	static UAbilityAsync_WaitGameplayEvent* WaitGameplayEventToActor(AActor* TargetActor, FGameplayTag EventTag, bool OnlyTriggerOnce = false, bool OnlyMatchExact = true);
diff --git a/Engine/Plugins/Runtime/GameplayAbilities/Source/GameplayAbilities/Public/Abilities/Async/AbilityAsync_WaitGameplayTag.h b/Engine/Plugins/Runtime/GameplayAbilities/Source/GameplayAbilities/Public/Abilities/Async/AbilityAsync_WaitGameplayTag.h
--- a/Engine/Plugins/Runtime/GameplayAbilities/Source/GameplayAbilities/Public/Abilities/Async/AbilityAsync_WaitGameplayTag.h
+++ b/Engine/Plugins/Runtime/GameplayAbilities/Source/GameplayAbilities/Public/Abilities/Async/AbilityAsync_WaitGameplayTag.h
@@ -42,6 +42,7 @@ public:
 	/**
 	 * Wait until the specified gameplay tag is Added to Target Actor's ability component
 	 * If the tag is already present when this task is started, it will immediately broadcast the Added event. It will keep listening as long as OnlyTriggerOnce = false.
+	 * If used in an ability graph, this async action will wait even after activation ends. It's recommended to use WaitGameplayTagAdd instead.
 	 */
 	UFUNCTION(BlueprintCallable, Category = "Ability|Tasks", meta = (DefaultToSelf = "TargetActor", BlueprintInternalUseOnly = "TRUE"))
 	static UAbilityAsync_WaitGameplayTagAdded* WaitGameplayTagAddToActor(AActor* TargetActor, FGameplayTag Tag, bool OnlyTriggerOnce=false);
@@ -61,6 +62,7 @@ public:
 	/**
 	 * Wait until the specified gameplay tag is Removed from Target Actor's ability component
 	 * If the tag is not present when this task is started, it will immediately broadcast the Removed event. It will keep listening as long as OnlyTriggerOnce = false.
+	 * If used in an ability graph, this async action will wait even after activation ends. It's recommended to use WaitGameplayTagRemove instead.
 	 */
 	UFUNCTION(BlueprintCallable, Category = "Ability|Tasks", meta = (DefaultToSelf = "TargetActor", BlueprintInternalUseOnly = "TRUE"))
 	static UAbilityAsync_WaitGameplayTagRemoved* WaitGameplayTagRemoveFromActor(AActor* TargetActor, FGameplayTag Tag, bool OnlyTriggerOnce=false);
diff --git a/Engine/Plugins/Runtime/GameplayAbilities/Source/GameplayAbilities/Public/Abilities/Async/AbilityAsync_WaitGameplayTagQuery.h b/Engine/Plugins/Runtime/GameplayAbilities/Source/GameplayAbilities/Public/Abilities/Async/AbilityAsync_WaitGameplayTagQuery.h
--- a/Engine/Plugins/Runtime/GameplayAbilities/Source/GameplayAbilities/Public/Abilities/Async/AbilityAsync_WaitGameplayTagQuery.h
+++ b/Engine/Plugins/Runtime/GameplayAbilities/Source/GameplayAbilities/Public/Abilities/Async/AbilityAsync_WaitGameplayTagQuery.h
@@ -62,6 +62,7 @@ protected:
 	 * 	Wait until the given gameplay tag query has become true or false, based on TriggerCondition, looking at the target actors ASC.
 	 *  If the the tag query already satisfies the TriggerCondition when this task is started, it will immediately broadcast the Triggered
 	 *  event. It will keep listening as long as bOnlyTriggerOnce = false.
+	 *  If used in an ability graph, this async action will wait even after activation ends. It's recommended to use WaitGameplayTagQuery instead.
 	 */
  	UFUNCTION(BlueprintCallable, Category = "Ability|Tasks", meta = (DefaultToSelf = "TargetActor", BlueprintInternalUseOnly = "TRUE"))
  	static UAbilityAsync_WaitGameplayTagQuery* WaitGameplayTagQueryOnActor(AActor* TargetActor, 
