Correctly set activation info on predicted instanced abilities that were rejected by server

diff --git a/Engine/Plugins/Runtime/GameplayAbilities/Source/GameplayAbilities/Private/AbilitySystemComponent_Abilities.cpp b/Engine/Plugins/Runtime/GameplayAbilities/Source/GameplayAbilities/Private/AbilitySystemComponent_Abilities.cpp
--- a/Engine/Plugins/Runtime/GameplayAbilities/Source/GameplayAbilities/Private/AbilitySystemComponent_Abilities.cpp
+++ b/Engine/Plugins/Runtime/GameplayAbilities/Source/GameplayAbilities/Private/AbilitySystemComponent_Abilities.cpp
@@ -2021,6 +2021,7 @@ void UAbilitySystemComponent::ClientActivateAbilityFailed_Implementation(FGamepl
 	{
 		if (Ability->CurrentActivationInfo.GetActivationPredictionKey().Current == PredictionKey)
 		{
+			Ability->CurrentActivationInfo.SetActivationRejected();
 			Ability->K2_EndAbility();
 		}
 	}
