Fix garbage ref by swapping FGameplayAbilityLocalAnimMontage::AnimatingAbility to a weak ptr.

diff --git a/Engine/Plugins/Runtime/GameplayAbilities/Source/GameplayAbilities/Private/AbilitySystemComponent_Abilities.cpp b/Engine/Plugins/Runtime/GameplayAbilities/Source/GameplayAbilities/Private/AbilitySystemComponent_Abilities.cpp
--- a/Engine/Plugins/Runtime/GameplayAbilities/Source/GameplayAbilities/Private/AbilitySystemComponent_Abilities.cpp
+++ b/Engine/Plugins/Runtime/GameplayAbilities/Source/GameplayAbilities/Private/AbilitySystemComponent_Abilities.cpp
@@ -1002,7 +1002,7 @@ void UAbilitySystemComponent::NotifyAbilityEnded(FGameplayAbilitySpecHandle Hand
 	ENetRole OwnerRole = GetOwnerRole();
 
 	// If AnimatingAbility ended, clear the pointer
-	if (LocalAnimMontageInfo.AnimatingAbility == Ability)
+	if (LocalAnimMontageInfo.AnimatingAbility.Get() == Ability)
 	{
 		ClearAnimatingAbility(Ability);
 	}
@@ -2699,12 +2699,15 @@ float UAbilitySystemComponent::PlayMontage(UGameplayAbility* InAnimatingAbility,
 		Duration = AnimInstance->Montage_Play(NewAnimMontage, InPlayRate, EMontagePlayReturnType::MontageLength, StartTimeSeconds);
 		if (Duration > 0.f)
 		{
-			if (LocalAnimMontageInfo.AnimatingAbility && LocalAnimMontageInfo.AnimatingAbility != InAnimatingAbility)
+			if (const UGameplayAbility* RawAnimatingAbility = LocalAnimMontageInfo.AnimatingAbility.Get())
 			{
-				// The ability that was previously animating will have already gotten the 'interrupted' callback.
-				// It may be a good idea to make this a global policy and 'cancel' the ability.
-				// 
-				// For now, we expect it to end itself when this happens.
+				if (RawAnimatingAbility != InAnimatingAbility)
+				{
+					// The ability that was previously animating will have already gotten the 'interrupted' callback.
+					// It may be a good idea to make this a global policy and 'cancel' the ability.
+					// 
+					// For now, we expect it to end itself when this happens.
+				}
 			}
 
 			if (NewAnimMontage->HasRootMotion() && AnimInstance->GetOwningActor())
@@ -3049,7 +3052,7 @@ void UAbilitySystemComponent::StopMontageIfCurrent(const UAnimMontage& Montage,
 
 void UAbilitySystemComponent::ClearAnimatingAbility(UGameplayAbility* Ability)
 {
-	if (LocalAnimMontageInfo.AnimatingAbility == Ability)
+	if (LocalAnimMontageInfo.AnimatingAbility.Get() == Ability)
 	{
 		Ability->SetCurrentMontage(nullptr);
 		LocalAnimMontageInfo.AnimatingAbility = nullptr;
@@ -3316,12 +3319,12 @@ void UAbilitySystemComponent::SetMontageRepAnimPositionMethod(ERepAnimPositionMe
 
 bool UAbilitySystemComponent::IsAnimatingAbility(UGameplayAbility* InAbility) const
 {
-	return (LocalAnimMontageInfo.AnimatingAbility == InAbility);
+	return (LocalAnimMontageInfo.AnimatingAbility.Get() == InAbility);
 }
 
 UGameplayAbility* UAbilitySystemComponent::GetAnimatingAbility()
 {
-	return LocalAnimMontageInfo.AnimatingAbility;
+	return LocalAnimMontageInfo.AnimatingAbility.Get();
 }
 
 void UAbilitySystemComponent::ConfirmAbilityTargetData(FGameplayAbilitySpecHandle AbilityHandle, FPredictionKey AbilityOriginalPredictionKey, const FGameplayAbilityTargetDataHandle& TargetData, const FGameplayTag& ApplicationTag)
diff --git a/Engine/Plugins/Runtime/GameplayAbilities/Source/GameplayAbilities/Public/Abilities/GameplayAbilityTypes.h b/Engine/Plugins/Runtime/GameplayAbilities/Source/GameplayAbilities/Public/Abilities/GameplayAbilityTypes.h
--- a/Engine/Plugins/Runtime/GameplayAbilities/Source/GameplayAbilities/Public/Abilities/GameplayAbilityTypes.h
+++ b/Engine/Plugins/Runtime/GameplayAbilities/Source/GameplayAbilities/Public/Abilities/GameplayAbilityTypes.h
@@ -226,7 +226,7 @@ struct GAMEPLAYABILITIES_API FGameplayAbilityLocalAnimMontage
 
 	/** The ability, if any, that instigated this montage */
 	UPROPERTY()
-	TObjectPtr<UGameplayAbility> AnimatingAbility;
+	TWeakObjectPtr<UGameplayAbility> AnimatingAbility;
 };
 
 
