PR #7491: GAS: Allow changing the evaluation channel of a GE mod from cpp (Contributed by [REDACTED_PERSON])

diff --git a/Engine/Plugins/Runtime/GameplayAbilities/Source/GameplayAbilities/Private/GameplayEffectTypes.cpp b/Engine/Plugins/Runtime/GameplayAbilities/Source/GameplayAbilities/Private/GameplayEffectTypes.cpp
--- a/Engine/Plugins/Runtime/GameplayAbilities/Source/GameplayAbilities/Private/GameplayEffectTypes.cpp
+++ b/Engine/Plugins/Runtime/GameplayAbilities/Source/GameplayAbilities/Private/GameplayEffectTypes.cpp
@@ -78,6 +78,14 @@ EGameplayModEvaluationChannel FGameplayModEvaluationChannelSettings::GetEvaluati
 	return EGameplayModEvaluationChannel::Channel0;
 }
 
+void FGameplayModEvaluationChannelSettings::SetEvaluationChannel(EGameplayModEvaluationChannel NewChannel)
+{
+	if (ensure(UAbilitySystemGlobals::Get().IsGameplayModEvaluationChannelValid(NewChannel)))
+	{
+		Channel = NewChannel;
+	}
+}
+
 float GameplayEffectUtilities::GetModifierBiasByModifierOp(EGameplayModOp::Type ModOp)
 {
 	static const float ModifierOpBiases[EGameplayModOp::Max] = {0.f, 1.f, 1.f, 0.f};
diff --git a/Engine/Plugins/Runtime/GameplayAbilities/Source/GameplayAbilities/Public/GameplayEffectTypes.h b/Engine/Plugins/Runtime/GameplayAbilities/Source/GameplayAbilities/Public/GameplayEffectTypes.h
--- a/Engine/Plugins/Runtime/GameplayAbilities/Source/GameplayAbilities/Public/GameplayEffectTypes.h
+++ b/Engine/Plugins/Runtime/GameplayAbilities/Source/GameplayAbilities/Public/GameplayEffectTypes.h
@@ -85,6 +85,8 @@ struct GAMEPLAYABILITIES_API FGameplayModEvaluationChannelSettings
 	static const FString ForceHideMetadataEnabledValue;
 #endif // #if WITH_EDITORONLY_DATA
 
+	void SetEvaluationChannel(EGameplayModEvaluationChannel NewChannel);
+
 protected:
 
 	/** Channel the settings would prefer to use, if possible/valid */
