Removing previously unused variable FGameplayModifierInfo::Magnitude (Gameplay Abilities Plugin).

diff --git a/Engine/Plugins/Runtime/GameplayAbilities/Source/GameplayAbilities/Private/AbilitySystemComponent.cpp b/Engine/Plugins/Runtime/GameplayAbilities/Source/GameplayAbilities/Private/AbilitySystemComponent.cpp
--- a/Engine/Plugins/Runtime/GameplayAbilities/Source/GameplayAbilities/Private/AbilitySystemComponent.cpp
+++ b/Engine/Plugins/Runtime/GameplayAbilities/Source/GameplayAbilities/Private/AbilitySystemComponent.cpp
@@ -1014,7 +1014,7 @@ void UAbilitySystemComponent::ExecuteGameplayEffect(FGameplayEffectSpec &Spec, F
 	{
 		ABILITY_VLOG(GetOwnerActor(), Log, TEXT("Executed %s"), *Spec.Def->GetFName().ToString());
 		
-		for (FGameplayModifierInfo Modifier : Spec.Def->Modifiers)
+		for (const FGameplayModifierInfo& Modifier : Spec.Def->Modifiers)
 		{
 			float Magnitude = 0.f;
 			Modifier.ModifierMagnitude.AttemptCalculateMagnitude(Spec, Magnitude);
diff --git a/Engine/Plugins/Runtime/GameplayAbilities/Source/GameplayAbilities/Private/GameplayEffect.cpp b/Engine/Plugins/Runtime/GameplayAbilities/Source/GameplayAbilities/Private/GameplayEffect.cpp
--- a/Engine/Plugins/Runtime/GameplayAbilities/Source/GameplayAbilities/Private/GameplayEffect.cpp
+++ b/Engine/Plugins/Runtime/GameplayAbilities/Source/GameplayAbilities/Private/GameplayEffect.cpp
@@ -84,14 +84,6 @@ void UGameplayEffect::PostLoad()
 	// Temporary post-load fix-up to preserve magnitude data
 	for (FGameplayModifierInfo& CurModInfo : Modifiers)
 	{
-		// If the old magnitude actually had some value in it, copy it over and then clear out the old data
-		static const FString GameplayEffectPostLoadContext(TEXT("UGameplayEffect::PostLoad"));
-		if (CurModInfo.Magnitude.Value != 0.f || CurModInfo.Magnitude.Curve.IsValid(GameplayEffectPostLoadContext))
-		{
-			CurModInfo.ModifierMagnitude.ScalableFloatMagnitude = CurModInfo.Magnitude;
-			CurModInfo.Magnitude = FScalableFloat();
-		}
-
 #if WITH_EDITOR
 		CurModInfo.ModifierMagnitude.ReportErrors(GetPathName());
 #endif // WITH_EDITOR
@@ -3472,7 +3464,7 @@ bool FActiveGameplayEffectsContainer::RemoveActiveGameplayEffect(FActiveGameplay
 			if (UE_LOG_ACTIVE(VLogAbilitySystem, Log))
 			{
 				ABILITY_VLOG(OwnerActor, Log, TEXT("Removed %s"), *Effect.Spec.Def->GetFName().ToString());
-				for (FGameplayModifierInfo Modifier : Effect.Spec.Def->Modifiers)
+				for (const FGameplayModifierInfo& Modifier : Effect.Spec.Def->Modifiers)
 				{
 					float Magnitude = 0.f;
 					Modifier.ModifierMagnitude.AttemptCalculateMagnitude(Effect.Spec, Magnitude);
@@ -3499,7 +3491,7 @@ void FActiveGameplayEffectsContainer::InternalExecutePeriodicGameplayEffect(FAct
 		{
 			AActor* OwnerActor = Owner->GetOwnerActor();
 			ABILITY_VLOG(OwnerActor, Log, TEXT("Executed Periodic Effect %s"), *ActiveEffect.Spec.Def->GetFName().ToString());
-			for (FGameplayModifierInfo Modifier : ActiveEffect.Spec.Def->Modifiers)
+			for (const FGameplayModifierInfo& Modifier : ActiveEffect.Spec.Def->Modifiers)
 			{
 				float Magnitude = 0.f;
 				Modifier.ModifierMagnitude.AttemptCalculateMagnitude(ActiveEffect.Spec, Magnitude);
diff --git a/Engine/Plugins/Runtime/GameplayAbilities/Source/GameplayAbilities/Public/GameplayEffect.h b/Engine/Plugins/Runtime/GameplayAbilities/Source/GameplayAbilities/Public/GameplayEffect.h
--- a/Engine/Plugins/Runtime/GameplayAbilities/Source/GameplayAbilities/Public/GameplayEffect.h
+++ b/Engine/Plugins/Runtime/GameplayAbilities/Source/GameplayAbilities/Public/GameplayEffect.h
@@ -494,24 +494,14 @@ USTRUCT(BlueprintType)
 struct GAMEPLAYABILITIES_API FGameplayModifierInfo
 {
 	GENERATED_USTRUCT_BODY()
-
-	FGameplayModifierInfo()	
-	: ModifierOp(EGameplayModOp::Additive)
-	{
-
-	}
-
+	
 	/** The Attribute we modify or the GE we modify modifies. */
 	UPROPERTY(EditDefaultsOnly, Category=GameplayModifier, meta=(FilterMetaTag="HideFromModifiers"))
 	FGameplayAttribute Attribute;
 
 	/** The numeric operation of this modifier: Override, Add, Multiply, etc  */
 	UPROPERTY(EditDefaultsOnly, Category=GameplayModifier)
-	TEnumAsByte<EGameplayModOp::Type> ModifierOp;
-
-	/** Now "deprecated," though being handled in a custom manner to avoid engine version bump. */
-	UPROPERTY()
-	FScalableFloat Magnitude;
+	TEnumAsByte<EGameplayModOp::Type> ModifierOp = EGameplayModOp::Additive;
 
 	/** Magnitude of the modifier */
 	UPROPERTY(EditDefaultsOnly, Category=GameplayModifier)
