Iris - MinimalGameplayCueReplicationProxy support.

Implemented NetSerializer and ReplicationFragment to mimic MinimalGameplayCueReplicationProxy behavior.

Note: It's advised not to implement custom gameplay logic in a NetSerialize method. It's adviced to markup struct members with UPROPERTY when they're to be replicated, even if there's a custom NetSerialize method. It's adviced to markup struct members with UPROPERTY(NotReplicated) when a UPROPERTY shouldn't be replicated, even when there's a custom NetSerialize method.

diff --git a/Engine/Plugins/Runtime/GameplayAbilities/Source/GameplayAbilities/Private/Serialization/InternalMinimalGameplayCueReplicationProxyNetSerializer.cpp b/Engine/Plugins/Runtime/GameplayAbilities/Source/GameplayAbilities/Private/Serialization/InternalMinimalGameplayCueReplicationProxyNetSerializer.cpp
new file mode 100644
--- /dev/null
+++ b/Engine/Plugins/Runtime/GameplayAbilities/Source/GameplayAbilities/Private/Serialization/InternalMinimalGameplayCueReplicationProxyNetSerializer.cpp
@@ -0,0 +1,21 @@
+// Copyright Epic Games, Inc. All Rights Reserved.
+
+#include "InternalMinimalGameplayCueReplicationProxyNetSerializer.h"
+
+#if UE_WITH_IRIS
+
+#include "GameplayCueInterface.h"
+
+void FMinimalGameplayCueReplicationProxyForNetSerializer::CopyReplicatedFieldsFrom(const FMinimalGameplayCueReplicationProxy& ReplicationProxy)
+{
+	Tags = ReplicationProxy.ReplicatedTags;
+	Locations = ReplicationProxy.ReplicatedLocations;
+}
+
+void FMinimalGameplayCueReplicationProxyForNetSerializer::AssignReplicatedFieldsTo(FMinimalGameplayCueReplicationProxy& ReplicationProxy) const
+{
+	ReplicationProxy.ReplicatedTags = Tags;
+	ReplicationProxy.ReplicatedLocations = Locations;
+}
+
+#endif
diff --git a/Engine/Plugins/Runtime/GameplayAbilities/Source/GameplayAbilities/Private/Serialization/InternalMinimalGameplayCueReplicationProxyNetSerializer.h b/Engine/Plugins/Runtime/GameplayAbilities/Source/GameplayAbilities/Private/Serialization/InternalMinimalGameplayCueReplicationProxyNetSerializer.h
new file mode 100644
--- /dev/null
+++ b/Engine/Plugins/Runtime/GameplayAbilities/Source/GameplayAbilities/Private/Serialization/InternalMinimalGameplayCueReplicationProxyNetSerializer.h
@@ -0,0 +1,28 @@
+// Copyright Epic Games, Inc. All Rights Reserved.
+
+#pragma once
+
+#include "UObject/ObjectMacros.h"
+#include "Containers/Array.h"
+#include "Engine/NetSerialization.h"
+#include "GameplayTagContainer.h"
+#include "InternalMinimalGameplayCueReplicationProxyNetSerializer.generated.h"
+
+class UAbilitySystemComponent;
+struct FMinimalGameplayCueReplicationProxy;
+
+USTRUCT()
+struct FMinimalGameplayCueReplicationProxyForNetSerializer
+{
+	GENERATED_BODY()
+
+public:
+	void CopyReplicatedFieldsFrom(const FMinimalGameplayCueReplicationProxy& ReplicationProxy);	
+	void AssignReplicatedFieldsTo(FMinimalGameplayCueReplicationProxy& ReplicationProxy) const;
+
+private:
+	UPROPERTY()
+	TArray<FGameplayTag> Tags;
+	UPROPERTY()
+	TArray<FVector_NetQuantize> Locations;
+};
diff --git a/Engine/Plugins/Runtime/GameplayAbilities/Source/GameplayAbilities/Private/Serialization/MinimalGameplayCueReplicationProxyNetSerializer.cpp b/Engine/Plugins/Runtime/GameplayAbilities/Source/GameplayAbilities/Private/Serialization/MinimalGameplayCueReplicationProxyNetSerializer.cpp
new file mode 100644
--- /dev/null
+++ b/Engine/Plugins/Runtime/GameplayAbilities/Source/GameplayAbilities/Private/Serialization/MinimalGameplayCueReplicationProxyNetSerializer.cpp
@@ -0,0 +1,252 @@
+// Copyright Epic Games, Inc. All Rights Reserved.
+
+#include "Serialization/MinimalGameplayCueReplicationProxyNetSerializer.h"
+
+#include UE_INLINE_GENERATED_CPP_BY_NAME(MinimalGameplayCueReplicationProxyNetSerializer)
+
+#if UE_WITH_IRIS
+
+#include "GameplayCueInterface.h"
+#include "Iris/Core/NetObjectReference.h"
+#include "Iris/Serialization/NetBitStreamReader.h"
+#include "Iris/Serialization/NetBitStreamWriter.h"
+#include "Iris/Serialization/NetSerializerDelegates.h"
+#include "Iris/Serialization/NetSerializers.h"
+#include "Iris/ReplicationState/PropertyNetSerializerInfoRegistry.h"
+#include "Iris/ReplicationState/ReplicationStateDescriptorBuilder.h"
+#include "Net/Core/Trace/NetTrace.h"
+#include "Serialization/InternalMinimalGameplayCueReplicationProxyNetSerializer.h"
+
+namespace UE::Net
+{
+
+extern FReplicationFragment* CreateAndRegisterMinimalGameplayCueReplicationProxyReplicationFragment(UObject* Owner, const FReplicationStateDescriptor* Descriptor, FFragmentRegistrationContext& Context);
+
+struct FMinimalGameplayCueReplicationProxyNetSerializer
+{
+	// Version
+	static const uint32 Version = 0;
+
+	// Traits
+	static constexpr bool bIsForwardingSerializer = true;
+	static constexpr bool bHasDynamicState = true;
+	// Few fields are replicated.
+	static constexpr bool bUseSerializerIsEqual = true;
+
+	struct FQuantizedType
+	{
+		alignas(16) uint8 Buffer[32];
+	};
+
+	typedef FMinimalGameplayCueReplicationProxy SourceType;
+	typedef FQuantizedType QuantizedType;
+	typedef FMinimalGameplayCueReplicationProxyNetSerializerConfig ConfigType;
+
+	inline static const ConfigType DefaultConfig;
+
+	static void Serialize(FNetSerializationContext&, const FNetSerializeArgs& Args);
+	static void Deserialize(FNetSerializationContext&, const FNetDeserializeArgs& Args);
+
+	static void SerializeDelta(FNetSerializationContext&, const FNetSerializeDeltaArgs& Args);
+	static void DeserializeDelta(FNetSerializationContext&, const FNetDeserializeDeltaArgs& Args);
+
+	static void Quantize(FNetSerializationContext&, const FNetQuantizeArgs& Args);
+	static void Dequantize(FNetSerializationContext&, const FNetDequantizeArgs& Args);
+
+	static bool IsEqual(FNetSerializationContext&, const FNetIsEqualArgs& Args);
+	static bool Validate(FNetSerializationContext&, const FNetValidateArgs& Args);
+
+	static void CloneDynamicState(FNetSerializationContext&, const FNetCloneDynamicStateArgs&);
+	static void FreeDynamicState(FNetSerializationContext&, const FNetFreeDynamicStateArgs&);
+
+	static void CollectNetReferences(FNetSerializationContext&, const FNetCollectReferencesArgs&);
+
+private:
+	class FNetSerializerRegistryDelegates final : private UE::Net::FNetSerializerRegistryDelegates
+	{
+	public:
+		virtual ~FNetSerializerRegistryDelegates();
+
+	private:
+		virtual void OnPreFreezeNetSerializerRegistry() override;
+		virtual void OnPostFreezeNetSerializerRegistry() override;
+	};
+
+	inline static FNetSerializerRegistryDelegates NetSerializerRegistryDelegates;
+	inline static FStructNetSerializerConfig StructNetSerializerConfig;
+	static const FNetSerializer* StructNetSerializer;
+};
+
+UE_NET_IMPLEMENT_SERIALIZER(FMinimalGameplayCueReplicationProxyNetSerializer);
+
+const FNetSerializer* FMinimalGameplayCueReplicationProxyNetSerializer::StructNetSerializer = &UE_NET_GET_SERIALIZER(FStructNetSerializer);
+
+void FMinimalGameplayCueReplicationProxyNetSerializer::Serialize(FNetSerializationContext& Context, const FNetSerializeArgs& Args)
+{
+	FNetSerializeArgs NetSerializeArgs = Args;
+	NetSerializeArgs.NetSerializerConfig = &StructNetSerializerConfig;
+	StructNetSerializer->Serialize(Context, NetSerializeArgs);
+}
+
+void FMinimalGameplayCueReplicationProxyNetSerializer::Deserialize(FNetSerializationContext& Context, const FNetDeserializeArgs& Args)
+{
+	FNetDeserializeArgs NetDeserializeArgs;
+	NetDeserializeArgs.NetSerializerConfig = &StructNetSerializerConfig;
+	StructNetSerializer->Deserialize(Context, NetDeserializeArgs);
+}
+
+void FMinimalGameplayCueReplicationProxyNetSerializer::SerializeDelta(FNetSerializationContext& Context, const FNetSerializeDeltaArgs& Args)
+{
+	FNetSerializeDeltaArgs NetSerializeDeltaArgs = Args;
+	NetSerializeDeltaArgs.NetSerializerConfig = &StructNetSerializerConfig;
+	StructNetSerializer->SerializeDelta(Context, NetSerializeDeltaArgs);
+}
+
+void FMinimalGameplayCueReplicationProxyNetSerializer::DeserializeDelta(FNetSerializationContext& Context, const FNetDeserializeDeltaArgs& Args)
+{
+	FNetDeserializeDeltaArgs NetDeserializeDeltaArgs = Args;
+	NetDeserializeDeltaArgs.NetSerializerConfig = &StructNetSerializerConfig;
+	StructNetSerializer->DeserializeDelta(Context, NetDeserializeDeltaArgs);
+}
+
+void FMinimalGameplayCueReplicationProxyNetSerializer::Quantize(FNetSerializationContext& Context, const FNetQuantizeArgs& Args)
+{
+	const SourceType& SourceValue = *reinterpret_cast<const SourceType*>(Args.Source);
+	
+	FMinimalGameplayCueReplicationProxyForNetSerializer TempValue;
+	TempValue.CopyReplicatedFieldsFrom(SourceValue);
+
+	FNetQuantizeArgs QuantizeArgs = Args;
+	QuantizeArgs.NetSerializerConfig = &StructNetSerializerConfig;
+	QuantizeArgs.Source = NetSerializerValuePointer(&TempValue);
+	StructNetSerializer->Quantize(Context, QuantizeArgs);
+}
+
+void FMinimalGameplayCueReplicationProxyNetSerializer::Dequantize(FNetSerializationContext& Context, const FNetDequantizeArgs& Args)
+{
+	const QuantizedType& SourceValue = *reinterpret_cast<const QuantizedType*>(Args.Source);
+	SourceType& TargetValue = *reinterpret_cast<SourceType*>(Args.Target);
+	
+	FMinimalGameplayCueReplicationProxyForNetSerializer TempValue;
+
+	FNetDequantizeArgs DequantizeArgs = Args;
+	DequantizeArgs.NetSerializerConfig = &StructNetSerializerConfig;
+	DequantizeArgs.Target = NetSerializerValuePointer(&TempValue);
+	StructNetSerializer->Dequantize(Context, DequantizeArgs);
+
+	TempValue.AssignReplicatedFieldsTo(TargetValue);
+}
+
+bool FMinimalGameplayCueReplicationProxyNetSerializer::IsEqual(FNetSerializationContext& Context, const FNetIsEqualArgs& Args)
+{
+	if (Args.bStateIsQuantized)
+	{
+		FNetIsEqualArgs IsEqualArgs = Args;
+		IsEqualArgs.NetSerializerConfig = &StructNetSerializerConfig;
+		if (!StructNetSerializer->IsEqual(Context, IsEqualArgs))
+		{
+			return false;
+		}
+	}
+	else
+	{
+		const SourceType& Source0 = *reinterpret_cast<const SourceType*>(Args.Source0);
+		const SourceType& Source1 = *reinterpret_cast<const SourceType*>(Args.Source1);
+
+		FMinimalGameplayCueReplicationProxyForNetSerializer TempValue0;
+		FMinimalGameplayCueReplicationProxyForNetSerializer TempValue1;
+
+		TempValue0.CopyReplicatedFieldsFrom(Source0);
+		TempValue1.CopyReplicatedFieldsFrom(Source1);
+
+		FNetIsEqualArgs IsEqualArgs = Args;
+		IsEqualArgs.Source0 = NetSerializerValuePointer(&TempValue0);
+		IsEqualArgs.Source1 = NetSerializerValuePointer(&TempValue1);
+		IsEqualArgs.NetSerializerConfig = &StructNetSerializerConfig;
+		if (!StructNetSerializer->IsEqual(Context, IsEqualArgs))
+		{
+			return false;
+		}
+	}
+
+	return true;
+}
+
+bool FMinimalGameplayCueReplicationProxyNetSerializer::Validate(FNetSerializationContext& Context, const FNetValidateArgs& Args)
+{
+	const SourceType& SourceValue = *reinterpret_cast<const SourceType*>(Args.Source);
+	
+	FMinimalGameplayCueReplicationProxyForNetSerializer TempValue;
+	TempValue.CopyReplicatedFieldsFrom(SourceValue);
+
+	FNetValidateArgs ValidateArgs = Args;
+	ValidateArgs.NetSerializerConfig = &StructNetSerializerConfig;
+	ValidateArgs.Source = NetSerializerValuePointer(&TempValue);
+	return StructNetSerializer->Validate(Context, ValidateArgs);
+}
+
+void FMinimalGameplayCueReplicationProxyNetSerializer::CollectNetReferences(FNetSerializationContext& Context, const FNetCollectReferencesArgs& Args)
+{
+	FNetCollectReferencesArgs CollectReferencesArgs = Args;
+	CollectReferencesArgs.NetSerializerConfig = &StructNetSerializerConfig;
+	StructNetSerializer->CollectNetReferences(Context, CollectReferencesArgs);
+}
+
+void FMinimalGameplayCueReplicationProxyNetSerializer::CloneDynamicState(FNetSerializationContext& Context, const FNetCloneDynamicStateArgs& Args)
+{
+	FNetCloneDynamicStateArgs CloneDynamicStateArgs = Args;
+	CloneDynamicStateArgs.NetSerializerConfig = &StructNetSerializerConfig;
+	StructNetSerializer->CloneDynamicState(Context, CloneDynamicStateArgs);
+}
+
+void FMinimalGameplayCueReplicationProxyNetSerializer::FreeDynamicState(FNetSerializationContext& Context, const FNetFreeDynamicStateArgs& Args)
+{
+	FNetFreeDynamicStateArgs FreeDynamicStateArgs = Args;
+	FreeDynamicStateArgs.NetSerializerConfig = &StructNetSerializerConfig;
+	StructNetSerializer->FreeDynamicState(Context, FreeDynamicStateArgs);
+}
+
+static const FName PropertyNetSerializerRegistry_NAME_MinimalGameplayCueReplicationProxyNetSerializer("MinimalGameplayCueReplicationProxy");
+UE_NET_IMPLEMENT_NAMED_STRUCT_NETSERIALIZER_WITH_CUSTOM_FRAGMENT_INFO(PropertyNetSerializerRegistry_NAME_MinimalGameplayCueReplicationProxyNetSerializer, FMinimalGameplayCueReplicationProxyNetSerializer, CreateAndRegisterMinimalGameplayCueReplicationProxyReplicationFragment);
+
+FMinimalGameplayCueReplicationProxyNetSerializer::FNetSerializerRegistryDelegates::~FNetSerializerRegistryDelegates()
+{	
+	UE_NET_UNREGISTER_NETSERIALIZER_INFO(PropertyNetSerializerRegistry_NAME_MinimalGameplayCueReplicationProxyNetSerializer);
+}
+
+void FMinimalGameplayCueReplicationProxyNetSerializer::FNetSerializerRegistryDelegates::OnPreFreezeNetSerializerRegistry()
+{
+	UE_NET_REGISTER_NETSERIALIZER_INFO(PropertyNetSerializerRegistry_NAME_MinimalGameplayCueReplicationProxyNetSerializer);
+}
+
+void FMinimalGameplayCueReplicationProxyNetSerializer::FNetSerializerRegistryDelegates::OnPostFreezeNetSerializerRegistry()
+{
+	constexpr SIZE_T ExpectedSizeOfFMinimalGameplayCueReplicationProxy = 992;
+	constexpr SIZE_T ExpectedAlignOfFMinimalGameplayCueReplicationProxy = 16;
+
+	constexpr SIZE_T ActualSizeOfFMinimalGameplayCueReplicationProxy = sizeof(FMinimalGameplayCueReplicationProxy);
+	constexpr SIZE_T ActualAlignOfFMinimalGameplayCueReplicationProxy = alignof(FMinimalGameplayCueReplicationProxy);
+
+	// Do our best to detect changes to FMinimalGameplayCueReplicationProxy
+	// If this assert triggers, this NetSerializer implementation must be verified against FMinimalGameplayCueReplicationProxy::NetSerialize before updating the size and alignment.
+	static_assert(ExpectedSizeOfFMinimalGameplayCueReplicationProxy == ActualSizeOfFMinimalGameplayCueReplicationProxy && ActualAlignOfFMinimalGameplayCueReplicationProxy == ExpectedAlignOfFMinimalGameplayCueReplicationProxy, "FGameplayAbilityRepAnimMontage layout has changed. Might need to update FGameplayAbilityRepAnimMontageNetSerializer to include new data or update the size.");
+
+	const UStruct* ReplicationProxyStruct = FMinimalGameplayCueReplicationProxyForNetSerializer::StaticStruct();
+	FReplicationStateDescriptorBuilder::FParameters Params;
+	StructNetSerializerConfig.StateDescriptor = FReplicationStateDescriptorBuilder::CreateDescriptorForStruct(ReplicationProxyStruct, Params);
+
+	const FReplicationStateDescriptor* Descriptor = StructNetSerializerConfig.StateDescriptor.GetReference();
+
+	// Validate our assumptions regarding quantized state size and alignment.
+	if ((sizeof(QuantizedType) < Descriptor->InternalSize) || (alignof(QuantizedType) < Descriptor->InternalAlignment))
+	{
+		LowLevelFatalError(TEXT("FMinimalGameplayCueReplicationProxyNetSerializer::FQuantizedType has size %u and alignment %u but requires size %u and alignment %u."), uint32(sizeof(FQuantizedType)), uint32(alignof(FQuantizedType)), Descriptor->InternalSize, Descriptor->InternalAlignment);
+	}
+
+	ensureMsgf(EnumHasAnyFlags(Descriptor->Traits, EReplicationStateTraits::HasDynamicState), TEXT("FMinimalGameplayCueReplicationProxyNetSerializer seems to no longer have dynamic state."));
+	ensureMsgf(!EnumHasAnyFlags(Descriptor->Traits, EReplicationStateTraits::HasObjectReference), TEXT("FMinimalGameplayCueReplicationProxyNetSerializer seems to have object references."));
+}
+
+}
+
+#endif // UE_WITH_IRIS
diff --git a/Engine/Plugins/Runtime/GameplayAbilities/Source/GameplayAbilities/Private/Serialization/MinimalGameplayCueReplicationProxyReplicationFragment.cpp b/Engine/Plugins/Runtime/GameplayAbilities/Source/GameplayAbilities/Private/Serialization/MinimalGameplayCueReplicationProxyReplicationFragment.cpp
new file mode 100644
--- /dev/null
+++ b/Engine/Plugins/Runtime/GameplayAbilities/Source/GameplayAbilities/Private/Serialization/MinimalGameplayCueReplicationProxyReplicationFragment.cpp
@@ -0,0 +1,211 @@
+// Copyright Epic Games, Inc. All Rights Reserved.
+
+#include "Serialization/MinimalGameplayCueReplicationProxyReplicationFragment.h"
+
+#if UE_WITH_IRIS
+
+#include "Engine/NetConnection.h"
+#include "Iris/ReplicationState/ReplicationStateDescriptor.h"
+#include "Iris/ReplicationSystem/ReplicationOperations.h"
+#include "Iris/Serialization/NetSerializationContext.h"
+#include "Iris/Serialization/NetSerializer.h"
+#include "AbilitySystemComponent.h"
+#include "GameplayCueInterface.h"
+#include <type_traits>
+
+namespace UE::Net
+{
+
+FReplicationFragment* CreateAndRegisterMinimalGameplayCueReplicationProxyReplicationFragment(UObject* Owner, const FReplicationStateDescriptor* Descriptor, FFragmentRegistrationContext& Context)
+{
+	if (FMinimalGameplayCueReplicationProxyReplicationFragment* Fragment = new FMinimalGameplayCueReplicationProxyReplicationFragment(Context.GetFragmentTraits() | EReplicationFragmentTraits::DeleteWithInstanceProtocol, Owner, Descriptor))
+	{
+		Fragment->Register(Context);
+		return Fragment;
+	}
+
+	return nullptr;
+}
+
+FMinimalGameplayCueReplicationProxyReplicationFragment::FMinimalGameplayCueReplicationProxyReplicationFragment(EReplicationFragmentTraits InTraits, UObject* InOwner, const FReplicationStateDescriptor* InDescriptor)
+: FReplicationFragment(InTraits)
+, ReplicationStateDescriptor(InDescriptor)
+, Owner(InOwner)
+{
+	// We don't want to create temporary states when applying replicated data as we only want to update the replicated fields.
+	Traits |= EReplicationFragmentTraits::HasPersistentTargetStateBuffer;
+
+	if (EnumHasAnyFlags(InTraits, EReplicationFragmentTraits::CanReplicate))
+	{
+		Traits |= EReplicationFragmentTraits::NeedsPoll;
+
+		// We assume the struct only replicates a few members and the serializer knows which ones. We don't want to mark the struct as dirty if it isn't.
+		check(EnumHasAnyFlags(InDescriptor->Traits, EReplicationStateTraits::UseSerializerIsEqual));
+
+		SrcReplicationState = MakeUnique<FPropertyReplicationState>(InDescriptor);
+	}
+	
+	if (EnumHasAnyFlags(InTraits, EReplicationFragmentTraits::CanReceive))
+	{
+		if (EnumHasAnyFlags(InDescriptor->Traits, EReplicationStateTraits::HasRepNotifies))
+		{
+			ensureMsgf(!EnumHasAnyFlags(InDescriptor->Traits, EReplicationStateTraits::KeepPreviousState), TEXT("FMinimalGameplayCueReplicationProxyReplicationFragment doesn't support OnRep calls with previous value."));
+		}
+
+		// We do custom callbacks in CallRepNotifies rather than in ApplyReplicatedState.
+		Traits |= EReplicationFragmentTraits::HasRepNotifies | EReplicationFragmentTraits::NeedsLegacyCallbacks;
+	}
+}
+
+void FMinimalGameplayCueReplicationProxyReplicationFragment::Register(FFragmentRegistrationContext& Context)
+{
+	FPropertyReplicationState* ReplicationState = SrcReplicationState.Get();
+	Context.RegisterReplicationFragment(this, ReplicationStateDescriptor.GetReference(), (ReplicationState ? ReplicationState->GetStateBuffer() : nullptr));
+}
+
+void FMinimalGameplayCueReplicationProxyReplicationFragment::ApplyReplicatedState(FReplicationStateApplyContext& ApplyContext) const
+{
+	uint8* ExternalStatePointer = reinterpret_cast<uint8*>(Owner) + ReplicationStateDescriptor->MemberProperties[0]->GetOffset_ForGC();
+
+	// Copy replicated tags to local tags to detect what was replicated. This information is used in MimicMinimalGameplayCueReplicationProxyReceiveLogic.
+	FMinimalGameplayCueReplicationProxy* ExternalSourceState = reinterpret_cast<FMinimalGameplayCueReplicationProxy*>(ExternalStatePointer);
+	ExternalSourceState->LocalTags = MoveTemp(ExternalSourceState->ReplicatedTags);
+	ExternalSourceState->LocalBitMask.Init(true /* initial value */, ExternalSourceState->LocalTags.Num());
+
+	// Dequantize replicated members to external source state. Need to call the serializer directly with the appropriate arguments as the external state doesn't have a ReplicationStateHeader in front of it.
+	FNetDequantizeArgs DequantizeArgs = {};
+	DequantizeArgs.Source = NetSerializerValuePointer(ApplyContext.StateBufferData.RawStateBuffer);
+	DequantizeArgs.Target = NetSerializerValuePointer(ExternalSourceState);
+	DequantizeArgs.NetSerializerConfig = ReplicationStateDescriptor->MemberSerializerDescriptors[0].SerializerConfig;
+	const FNetSerializer* Serializer = ReplicationStateDescriptor->MemberSerializerDescriptors[0].Serializer;
+	Serializer->Dequantize(*ApplyContext.NetSerializationContext, DequantizeArgs);
+}
+
+bool FMinimalGameplayCueReplicationProxyReplicationFragment::PollReplicatedState(EReplicationFragmentPollFlags PollOption)
+{
+	if (EnumHasAnyFlags(PollOption, EReplicationFragmentPollFlags::PollAllState))
+	{
+		const uint8* ExternalStateBuffer = reinterpret_cast<uint8*>(Owner) + ReplicationStateDescriptor->MemberProperties[0]->GetOffset_ForGC();
+		const FMinimalGameplayCueReplicationProxy* ExternalSourceState = reinterpret_cast<const FMinimalGameplayCueReplicationProxy*>(ExternalStateBuffer);
+
+		const void* CachedStateBuffer = SrcReplicationState->GetStateBuffer();
+		const FMinimalGameplayCueReplicationProxy* CachedState = reinterpret_cast<const FMinimalGameplayCueReplicationProxy*>(static_cast<const uint8*>(CachedStateBuffer) + ReplicationStateDescriptor->MemberDescriptors[0].ExternalMemberOffset);
+
+		if (ExternalSourceState->LastSourceArrayReplicationKey != CachedState->LastSourceArrayReplicationKey)
+		{
+			return SrcReplicationState->PollPropertyReplicationState(Owner);
+		}
+	}
+
+	return SrcReplicationState->IsDirty(0);
+}
+
+void FMinimalGameplayCueReplicationProxyReplicationFragment::CallRepNotifies(FReplicationStateApplyContext& Context)
+{
+	MimicMinimalGameplayCueReplicationProxyReceiveLogic(Context);
+	CallRepNotify(Context);
+}
+
+void FMinimalGameplayCueReplicationProxyReplicationFragment::MimicMinimalGameplayCueReplicationProxyReceiveLogic(FReplicationStateApplyContext& Context)
+{
+	uint8* ExternalStatePointer = reinterpret_cast<uint8*>(Owner) + ReplicationStateDescriptor->MemberProperties[0]->GetOffset_ForGC();
+	FMinimalGameplayCueReplicationProxy* ExternalSourceState = reinterpret_cast<FMinimalGameplayCueReplicationProxy*>(ExternalStatePointer);
+
+	// Check whether we need to perform the complex logic or not.
+	UAbilitySystemComponent* StateOwner = ExternalSourceState->Owner;
+	bool UpdateOwnerTagMap = StateOwner != nullptr;
+	if (ExternalSourceState->bRequireNonOwningNetConnection && StateOwner)
+	{
+		if (AActor* OwningActor = StateOwner->GetOwner())
+		{			
+			if (const UNetConnection* OwnerNetConnection = OwningActor->GetNetConnection())
+			{
+				if (OwnerNetConnection->GetConnectionId() == Context.NetSerializationContext->GetLocalConnectionId())
+				{
+					UpdateOwnerTagMap = false;
+				}
+			}
+		}
+	}
+
+	if (!UpdateOwnerTagMap)
+	{
+		return;
+	}
+
+	FGameplayCueParameters Parameters;
+	ExternalSourceState->InitGameplayCueParametersFunc(Parameters, ExternalSourceState->Owner);
+	const FVector OriginalLocationParameter = Parameters.Location;
+
+	const int32 TagCount = ExternalSourceState->ReplicatedTags.Num();
+	for (int32 TagIt = 0; TagIt < TagCount; ++TagIt)
+	{
+		FGameplayTag& ReplicatedTag = ExternalSourceState->ReplicatedTags[TagIt];
+		FVector_NetQuantize& ReplicatedLocation = ExternalSourceState->ReplicatedLocations[TagIt];
+
+		const bool bHasReplicatedLocation = !ReplicatedLocation.IsZero();
+		if (bHasReplicatedLocation)
+		{
+			Parameters.Location = ReplicatedLocation;
+		}
+		else
+		{
+			Parameters.Location = OriginalLocationParameter;
+		}
+
+		const int32 LocalIdx = ExternalSourceState->LocalTags.IndexOfByKey(ReplicatedTag);
+		if (LocalIdx != INDEX_NONE)
+		{
+			// This tag already existed and is accounted for
+			ExternalSourceState->LocalBitMask[LocalIdx] = false;
+		}
+		else
+		{
+			// Not sure this needs to be set here as it's unconditionally set later. Keeping to mimic original behavior as much as possible.
+			ExternalSourceState->bCachedModifiedOwnerTags = true;
+
+			// This is a new tag, we need to invoke the WhileActive gameplay cue event
+			StateOwner->SetTagMapCount(ReplicatedTag, 1);
+			StateOwner->InvokeGameplayCueEvent(ReplicatedTag, EGameplayCueEvent::WhileActive, Parameters);
+		}
+	}
+
+	// Restore the location in case we touched it
+	Parameters.Location = OriginalLocationParameter;
+
+	ExternalSourceState->bCachedModifiedOwnerTags = true;
+	for (auto It = ExternalSourceState->LocalTags.CreateConstIterator(); It; ++It)
+	{
+		if (!ExternalSourceState->LocalBitMask[It.GetIndex()])
+		{
+			continue;
+		}
+
+		FGameplayTag RemovedTag = *It;
+		StateOwner->SetTagMapCount(RemovedTag, 0);
+		StateOwner->InvokeGameplayCueEvent(RemovedTag, EGameplayCueEvent::Removed, Parameters);
+	}
+}
+
+void FMinimalGameplayCueReplicationProxyReplicationFragment::CallRepNotify(FReplicationStateApplyContext& ApplyContext)
+{
+	if (const UFunction* RepNotifyFunction = ReplicationStateDescriptor->MemberPropertyDescriptors[0].RepNotifyFunction)
+	{
+		if (ApplyContext.bIsInit)
+		{
+			const uint8* ReceivedState = ApplyContext.StateBufferData.RawStateBuffer;
+			const uint8* DefaultState = ReplicationStateDescriptor->DefaultStateBuffer;
+
+			if (FReplicationStateOperations::IsEqualQuantizedState(*ApplyContext.NetSerializationContext, ReceivedState, DefaultState, ReplicationStateDescriptor))
+			{
+				return;
+			}
+		}
+
+		Owner->ProcessEvent(const_cast<UFunction*>(RepNotifyFunction), nullptr);
+	}
+}
+
+}
+
+#endif
diff --git a/Engine/Plugins/Runtime/GameplayAbilities/Source/GameplayAbilities/Public/GameplayCueInterface.h b/Engine/Plugins/Runtime/GameplayAbilities/Source/GameplayAbilities/Public/GameplayCueInterface.h
--- a/Engine/Plugins/Runtime/GameplayAbilities/Source/GameplayAbilities/Public/GameplayCueInterface.h
+++ b/Engine/Plugins/Runtime/GameplayAbilities/Source/GameplayAbilities/Public/GameplayCueInterface.h
@@ -12,6 +12,14 @@
 #include "GameplayPrediction.h"
 #include "GameplayCueInterface.generated.h"
 
+#if UE_WITH_IRIS
+struct FMinimalGameplayCueReplicationProxyForNetSerializer;
+namespace UE::Net
+{
+	class FMinimalGameplayCueReplicationProxyReplicationFragment;
+}
+#endif
+
 /** Interface for actors that wish to handle GameplayCue events from GameplayEffects. Native only because blueprints can't implement interfaces with native functions */
 UINTERFACE(MinimalAPI, meta = (CannotImplementInterfaceInBlueprint))
 class UGameplayCueInterface: public UInterface
@@ -220,6 +228,10 @@ struct GAMEPLAYABILITIES_API FMinimalGameplayCueReplicationProxy
 	bool operator!=(const FMinimalGameplayCueReplicationProxy& Other) const { return !(*this == Other); }
 
 private:
+#if UE_WITH_IRIS
+	friend FMinimalGameplayCueReplicationProxyForNetSerializer;
+	friend UE::Net::FMinimalGameplayCueReplicationProxyReplicationFragment;
+#endif
 
 	enum { NumInlineTags = 16 };
 
diff --git a/Engine/Plugins/Runtime/GameplayAbilities/Source/GameplayAbilities/Public/Serialization/MinimalGameplayCueReplicationProxyNetSerializer.h b/Engine/Plugins/Runtime/GameplayAbilities/Source/GameplayAbilities/Public/Serialization/MinimalGameplayCueReplicationProxyNetSerializer.h
new file mode 100644
--- /dev/null
+++ b/Engine/Plugins/Runtime/GameplayAbilities/Source/GameplayAbilities/Public/Serialization/MinimalGameplayCueReplicationProxyNetSerializer.h
@@ -0,0 +1,19 @@
+// Copyright Epic Games, Inc. All Rights Reserved.
+
+#pragma once
+
+#include "Iris/Serialization/NetSerializer.h"
+#include "MinimalGameplayCueReplicationProxyNetSerializer.generated.h"
+
+USTRUCT()
+struct FMinimalGameplayCueReplicationProxyNetSerializerConfig : public FNetSerializerConfig
+{
+	GENERATED_BODY()
+};
+
+namespace UE::Net
+{
+
+UE_NET_DECLARE_SERIALIZER(FMinimalGameplayCueReplicationProxyNetSerializer, GAMEPLAYABILITIES_API);
+
+}
diff --git a/Engine/Plugins/Runtime/GameplayAbilities/Source/GameplayAbilities/Public/Serialization/MinimalGameplayCueReplicationProxyReplicationFragment.h b/Engine/Plugins/Runtime/GameplayAbilities/Source/GameplayAbilities/Public/Serialization/MinimalGameplayCueReplicationProxyReplicationFragment.h
new file mode 100644
--- /dev/null
+++ b/Engine/Plugins/Runtime/GameplayAbilities/Source/GameplayAbilities/Public/Serialization/MinimalGameplayCueReplicationProxyReplicationFragment.h
@@ -0,0 +1,46 @@
+// Copyright Epic Games, Inc. All Rights Reserved.
+
+#pragma once
+
+#if UE_WITH_IRIS
+
+#include "Iris/ReplicationState/PropertyReplicationState.h"
+#include "Iris/ReplicationSystem/ReplicationFragment.h"
+#include "Templates/UniquePtr.h"
+
+namespace UE::Net
+{
+
+FReplicationFragment* CreateAndRegisterMinimalGameplayCueReplicationProxyReplicationFragment(UObject* Owner, const FReplicationStateDescriptor* Descriptor, FFragmentRegistrationContext& Context);
+
+/**
+ * ReplicationFragment implementation used for Iris backwards compatibility until there's a simple way to markup arbitrary members or declare replicated states in a different manner.
+ * That coupled with generic Iris ApplyState funcitonality and not serializing data into gameplay objects directly would make the need for the fragment and the custom serializer go away.
+ * Instead the custom logic would be handled in the ApplyState implementation.
+ * @note This is NOT what good looks like. Markup your data and avoid custom gameplay logic in serialization code.
+ */
+class FMinimalGameplayCueReplicationProxyReplicationFragment : public FReplicationFragment
+{
+public:
+	FMinimalGameplayCueReplicationProxyReplicationFragment(EReplicationFragmentTraits InTraits, UObject* InOwner, const FReplicationStateDescriptor* InDescriptor);
+
+	void Register(FFragmentRegistrationContext& Context);
+
+protected:
+	virtual void ApplyReplicatedState(FReplicationStateApplyContext& Context) const override;
+	virtual bool PollReplicatedState(EReplicationFragmentPollFlags PollOption) override;
+	virtual void CallRepNotifies(FReplicationStateApplyContext& Context) override;
+
+private:
+	void MimicMinimalGameplayCueReplicationProxyReceiveLogic(FReplicationStateApplyContext& Context);
+	void CallRepNotify(FReplicationStateApplyContext& Context);
+
+private:
+	TRefCountPtr<const FReplicationStateDescriptor> ReplicationStateDescriptor;
+	TUniquePtr<FPropertyReplicationState> SrcReplicationState;
+	UObject* Owner = nullptr;
+};
+
+}
+
+#endif
