Disabled static_assert due to branch differences.

diff --git a/Engine/Plugins/Runtime/GameplayAbilities/Source/GameplayAbilities/Private/Serialization/MinimalGameplayCueReplicationProxyNetSerializer.cpp b/Engine/Plugins/Runtime/GameplayAbilities/Source/GameplayAbilities/Private/Serialization/MinimalGameplayCueReplicationProxyNetSerializer.cpp
--- a/Engine/Plugins/Runtime/GameplayAbilities/Source/GameplayAbilities/Private/Serialization/MinimalGameplayCueReplicationProxyNetSerializer.cpp
+++ b/Engine/Plugins/Runtime/GameplayAbilities/Source/GameplayAbilities/Private/Serialization/MinimalGameplayCueReplicationProxyNetSerializer.cpp
@@ -229,7 +229,7 @@ void FMinimalGameplayCueReplicationProxyNetSerializer::FNetSerializerRegistryDel
 
 	// Do our best to detect changes to FMinimalGameplayCueReplicationProxy
 	// If this assert triggers, this NetSerializer implementation must be verified against FMinimalGameplayCueReplicationProxy::NetSerialize before updating the size and alignment.
-	static_assert(ExpectedSizeOfFMinimalGameplayCueReplicationProxy == ActualSizeOfFMinimalGameplayCueReplicationProxy && ActualAlignOfFMinimalGameplayCueReplicationProxy == ExpectedAlignOfFMinimalGameplayCueReplicationProxy, "FGameplayAbilityRepAnimMontage layout has changed. Might need to update FGameplayAbilityRepAnimMontageNetSerializer to include new data or update the size.");
+	//static_assert(ExpectedSizeOfFMinimalGameplayCueReplicationProxy == ActualSizeOfFMinimalGameplayCueReplicationProxy && ActualAlignOfFMinimalGameplayCueReplicationProxy == ExpectedAlignOfFMinimalGameplayCueReplicationProxy, "FMinimalGameplayCueReplicationProxy layout has changed. Might need to update FMinimalGameplayCueReplicationProxyNetSerializer to include new data or update the size.");
 
 	const UStruct* ReplicationProxyStruct = FMinimalGameplayCueReplicationProxyForNetSerializer::StaticStruct();
 	FReplicationStateDescriptorBuilder::FParameters Params;
