Iris - Fixed bug in FMinimalGameplayCueReplicationProxyNetSerializer Deserialize.

diff --git a/Engine/Plugins/Runtime/GameplayAbilities/Source/GameplayAbilities/Private/Serialization/MinimalGameplayCueReplicationProxyNetSerializer.cpp b/Engine/Plugins/Runtime/GameplayAbilities/Source/GameplayAbilities/Private/Serialization/MinimalGameplayCueReplicationProxyNetSerializer.cpp
--- a/Engine/Plugins/Runtime/GameplayAbilities/Source/GameplayAbilities/Private/Serialization/MinimalGameplayCueReplicationProxyNetSerializer.cpp
+++ b/Engine/Plugins/Runtime/GameplayAbilities/Source/GameplayAbilities/Private/Serialization/MinimalGameplayCueReplicationProxyNetSerializer.cpp
@@ -90,7 +90,7 @@ void FMinimalGameplayCueReplicationProxyNetSerializer::Serialize(FNetSerializati
 
 void FMinimalGameplayCueReplicationProxyNetSerializer::Deserialize(FNetSerializationContext& Context, const FNetDeserializeArgs& Args)
 {
-	FNetDeserializeArgs NetDeserializeArgs;
+	FNetDeserializeArgs NetDeserializeArgs = Args;
 	NetDeserializeArgs.NetSerializerConfig = &StructNetSerializerConfig;
 	StructNetSerializer->Deserialize(Context, NetDeserializeArgs);
 }
