Iris - MinimalReplicationTagCountMap support.

Implemented NetSerializer and ReplicationFragment to mimic MinimalReplicationTagCountMap NetSerialize behavior.

diff --git a/Engine/Plugins/Runtime/GameplayAbilities/Source/GameplayAbilities/Private/Serialization/InternalMinimalReplicationTagCountMapNetSerializer.cpp b/Engine/Plugins/Runtime/GameplayAbilities/Source/GameplayAbilities/Private/Serialization/InternalMinimalReplicationTagCountMapNetSerializer.cpp
new file mode 100644
--- /dev/null
+++ b/Engine/Plugins/Runtime/GameplayAbilities/Source/GameplayAbilities/Private/Serialization/InternalMinimalReplicationTagCountMapNetSerializer.cpp
@@ -0,0 +1,33 @@
+// Copyright Epic Games, Inc. All Rights Reserved.
+
+#include "InternalMinimalReplicationTagCountMapNetSerializer.h"
+
+#if UE_WITH_IRIS
+
+#include "GameplayEffectTypes.h"
+
+void FMinimalReplicationTagCountMapForNetSerializer::CopyReplicatedFieldsFrom(const FMinimalReplicationTagCountMap& TagCountMap)
+{
+	TagCountMap.TagMap.GetKeys(Tags);
+}
+
+void FMinimalReplicationTagCountMapForNetSerializer::AssignReplicatedFieldsTo(FMinimalReplicationTagCountMap& TagCountMap) const
+{
+	TMap<FGameplayTag, int32>& TagMap = TagCountMap.TagMap;
+	for (const FGameplayTag& Tag : Tags)
+	{
+		TagMap.FindOrAdd(Tag) = 1;
+	}
+}
+
+void FMinimalReplicationTagCountMapForNetSerializer::ClampTagCount(SIZE_T MaxTagCount)
+{
+	if (static_cast<uint32>(Tags.Num()) > MaxTagCount)
+	{
+		// This struct is only used on stack with a very limited lifetime. Let's not spend CPU time on re-allocating.
+		constexpr bool bShrinkMemory = false;
+		Tags.SetNum(static_cast<int32>(static_cast<SSIZE_T>(MaxTagCount)), bShrinkMemory);
+	}
+}
+
+#endif
diff --git a/Engine/Plugins/Runtime/GameplayAbilities/Source/GameplayAbilities/Private/Serialization/InternalMinimalReplicationTagCountMapNetSerializer.h b/Engine/Plugins/Runtime/GameplayAbilities/Source/GameplayAbilities/Private/Serialization/InternalMinimalReplicationTagCountMapNetSerializer.h
new file mode 100644
--- /dev/null
+++ b/Engine/Plugins/Runtime/GameplayAbilities/Source/GameplayAbilities/Private/Serialization/InternalMinimalReplicationTagCountMapNetSerializer.h
@@ -0,0 +1,32 @@
+// Copyright Epic Games, Inc. All Rights Reserved.
+
+#pragma once
+
+#include "UObject/ObjectMacros.h"
+#include "Containers/Array.h"
+#include "GameplayTagContainer.h"
+#include "InternalMinimalReplicationTagCountMapNetSerializer.generated.h"
+
+struct FMinimalReplicationTagCountMap;
+
+USTRUCT()
+struct FMinimalReplicationTagCountMapForNetSerializer
+{
+	GENERATED_BODY()
+
+public:
+	void CopyReplicatedFieldsFrom(const FMinimalReplicationTagCountMap& TagCountMap);	
+	void AssignReplicatedFieldsTo(FMinimalReplicationTagCountMap& TagCountMap) const;
+
+	SIZE_T GetTagCount() const;
+	void ClampTagCount(SIZE_T MaxTagCount);
+
+private:
+	UPROPERTY()
+	TArray<FGameplayTag> Tags;
+};
+
+inline SIZE_T FMinimalReplicationTagCountMapForNetSerializer::GetTagCount() const
+{
+	return static_cast<uint32>(Tags.Num());
+}
\ No newline at end of file
diff --git a/Engine/Plugins/Runtime/GameplayAbilities/Source/GameplayAbilities/Private/Serialization/MinimalReplicationTagCountMapReplicationFragment.cpp b/Engine/Plugins/Runtime/GameplayAbilities/Source/GameplayAbilities/Private/Serialization/MinimalReplicationTagCountMapReplicationFragment.cpp
new file mode 100644
--- /dev/null
+++ b/Engine/Plugins/Runtime/GameplayAbilities/Source/GameplayAbilities/Private/Serialization/MinimalReplicationTagCountMapReplicationFragment.cpp
@@ -0,0 +1,158 @@
+// Copyright Epic Games, Inc. All Rights Reserved.
+
+#include "Serialization/MinimalReplicationTagCountMapReplicationFragment.h"
+
+#if UE_WITH_IRIS
+
+#include "Engine/NetConnection.h"
+#include "Iris/ReplicationState/ReplicationStateDescriptor.h"
+#include "Iris/ReplicationSystem/ReplicationOperations.h"
+#include "Iris/Serialization/NetSerializationContext.h"
+#include "Iris/Serialization/NetSerializer.h"
+#include "AbilitySystemComponent.h"
+#include "GameplayEffectTypes.h"
+
+namespace UE::Net
+{
+
+FReplicationFragment* CreateAndRegisterMinimalReplicationTagCountMapReplicationFragment(UObject* Owner, const FReplicationStateDescriptor* Descriptor, FFragmentRegistrationContext& Context)
+{
+	if (FMinimalReplicationTagCountMapReplicationFragment* Fragment = new FMinimalReplicationTagCountMapReplicationFragment(Context.GetFragmentTraits() | EReplicationFragmentTraits::DeleteWithInstanceProtocol, Owner, Descriptor))
+	{
+		Fragment->Register(Context);
+		return Fragment;
+	}
+
+	return nullptr;
+}
+
+FMinimalReplicationTagCountMapReplicationFragment::FMinimalReplicationTagCountMapReplicationFragment(EReplicationFragmentTraits InTraits, UObject* InOwner, const FReplicationStateDescriptor* InDescriptor)
+: FReplicationFragment(InTraits)
+, ReplicationStateDescriptor(InDescriptor)
+, Owner(InOwner)
+{
+	// We don't want to create temporary states when applying replicated data as we only want to update the replicated fields.
+	Traits |= EReplicationFragmentTraits::HasPersistentTargetStateBuffer;
+
+	if (EnumHasAnyFlags(InTraits, EReplicationFragmentTraits::CanReplicate))
+	{
+		Traits |= EReplicationFragmentTraits::NeedsPoll;
+
+		// We assume the struct only replicates a few members and the serializer knows which ones. We don't want to mark the struct as dirty if it isn't.
+		check(EnumHasAnyFlags(InDescriptor->Traits, EReplicationStateTraits::UseSerializerIsEqual));
+
+		SrcReplicationState = MakeUnique<FPropertyReplicationState>(InDescriptor);
+	}
+	
+	if (EnumHasAnyFlags(InTraits, EReplicationFragmentTraits::CanReceive))
+	{
+		if (EnumHasAnyFlags(InDescriptor->Traits, EReplicationStateTraits::HasRepNotifies))
+		{
+			ensureMsgf(!EnumHasAnyFlags(InDescriptor->Traits, EReplicationStateTraits::KeepPreviousState), TEXT("FMinimalReplicationTagCountMapReplicationFragment doesn't support OnRep calls with previous value."));
+		}
+
+		// We do custom callbacks in CallRepNotifies rather than in ApplyReplicatedState.
+		Traits |= EReplicationFragmentTraits::HasRepNotifies | EReplicationFragmentTraits::NeedsLegacyCallbacks;
+	}
+}
+
+void FMinimalReplicationTagCountMapReplicationFragment::Register(FFragmentRegistrationContext& Context)
+{
+	FPropertyReplicationState* ReplicationState = SrcReplicationState.Get();
+	Context.RegisterReplicationFragment(this, ReplicationStateDescriptor.GetReference(), (ReplicationState ? ReplicationState->GetStateBuffer() : nullptr));
+}
+
+void FMinimalReplicationTagCountMapReplicationFragment::ApplyReplicatedState(FReplicationStateApplyContext& ApplyContext) const
+{
+	uint8* ExternalStatePointer = reinterpret_cast<uint8*>(Owner) + ReplicationStateDescriptor->MemberProperties[0]->GetOffset_ForGC();
+
+	// Clear tag counts in the map. This information is used in MimicMinimalReplicationTagCountMapReceiveLogic.
+	FMinimalReplicationTagCountMap* ExternalSourceState = reinterpret_cast<FMinimalReplicationTagCountMap*>(ExternalStatePointer);
+	for (auto& Pair : ExternalSourceState->TagMap)
+	{
+		Pair.Value = 0;
+	}
+
+	// Dequantize replicated members to external source state. Need to call the serializer directly with the appropriate arguments as the external state doesn't have a ReplicationStateHeader in front of it.
+	FNetDequantizeArgs DequantizeArgs = {};
+	DequantizeArgs.Source = NetSerializerValuePointer(ApplyContext.StateBufferData.RawStateBuffer);
+	DequantizeArgs.Target = NetSerializerValuePointer(ExternalSourceState);
+	DequantizeArgs.NetSerializerConfig = ReplicationStateDescriptor->MemberSerializerDescriptors[0].SerializerConfig;
+	const FNetSerializer* Serializer = ReplicationStateDescriptor->MemberSerializerDescriptors[0].Serializer;
+	Serializer->Dequantize(*ApplyContext.NetSerializationContext, DequantizeArgs);
+}
+
+bool FMinimalReplicationTagCountMapReplicationFragment::PollReplicatedState(EReplicationFragmentPollFlags PollOption)
+{
+	if (EnumHasAnyFlags(PollOption, EReplicationFragmentPollFlags::PollAllState))
+	{
+		const uint8* ExternalStateBuffer = reinterpret_cast<uint8*>(Owner) + ReplicationStateDescriptor->MemberProperties[0]->GetOffset_ForGC();
+		const FMinimalReplicationTagCountMap* ExternalSourceState = reinterpret_cast<const FMinimalReplicationTagCountMap*>(ExternalStateBuffer);
+
+		const void* CachedStateBuffer = SrcReplicationState->GetStateBuffer();
+		const FMinimalReplicationTagCountMap* CachedState = reinterpret_cast<const FMinimalReplicationTagCountMap*>(static_cast<const uint8*>(CachedStateBuffer) + ReplicationStateDescriptor->MemberDescriptors[0].ExternalMemberOffset);
+
+		return SrcReplicationState->PollPropertyReplicationState(Owner);
+	}
+
+	return SrcReplicationState->IsDirty(0);
+}
+
+void FMinimalReplicationTagCountMapReplicationFragment::CallRepNotifies(FReplicationStateApplyContext& Context)
+{
+	MimicMinimalReplicationTagCountMapReceiveLogic(Context);
+	CallRepNotify(Context);
+}
+
+void FMinimalReplicationTagCountMapReplicationFragment::MimicMinimalReplicationTagCountMapReceiveLogic(FReplicationStateApplyContext& Context)
+{
+	uint8* ExternalStatePointer = reinterpret_cast<uint8*>(Owner) + ReplicationStateDescriptor->MemberProperties[0]->GetOffset_ForGC();
+	FMinimalReplicationTagCountMap* ExternalSourceState = reinterpret_cast<FMinimalReplicationTagCountMap*>(ExternalStatePointer);
+
+	// Check whether we need to perform the complex logic or not.
+	UAbilitySystemComponent* StateOwner = ExternalSourceState->Owner;
+	bool UpdateOwnerTagMap = StateOwner != nullptr;
+	if (ExternalSourceState->bRequireNonOwningNetConnection && StateOwner)
+	{
+		if (AActor* OwningActor = StateOwner->GetOwner())
+		{			
+			if (const UNetConnection* OwnerNetConnection = OwningActor->GetNetConnection())
+			{
+				if (OwnerNetConnection->GetConnectionId() == Context.NetSerializationContext->GetLocalConnectionId())
+				{
+					UpdateOwnerTagMap = false;
+				}
+			}
+		}
+	}
+
+	if (!UpdateOwnerTagMap)
+	{
+		return;
+	}
+
+	ExternalSourceState->UpdateOwnerTagMap();
+}
+
+void FMinimalReplicationTagCountMapReplicationFragment::CallRepNotify(FReplicationStateApplyContext& ApplyContext)
+{
+	if (const UFunction* RepNotifyFunction = ReplicationStateDescriptor->MemberPropertyDescriptors[0].RepNotifyFunction)
+	{
+		if (ApplyContext.bIsInit)
+		{
+			const uint8* ReceivedState = ApplyContext.StateBufferData.RawStateBuffer;
+			const uint8* DefaultState = ReplicationStateDescriptor->DefaultStateBuffer;
+
+			if (FReplicationStateOperations::IsEqualQuantizedState(*ApplyContext.NetSerializationContext, ReceivedState, DefaultState, ReplicationStateDescriptor))
+			{
+				return;
+			}
+		}
+
+		Owner->ProcessEvent(const_cast<UFunction*>(RepNotifyFunction), nullptr);
+	}
+}
+
+}
+
+#endif
diff --git a/Engine/Plugins/Runtime/GameplayAbilities/Source/GameplayAbilities/Private/Serialization/MinimalReplicationTagCountMapReplicationFragment.h b/Engine/Plugins/Runtime/GameplayAbilities/Source/GameplayAbilities/Private/Serialization/MinimalReplicationTagCountMapReplicationFragment.h
new file mode 100644
--- /dev/null
+++ b/Engine/Plugins/Runtime/GameplayAbilities/Source/GameplayAbilities/Private/Serialization/MinimalReplicationTagCountMapReplicationFragment.h
@@ -0,0 +1,46 @@
+// Copyright Epic Games, Inc. All Rights Reserved.
+
+#pragma once
+
+#if UE_WITH_IRIS
+
+#include "Iris/ReplicationState/PropertyReplicationState.h"
+#include "Iris/ReplicationSystem/ReplicationFragment.h"
+#include "Templates/UniquePtr.h"
+
+namespace UE::Net
+{
+
+FReplicationFragment* CreateAndRegisterMinimalReplicationTagCountMapReplicationFragment(UObject* Owner, const FReplicationStateDescriptor* Descriptor, FFragmentRegistrationContext& Context);
+
+/**
+ * ReplicationFragment implementation used for Iris backwards compatibility until there's a simple way to markup arbitrary members or declare replicated states in a different manner.
+ * That coupled with generic Iris ApplyState functionality and not serializing data into gameplay objects directly would make the need for the fragment and the custom serializer go away.
+ * Instead the custom logic would be handled in the ApplyState implementation.
+ * @note This is NOT what good looks like. Markup your data and avoid custom gameplay logic in serialization code.
+ */
+class FMinimalReplicationTagCountMapReplicationFragment : public FReplicationFragment
+{
+public:
+	FMinimalReplicationTagCountMapReplicationFragment(EReplicationFragmentTraits InTraits, UObject* InOwner, const FReplicationStateDescriptor* InDescriptor);
+
+	void Register(FFragmentRegistrationContext& Context);
+
+protected:
+	virtual void ApplyReplicatedState(FReplicationStateApplyContext& Context) const override;
+	virtual bool PollReplicatedState(EReplicationFragmentPollFlags PollOption) override;
+	virtual void CallRepNotifies(FReplicationStateApplyContext& Context) override;
+
+private:
+	void MimicMinimalReplicationTagCountMapReceiveLogic(FReplicationStateApplyContext& Context);
+	void CallRepNotify(FReplicationStateApplyContext& Context);
+
+private:
+	TRefCountPtr<const FReplicationStateDescriptor> ReplicationStateDescriptor;
+	TUniquePtr<FPropertyReplicationState> SrcReplicationState;
+	UObject* Owner = nullptr;
+};
+
+}
+
+#endif
diff --git a/Engine/Plugins/Runtime/GameplayAbilities/Source/GameplayAbilities/Private/Serialization/MinimalReplicationTagCountNetSerializer.cpp b/Engine/Plugins/Runtime/GameplayAbilities/Source/GameplayAbilities/Private/Serialization/MinimalReplicationTagCountNetSerializer.cpp
--- a/Engine/Plugins/Runtime/GameplayAbilities/Source/GameplayAbilities/Private/Serialization/MinimalReplicationTagCountNetSerializer.cpp
+++ b/Engine/Plugins/Runtime/GameplayAbilities/Source/GameplayAbilities/Private/Serialization/MinimalReplicationTagCountNetSerializer.cpp
@@ -1,35 +1,279 @@
 // Copyright Epic Games, Inc. All Rights Reserved.
+
+#include "Serialization/MinimalReplicationTagCountNetSerializer.h"
+
+#include UE_INLINE_GENERATED_CPP_BY_NAME(MinimalReplicationTagCountNetSerializer)
+
 #if UE_WITH_IRIS
 
+#include "AbilitySystemGlobals.h"
 #include "GameplayEffectTypes.h"
-#include "Iris/ReplicationState/PropertyNetSerializerInfoRegistry.h"
 #include "Iris/Serialization/NetSerializerDelegates.h"
+#include "Iris/Serialization/NetSerializers.h"
+#include "Iris/ReplicationState/PropertyNetSerializerInfoRegistry.h"
+#include "Iris/ReplicationState/ReplicationStateDescriptorBuilder.h"
+#include "Serialization/InternalMinimalReplicationTagCountMapNetSerializer.h"
 
 namespace UE::Net
 {
 
-// Until we implement a proper NetSerializer JIRA:UE-161669 we fall back on the LastResortNetSerializer
-static const FName PropertyNetSerializerRegistry_NAME_MinimalReplicationTagCountMap("MinimalReplicationTagCountMap");
-UE_NET_IMPLEMENT_NAMED_STRUCT_LASTRESORT_NETSERIALIZER_INFO(PropertyNetSerializerRegistry_NAME_MinimalReplicationTagCountMap);
+extern FReplicationFragment* CreateAndRegisterMinimalReplicationTagCountMapReplicationFragment(UObject* Owner, const FReplicationStateDescriptor* Descriptor, FFragmentRegistrationContext& Context);
 
-struct FMinimalReplicationTagCountMapNetSerializerRegistryDelegates : protected FNetSerializerRegistryDelegates
+struct FMinimalReplicationTagCountMapNetSerializer
 {
-	~FMinimalReplicationTagCountMapNetSerializerRegistryDelegates();
-	virtual void OnPreFreezeNetSerializerRegistry();
+	// Version
+	static const uint32 Version = 0;
+
+	// Traits
+	static constexpr bool bIsForwardingSerializer = true;
+	static constexpr bool bHasDynamicState = true;
+	static constexpr bool bUseSerializerIsEqual = true;
+
+	struct FQuantizedType
+	{
+		alignas(16) uint8 Buffer[16];
+	};
+
+	typedef FMinimalReplicationTagCountMap SourceType;
+	typedef FQuantizedType QuantizedType;
+	typedef FMinimalReplicationTagCountMapNetSerializerConfig ConfigType;
+
+	inline static const ConfigType DefaultConfig;
+
+	static void Serialize(FNetSerializationContext&, const FNetSerializeArgs& Args);
+	static void Deserialize(FNetSerializationContext&, const FNetDeserializeArgs& Args);
+
+	static void SerializeDelta(FNetSerializationContext&, const FNetSerializeDeltaArgs& Args);
+	static void DeserializeDelta(FNetSerializationContext&, const FNetDeserializeDeltaArgs& Args);
+
+	static void Quantize(FNetSerializationContext&, const FNetQuantizeArgs& Args);
+	static void Dequantize(FNetSerializationContext&, const FNetDequantizeArgs& Args);
+
+	static bool IsEqual(FNetSerializationContext&, const FNetIsEqualArgs& Args);
+	static bool Validate(FNetSerializationContext&, const FNetValidateArgs& Args);
+
+	static void CloneDynamicState(FNetSerializationContext&, const FNetCloneDynamicStateArgs&);
+	static void FreeDynamicState(FNetSerializationContext&, const FNetFreeDynamicStateArgs&);
+
+	static void CollectNetReferences(FNetSerializationContext&, const FNetCollectReferencesArgs&);
+
+private:
+	static SIZE_T GetMaxReplicatedTagCount();
+
+	class FNetSerializerRegistryDelegates final : private UE::Net::FNetSerializerRegistryDelegates
+	{
+	public:
+		virtual ~FNetSerializerRegistryDelegates();
+
+	private:
+		virtual void OnPreFreezeNetSerializerRegistry() override;
+		virtual void OnPostFreezeNetSerializerRegistry() override;
+	};
+
+	inline static FNetSerializerRegistryDelegates NetSerializerRegistryDelegates;
+	inline static FStructNetSerializerConfig StructNetSerializerConfig;
+	static const FNetSerializer* StructNetSerializer;
 };
-static FMinimalReplicationTagCountMapNetSerializerRegistryDelegates MinimalReplicationTagCountMapNetSerializerRegistryDelegates;
 
-FMinimalReplicationTagCountMapNetSerializerRegistryDelegates::~FMinimalReplicationTagCountMapNetSerializerRegistryDelegates()
+UE_NET_IMPLEMENT_SERIALIZER(FMinimalReplicationTagCountMapNetSerializer);
+
+const FNetSerializer* FMinimalReplicationTagCountMapNetSerializer::StructNetSerializer = &UE_NET_GET_SERIALIZER(FStructNetSerializer);
+
+void FMinimalReplicationTagCountMapNetSerializer::Serialize(FNetSerializationContext& Context, const FNetSerializeArgs& Args)
+{
+	FNetSerializeArgs NetSerializeArgs = Args;
+	NetSerializeArgs.NetSerializerConfig = &StructNetSerializerConfig;
+	StructNetSerializer->Serialize(Context, NetSerializeArgs);
+}
+
+void FMinimalReplicationTagCountMapNetSerializer::Deserialize(FNetSerializationContext& Context, const FNetDeserializeArgs& Args)
+{
+	FNetDeserializeArgs NetDeserializeArgs = Args;
+	NetDeserializeArgs.NetSerializerConfig = &StructNetSerializerConfig;
+	StructNetSerializer->Deserialize(Context, NetDeserializeArgs);
+}
+
+void FMinimalReplicationTagCountMapNetSerializer::SerializeDelta(FNetSerializationContext& Context, const FNetSerializeDeltaArgs& Args)
+{
+	FNetSerializeDeltaArgs NetSerializeDeltaArgs = Args;
+	NetSerializeDeltaArgs.NetSerializerConfig = &StructNetSerializerConfig;
+	StructNetSerializer->SerializeDelta(Context, NetSerializeDeltaArgs);
+}
+
+void FMinimalReplicationTagCountMapNetSerializer::DeserializeDelta(FNetSerializationContext& Context, const FNetDeserializeDeltaArgs& Args)
+{
+	FNetDeserializeDeltaArgs NetDeserializeDeltaArgs = Args;
+	NetDeserializeDeltaArgs.NetSerializerConfig = &StructNetSerializerConfig;
+	StructNetSerializer->DeserializeDelta(Context, NetDeserializeDeltaArgs);
+}
+
+void FMinimalReplicationTagCountMapNetSerializer::Quantize(FNetSerializationContext& Context, const FNetQuantizeArgs& Args)
+{
+	const SourceType& SourceValue = *reinterpret_cast<const SourceType*>(Args.Source);
+	
+	FMinimalReplicationTagCountMapForNetSerializer TempValue;
+	TempValue.CopyReplicatedFieldsFrom(SourceValue);
+
+	// Cap the number of replicated tags.
+	const SIZE_T MaxTagCount = GetMaxReplicatedTagCount();
+	TempValue.ClampTagCount(MaxTagCount);
+
+	FNetQuantizeArgs QuantizeArgs = Args;
+	QuantizeArgs.NetSerializerConfig = &StructNetSerializerConfig;
+	QuantizeArgs.Source = NetSerializerValuePointer(&TempValue);
+	StructNetSerializer->Quantize(Context, QuantizeArgs);
+}
+
+void FMinimalReplicationTagCountMapNetSerializer::Dequantize(FNetSerializationContext& Context, const FNetDequantizeArgs& Args)
+{
+	const QuantizedType& SourceValue = *reinterpret_cast<const QuantizedType*>(Args.Source);
+	SourceType& TargetValue = *reinterpret_cast<SourceType*>(Args.Target);
+	
+	// The temporary value is required because of two reasons; the tag count needs to be verified and the real arget is a Map, which isn't supported for replication.
+	FMinimalReplicationTagCountMapForNetSerializer TempValue;
+
+	FNetDequantizeArgs DequantizeArgs = Args;
+	DequantizeArgs.NetSerializerConfig = &StructNetSerializerConfig;
+	DequantizeArgs.Target = NetSerializerValuePointer(&TempValue);
+	StructNetSerializer->Dequantize(Context, DequantizeArgs);
+
+	// Don't allow too large arrays.
+	if (TempValue.GetTagCount() > GetMaxReplicatedTagCount())
+	{
+		Context.SetError(GNetError_ArraySizeTooLarge);
+		return;
+	}
+
+	TempValue.AssignReplicatedFieldsTo(TargetValue);
+}
+
+bool FMinimalReplicationTagCountMapNetSerializer::IsEqual(FNetSerializationContext& Context, const FNetIsEqualArgs& Args)
 {
+	if (Args.bStateIsQuantized)
+	{
+		FNetIsEqualArgs IsEqualArgs = Args;
+		IsEqualArgs.NetSerializerConfig = &StructNetSerializerConfig;
+		if (!StructNetSerializer->IsEqual(Context, IsEqualArgs))
+		{
+			return false;
+		}
+	}
+	else
+	{
+		const SourceType& Source0 = *reinterpret_cast<const SourceType*>(Args.Source0);
+		const SourceType& Source1 = *reinterpret_cast<const SourceType*>(Args.Source1);
+
+		FMinimalReplicationTagCountMapForNetSerializer TempValue0;
+		FMinimalReplicationTagCountMapForNetSerializer TempValue1;
+
+		const SIZE_T MaxTagCount = GetMaxReplicatedTagCount();
+		TempValue0.CopyReplicatedFieldsFrom(Source0);
+		TempValue0.ClampTagCount(MaxTagCount);
+
+		TempValue1.CopyReplicatedFieldsFrom(Source1);
+		TempValue0.ClampTagCount(MaxTagCount);
+
+		FNetIsEqualArgs IsEqualArgs = Args;
+		IsEqualArgs.Source0 = NetSerializerValuePointer(&TempValue0);
+		IsEqualArgs.Source1 = NetSerializerValuePointer(&TempValue1);
+		IsEqualArgs.NetSerializerConfig = &StructNetSerializerConfig;
+		if (!StructNetSerializer->IsEqual(Context, IsEqualArgs))
+		{
+			return false;
+		}
+	}
+
+	return true;
+}
+
+bool FMinimalReplicationTagCountMapNetSerializer::Validate(FNetSerializationContext& Context, const FNetValidateArgs& Args)
+{
+	const SourceType& SourceValue = *reinterpret_cast<const SourceType*>(Args.Source);
+	
+	FMinimalReplicationTagCountMapForNetSerializer TempValue;
+	TempValue.CopyReplicatedFieldsFrom(SourceValue);
+
+	// MinimalReplicationTagCountMap logs an error when too many tags are detected. We return false during validation. 
+	// Similar to the original implementation we clamp the number of tags that will be replicated. This is done during quantization.
+	const SIZE_T MaxTagCount = GetMaxReplicatedTagCount();
+	if (TempValue.GetTagCount() > MaxTagCount)
+	{
+		return false;
+	}
+
+	FNetValidateArgs ValidateArgs = Args;
+	ValidateArgs.NetSerializerConfig = &StructNetSerializerConfig;
+	ValidateArgs.Source = NetSerializerValuePointer(&TempValue);
+	return StructNetSerializer->Validate(Context, ValidateArgs);
+}
+
+void FMinimalReplicationTagCountMapNetSerializer::CollectNetReferences(FNetSerializationContext& Context, const FNetCollectReferencesArgs& Args)
+{
+}
+
+void FMinimalReplicationTagCountMapNetSerializer::CloneDynamicState(FNetSerializationContext& Context, const FNetCloneDynamicStateArgs& Args)
+{
+	FNetCloneDynamicStateArgs CloneDynamicStateArgs = Args;
+	CloneDynamicStateArgs.NetSerializerConfig = &StructNetSerializerConfig;
+	StructNetSerializer->CloneDynamicState(Context, CloneDynamicStateArgs);
+}
+
+void FMinimalReplicationTagCountMapNetSerializer::FreeDynamicState(FNetSerializationContext& Context, const FNetFreeDynamicStateArgs& Args)
+{
+	FNetFreeDynamicStateArgs FreeDynamicStateArgs = Args;
+	FreeDynamicStateArgs.NetSerializerConfig = &StructNetSerializerConfig;
+	StructNetSerializer->FreeDynamicState(Context, FreeDynamicStateArgs);
+}
+
+SIZE_T FMinimalReplicationTagCountMapNetSerializer::GetMaxReplicatedTagCount()
+{
+	const int32 BitCount = UAbilitySystemGlobals::Get().MinimalReplicationTagCountBits;
+	const int32 ClampedBitCount = FMath::Clamp(BitCount, 1, 16);
+
+	const uint32 MaxCount = (1U << static_cast<uint32>(ClampedBitCount)) - 1U;
+	return MaxCount;
+}
+
+static const FName PropertyNetSerializerRegistry_NAME_MinimalReplicationTagCountMap("MinimalReplicationTagCountMap");
+UE_NET_IMPLEMENT_NAMED_STRUCT_NETSERIALIZER_WITH_CUSTOM_FRAGMENT_INFO(PropertyNetSerializerRegistry_NAME_MinimalReplicationTagCountMap, FMinimalReplicationTagCountMapNetSerializer, CreateAndRegisterMinimalReplicationTagCountMapReplicationFragment);
+
+FMinimalReplicationTagCountMapNetSerializer::FNetSerializerRegistryDelegates::~FNetSerializerRegistryDelegates()
+{	
 	UE_NET_UNREGISTER_NETSERIALIZER_INFO(PropertyNetSerializerRegistry_NAME_MinimalReplicationTagCountMap);
 }
 
-void FMinimalReplicationTagCountMapNetSerializerRegistryDelegates::OnPreFreezeNetSerializerRegistry()
+void FMinimalReplicationTagCountMapNetSerializer::FNetSerializerRegistryDelegates::OnPreFreezeNetSerializerRegistry()
 {
 	UE_NET_REGISTER_NETSERIALIZER_INFO(PropertyNetSerializerRegistry_NAME_MinimalReplicationTagCountMap);
 }
 
+void FMinimalReplicationTagCountMapNetSerializer::FNetSerializerRegistryDelegates::OnPostFreezeNetSerializerRegistry()
+{
+	constexpr SIZE_T ExpectedSizeOfFMinimalReplicationTagCountMap = 96;
+	constexpr SIZE_T ExpectedAlignOfFMinimalReplicationTagCountMap = 8;
+
+	constexpr SIZE_T ActualSizeOfFMinimalReplicationTagCountMap = sizeof(FMinimalReplicationTagCountMap);
+	constexpr SIZE_T ActualAlignOfFMinimalReplicationTagCountMap = alignof(FMinimalReplicationTagCountMap);
+
+	// Try to detect changes to FMinimalReplicationTagCountMap. Ensure instead of check as it's not known to be an error until verified.
+	ensureMsgf(ExpectedSizeOfFMinimalReplicationTagCountMap == ActualSizeOfFMinimalReplicationTagCountMap && ExpectedAlignOfFMinimalReplicationTagCountMap == ActualAlignOfFMinimalReplicationTagCountMap, TEXT("%s"), TEXT("FMinimalReplicationTagCountMap layout has changed. Might need to update FMinimalReplicationTagCountMapNetSerializer to include new data."));
+
+	const UStruct* TagCountMapStruct = FMinimalReplicationTagCountMapForNetSerializer::StaticStruct();
+	FReplicationStateDescriptorBuilder::FParameters Params;
+	StructNetSerializerConfig.StateDescriptor = FReplicationStateDescriptorBuilder::CreateDescriptorForStruct(TagCountMapStruct, Params);
+
+	const FReplicationStateDescriptor* Descriptor = StructNetSerializerConfig.StateDescriptor.GetReference();
+
+	// Validate our assumptions regarding quantized state size and alignment.
+	if ((sizeof(QuantizedType) < Descriptor->InternalSize) || (alignof(QuantizedType) < Descriptor->InternalAlignment))
+	{
+		LowLevelFatalError(TEXT("FMinimalReplicationTagCountMapNetSerializer::FQuantizedType has size %u and alignment %u but requires size %u and alignment %u."), uint32(sizeof(FQuantizedType)), uint32(alignof(FQuantizedType)), Descriptor->InternalSize, Descriptor->InternalAlignment);
+	}
+
+	ensureMsgf(EnumHasAnyFlags(Descriptor->Traits, EReplicationStateTraits::HasDynamicState), TEXT("FMinimalReplicationTagCountMapNetSerializer seems to no longer have dynamic state. Should remove trait from the NetSerializer."));
+	ensureMsgf(!EnumHasAnyFlags(Descriptor->Traits, EReplicationStateTraits::HasObjectReference), TEXT("FMinimalReplicationTagCountMapNetSerializer seems to have object references. Need to add trait to the NetSerializer."));
 }
 
-#endif
+}
 
+#endif
diff --git a/Engine/Plugins/Runtime/GameplayAbilities/Source/GameplayAbilities/Public/GameplayEffectTypes.h b/Engine/Plugins/Runtime/GameplayAbilities/Source/GameplayAbilities/Public/GameplayEffectTypes.h
--- a/Engine/Plugins/Runtime/GameplayAbilities/Source/GameplayAbilities/Public/GameplayEffectTypes.h
+++ b/Engine/Plugins/Runtime/GameplayAbilities/Source/GameplayAbilities/Public/GameplayEffectTypes.h
@@ -30,9 +30,11 @@ struct FActiveGameplayEffect;
 struct FGameplayEffectModCallbackData;
 struct FGameplayEffectSpec;
 struct FHitResult;
+struct FMinimalReplicationTagCountMapForNetSerializer;
 namespace UE::Net
 {
 	struct FGameplayEffectContextHandleAccessorForNetSerializer;
+	class FMinimalReplicationTagCountMapReplicationFragment;
 }
 
 /** Wrappers to convert enum to string. These are fairly slow */
@@ -1635,6 +1637,8 @@ struct GAMEPLAYABILITIES_API FMinimalReplicationTagCountMap
 	void RemoveAllTags();
 
 private:
+	friend FMinimalReplicationTagCountMapForNetSerializer;
+	friend UE::Net::FMinimalReplicationTagCountMapReplicationFragment;
 
 	bool bRequireNonOwningNetConnection = false;
 	void UpdateOwnerTagMap();
diff --git a/Engine/Plugins/Runtime/GameplayAbilities/Source/GameplayAbilities/Public/Serialization/MinimalReplicationTagCountNetSerializer.h b/Engine/Plugins/Runtime/GameplayAbilities/Source/GameplayAbilities/Public/Serialization/MinimalReplicationTagCountNetSerializer.h
new file mode 100644
--- /dev/null
+++ b/Engine/Plugins/Runtime/GameplayAbilities/Source/GameplayAbilities/Public/Serialization/MinimalReplicationTagCountNetSerializer.h
@@ -0,0 +1,19 @@
+// Copyright Epic Games, Inc. All Rights Reserved.
+
+#pragma once
+
+#include "Iris/Serialization/NetSerializer.h"
+#include "MinimalReplicationTagCountNetSerializer.generated.h"
+
+USTRUCT()
+struct FMinimalReplicationTagCountMapNetSerializerConfig : public FNetSerializerConfig
+{
+	GENERATED_BODY()
+};
+
+namespace UE::Net
+{
+
+UE_NET_DECLARE_SERIALIZER(FMinimalReplicationTagCountMapNetSerializer, GAMEPLAYABILITIES_API);
+
+}
diff --git a/Engine/Source/Runtime/Experimental/Iris/Core/Public/Iris/ReplicationState/PropertyNetSerializerInfoRegistry.h b/Engine/Source/Runtime/Experimental/Iris/Core/Public/Iris/ReplicationState/PropertyNetSerializerInfoRegistry.h
--- a/Engine/Source/Runtime/Experimental/Iris/Core/Public/Iris/ReplicationState/PropertyNetSerializerInfoRegistry.h
+++ b/Engine/Source/Runtime/Experimental/Iris/Core/Public/Iris/ReplicationState/PropertyNetSerializerInfoRegistry.h
@@ -167,6 +167,11 @@ const UE::Net::FPropertyNetSerializerInfo& GetPropertyNetSerializerInfo_##Proper
 #define UE_NET_IMPLEMENT_NAMED_STRUCT_NETSERIALIZER_INFO(Name, SerializerName) \
 const UE::Net::FPropertyNetSerializerInfo& GetPropertyNetSerializerInfo_##Name() { static UE::Net::FNamedStructPropertyNetSerializerInfo StaticInstance(Name, UE_NET_GET_SERIALIZER(SerializerName)); return StaticInstance; };
 
+/**
+ * Force a struct to use a custom serializer and custom fragment. Use of custom fragments is highly discouraged as serialization order is modified, ending up after normal class properties.
+ * It probably requires additional memory allocations to create the fragment as well. Only use a last resort when all other options have been exhausted.
+ * It should never be required to use custom fragments except for very rare backward compatibility purposes. Ask the experts first to try to find a better solution.
+ */
 #define UE_NET_IMPLEMENT_NAMED_STRUCT_NETSERIALIZER_WITH_CUSTOM_FRAGMENT_INFO(Name, SerializerName, CreateAndRegisterReplicationFragmentFunction) \
 const UE::Net::FPropertyNetSerializerInfo& GetPropertyNetSerializerInfo_##Name() \
 { \
