[GameplayEffect]
* Moved FGameplayEffectAttributeCaptureDefinition and corresponding enum from GameplayEffectTypes.h to their own header in order to be able to reduce dependencies

diff --git a/Engine/Plugins/Runtime/GameplayAbilities/Source/GameplayAbilities/Public/GameplayEffectAttributeCaptureDefinition.h b/Engine/Plugins/Runtime/GameplayAbilities/Source/GameplayAbilities/Public/GameplayEffectAttributeCaptureDefinition.h
new file mode 100644
--- /dev/null
+++ b/Engine/Plugins/Runtime/GameplayAbilities/Source/GameplayAbilities/Public/GameplayEffectAttributeCaptureDefinition.h
@@ -0,0 +1,69 @@
+// Copyright Epic Games, Inc. All Rights Reserved.
+
+#pragma once
+
+#include "AttributeSet.h"
+#include "GameplayEffectAttributeCaptureDefinition.generated.h"
+
+/** Enumeration for options of where to capture gameplay attributes from for gameplay effects. */
+UENUM()
+enum class EGameplayEffectAttributeCaptureSource : uint8
+{
+	/** Source (caster) of the gameplay effect. */
+	Source,	
+	/** Target (recipient) of the gameplay effect. */
+	Target	
+};
+
+
+/** Struct defining gameplay attribute capture options for gameplay effects */
+USTRUCT(BlueprintType)
+struct GAMEPLAYABILITIES_API FGameplayEffectAttributeCaptureDefinition
+{
+	GENERATED_USTRUCT_BODY()
+
+	FGameplayEffectAttributeCaptureDefinition()
+	{
+		AttributeSource = EGameplayEffectAttributeCaptureSource::Source;
+		bSnapshot = false;
+	}
+
+	FGameplayEffectAttributeCaptureDefinition(FGameplayAttribute InAttribute, EGameplayEffectAttributeCaptureSource InSource, bool InSnapshot)
+		: AttributeToCapture(InAttribute), AttributeSource(InSource), bSnapshot(InSnapshot)
+	{
+
+	}
+
+	/** Gameplay attribute to capture */
+	UPROPERTY(EditDefaultsOnly, Category=Capture)
+	FGameplayAttribute AttributeToCapture;
+
+	/** Source of the gameplay attribute */
+	UPROPERTY(EditDefaultsOnly, Category=Capture)
+	EGameplayEffectAttributeCaptureSource AttributeSource;
+
+	/** Whether the attribute should be snapshotted or not */
+	UPROPERTY(EditDefaultsOnly, Category=Capture)
+	bool bSnapshot;
+
+	/** Equality/Inequality operators */
+	bool operator==(const FGameplayEffectAttributeCaptureDefinition& Other) const;
+	bool operator!=(const FGameplayEffectAttributeCaptureDefinition& Other) const;
+
+	/**
+	 * Get type hash for the capture definition; Implemented to allow usage in TMap
+	 *
+	 * @param CaptureDef Capture definition to get the type hash of
+	 */
+	friend uint32 GetTypeHash(const FGameplayEffectAttributeCaptureDefinition& CaptureDef)
+	{
+		uint32 Hash = 0;
+		Hash = HashCombine(Hash, GetTypeHash(CaptureDef.AttributeToCapture));
+		Hash = HashCombine(Hash, GetTypeHash(static_cast<uint8>(CaptureDef.AttributeSource)));
+		Hash = HashCombine(Hash, GetTypeHash(CaptureDef.bSnapshot));
+		return Hash;
+	}
+
+	FString ToSimpleString() const;
+};
+
diff --git a/Engine/Plugins/Runtime/GameplayAbilities/Source/GameplayAbilities/Public/GameplayEffectTypes.h b/Engine/Plugins/Runtime/GameplayAbilities/Source/GameplayAbilities/Public/GameplayEffectTypes.h
--- a/Engine/Plugins/Runtime/GameplayAbilities/Source/GameplayAbilities/Public/GameplayEffectTypes.h
+++ b/Engine/Plugins/Runtime/GameplayAbilities/Source/GameplayAbilities/Public/GameplayEffectTypes.h
@@ -4,6 +4,7 @@
 
 #include "ActiveGameplayEffectHandle.h"
 #include "CoreMinimal.h"
+#include "GameplayEffectAttributeCaptureDefinition.h"
 #include "UObject/ObjectMacros.h"
 #include "UObject/Object.h"
 #include "UObject/Class.h"
@@ -144,17 +145,6 @@ namespace GameplayEffectUtilities
 	GAMEPLAYABILITIES_API float ComputeStackedModifierMagnitude(float BaseComputedMagnitude, int32 StackCount, EGameplayModOp::Type ModOp);
 }
 
-
-/** Enumeration for options of where to capture gameplay attributes from for gameplay effects. */
-UENUM()
-enum class EGameplayEffectAttributeCaptureSource : uint8
-{
-	/** Source (caster) of the gameplay effect. */
-	Source,	
-	/** Target (recipient) of the gameplay effect. */
-	Target	
-};
-
 /** Enumeration for ways a single GameplayEffect asset can stack. */
 UENUM()
 enum class EGameplayEffectStackingType : uint8
@@ -220,58 +210,6 @@ struct FGameplayModifierEvaluatedData
 };
 
 
-/** Struct defining gameplay attribute capture options for gameplay effects */
-USTRUCT(BlueprintType)
-struct GAMEPLAYABILITIES_API FGameplayEffectAttributeCaptureDefinition
-{
-	GENERATED_USTRUCT_BODY()
-
-	FGameplayEffectAttributeCaptureDefinition()
-	{
-		AttributeSource = EGameplayEffectAttributeCaptureSource::Source;
-		bSnapshot = false;
-	}
-
-	FGameplayEffectAttributeCaptureDefinition(FGameplayAttribute InAttribute, EGameplayEffectAttributeCaptureSource InSource, bool InSnapshot)
-		: AttributeToCapture(InAttribute), AttributeSource(InSource), bSnapshot(InSnapshot)
-	{
-
-	}
-
-	/** Gameplay attribute to capture */
-	UPROPERTY(EditDefaultsOnly, Category=Capture)
-	FGameplayAttribute AttributeToCapture;
-
-	/** Source of the gameplay attribute */
-	UPROPERTY(EditDefaultsOnly, Category=Capture)
-	EGameplayEffectAttributeCaptureSource AttributeSource;
-
-	/** Whether the attribute should be snapshotted or not */
-	UPROPERTY(EditDefaultsOnly, Category=Capture)
-	bool bSnapshot;
-
-	/** Equality/Inequality operators */
-	bool operator==(const FGameplayEffectAttributeCaptureDefinition& Other) const;
-	bool operator!=(const FGameplayEffectAttributeCaptureDefinition& Other) const;
-
-	/**
-	 * Get type hash for the capture definition; Implemented to allow usage in TMap
-	 *
-	 * @param CaptureDef Capture definition to get the type hash of
-	 */
-	friend uint32 GetTypeHash(const FGameplayEffectAttributeCaptureDefinition& CaptureDef)
-	{
-		uint32 Hash = 0;
-		Hash = HashCombine(Hash, GetTypeHash(CaptureDef.AttributeToCapture));
-		Hash = HashCombine(Hash, GetTypeHash(static_cast<uint8>(CaptureDef.AttributeSource)));
-		Hash = HashCombine(Hash, GetTypeHash(CaptureDef.bSnapshot));
-		return Hash;
-	}
-
-	FString ToSimpleString() const;
-};
-
-
 /**
  * Data structure that stores an instigator and related data, such as positions and targets
  * Games can subclass this structure and add game-specific information
