Adding some basic unit tests for the AbilitySystemComponent

diff --git a/Engine/Plugins/Runtime/GameplayAbilities/Source/GameplayAbilities/Private/GameplayEffectTests.cpp b/Engine/Plugins/Runtime/GameplayAbilities/Source/GameplayAbilities/Private/GameplayEffectTests.cpp
deleted file mode 100644
diff --git a/Engine/Plugins/Runtime/GameplayAbilities/Source/GameplayAbilities/Private/Tests/AbilitySystemComponentTests.cpp b/Engine/Plugins/Runtime/GameplayAbilities/Source/GameplayAbilities/Private/Tests/AbilitySystemComponentTests.cpp
new file mode 100644
--- /dev/null
+++ b/Engine/Plugins/Runtime/GameplayAbilities/Source/GameplayAbilities/Private/Tests/AbilitySystemComponentTests.cpp
@@ -0,0 +1,340 @@
+// Copyright Epic Games, Inc. All Rights Reserved.
+
+#include "CoreMinimal.h"
+#include "Misc/AutomationTest.h"
+#include "Engine/DataTable.h"
+#include "Engine/Engine.h"
+#include "GameplayAbilitiesModule.h"
+#include "GameplayTagsModule.h"
+#include "GameFramework/PlayerController.h"
+#include "AbilitySystemComponent.h"
+#include "AbilitySystemTestPawn.h"
+#include "AbilitySystemGlobals.h"
+#include "NativeGameplayTags.h"
+
+// I'm worried about mismatch between Editor/NonEditor builds for NetSerialization, so I've defined this globally.
+UE_DEFINE_GAMEPLAY_TAG_STATIC(TestsDotGenericTag, "Tests.GenericTag");
+
+#if WITH_EDITOR
+
+/** 
+ * This is just a scoped class that helps register for all of the AbilitySystemComponent Callbacks
+ */
+struct TestAllAbilitySystemComponentCallbacks
+{
+	bool bReceivedAbilityActivated = false;
+	bool bReceivedAbilityCommitted = false;
+	bool bReceiveAbilityFailed = false;
+	bool bReceiveAbilityEnded = false;
+
+	TestAllAbilitySystemComponentCallbacks(UAbilitySystemComponent* InAbilitySystemComponent, FAutomationTestBase* Test, const UGameplayAbility* ExpectedAbility)
+		: AbilitySystemComponent(InAbilitySystemComponent)
+	{
+		check(AbilitySystemComponent);
+		check(Test);
+		check(ExpectedAbility);
+
+		auto IsCorrectAbility = [&](const UGameplayAbility* InGameplayAbility)
+		{
+			if (ExpectedAbility->GetInstancingPolicy() == EGameplayAbilityInstancingPolicy::NonInstanced)
+			{
+				return InGameplayAbility == ExpectedAbility;
+			}
+			else
+			{
+				return InGameplayAbility->GetClass()->GetDefaultObject() == ExpectedAbility;
+			}
+		};
+
+		// Register All ASC Callbacks so we can detect the flow
+		OnAbilityActivated = AbilitySystemComponent->AbilityActivatedCallbacks.AddLambda([&](UGameplayAbility* InGameplayAbility)
+			{
+				const bool bIsCorrectAbility = IsCorrectAbility(InGameplayAbility);
+				Test->TestTrue(TEXT(" AbilityActivatedCallbacks with Expected GameplayAbility Instance"), bIsCorrectAbility);
+				bReceivedAbilityActivated = true;
+			});
+
+		OnAbilityCommitted = AbilitySystemComponent->AbilityCommittedCallbacks.AddLambda([&](UGameplayAbility* InGameplayAbility)
+			{
+				const bool bIsCorrectAbility = IsCorrectAbility(InGameplayAbility);
+				Test->TestTrue(TEXT(" AbilityCommittedCallbacks with Expected GameplayAbility Instance"), bIsCorrectAbility);
+				bReceivedAbilityCommitted = true;
+			});
+
+		OnAbilityFailed = AbilitySystemComponent->AbilityFailedCallbacks.AddLambda([&](const UGameplayAbility* InGameplayAbility, const FGameplayTagContainer&)
+			{
+				const bool bIsCorrectAbility = IsCorrectAbility(InGameplayAbility);
+				Test->TestTrue(TEXT(" AbilityFailedCallbacks with Expected GameplayAbility Instance"), bIsCorrectAbility);
+				bReceiveAbilityFailed = true;
+			});
+
+		OnAbilityEnded = AbilitySystemComponent->AbilityEndedCallbacks.AddLambda([&](UGameplayAbility* InGameplayAbility)
+			{
+				const bool bIsCorrectAbility = IsCorrectAbility(InGameplayAbility);
+				Test->TestTrue(TEXT(" AbilityEndedCallbacks with Expected GameplayAbility Instance"), bIsCorrectAbility);
+				bReceiveAbilityEnded = true;
+			});
+	}
+
+	~TestAllAbilitySystemComponentCallbacks()
+	{
+		AbilitySystemComponent->AbilityActivatedCallbacks.Remove(OnAbilityActivated);
+		AbilitySystemComponent->AbilityCommittedCallbacks.Remove(OnAbilityCommitted);
+		AbilitySystemComponent->AbilityFailedCallbacks.Remove(OnAbilityFailed);
+		AbilitySystemComponent->AbilityEndedCallbacks.Remove(OnAbilityEnded);
+	}
+
+private:
+	UAbilitySystemComponent* AbilitySystemComponent = nullptr;
+	FDelegateHandle OnAbilityActivated;
+	FDelegateHandle OnAbilityCommitted;
+	FDelegateHandle OnAbilityFailed;
+	FDelegateHandle OnAbilityEnded;
+}; // End TestAllAbilitySystemComponentCallbacks
+
+/**
+ * The implementation of the test suite
+ */
+class AbilitySystemComponentTestSuite
+{
+public:
+	AbilitySystemComponentTestSuite(UWorld* WorldIn, FAutomationTestBase* TestIn)
+		: World(WorldIn)
+		, Test(TestIn)
+	{
+		// run before each test
+
+		// set up the source actor
+		SourceActor = World->SpawnActor<AAbilitySystemTestPawn>();
+		SourceController = World->SpawnActor<APlayerController>();
+		SourceController->Possess(SourceActor);
+		SourceComponent = SourceActor->GetAbilitySystemComponent();
+
+		// set up the destination actor
+		DestActor = World->SpawnActor<AAbilitySystemTestPawn>();
+		DestComponent = DestActor->GetAbilitySystemComponent();
+	}
+
+	~AbilitySystemComponentTestSuite()
+	{
+		// run after each test
+
+		// destroy the actors
+		if (SourceActor)
+		{
+			World->EditorDestroyActor(SourceActor, false);
+		}
+		if (DestActor)
+		{
+			World->EditorDestroyActor(DestActor, false);
+		}
+	}
+
+public: // the tests
+
+	void Test_ActivateAbilityFlow()
+	{
+		// Tricky part here:  You pass in the FGameplayAbilitySpec to GiveAbility, but it makes a copy and leaves yours unchanged.
+		FGameplayAbilitySpec TempAbilitySpec{ UGameplayAbility::StaticClass(), 1 };
+		FGameplayAbilitySpecHandle GivenAbilitySpecHandle = SourceComponent->GiveAbility(TempAbilitySpec);
+		FGameplayAbilitySpec* AbilitySpec = SourceComponent->FindAbilitySpecFromHandle(GivenAbilitySpecHandle);
+
+		TArray<FGameplayAbilitySpecHandle> GameplayAbilitySpecHandles;
+		SourceComponent->GetAllAbilities(GameplayAbilitySpecHandles);
+		const bool bHasAbility = GameplayAbilitySpecHandles.Num() > 0;
+		const bool bHasCorrectAbility = bHasAbility && (GameplayAbilitySpecHandles[0] == GivenAbilitySpecHandle);
+		Test->TestTrue(TEXT("GiveAbility"), bHasCorrectAbility);
+
+		const TArray<FGameplayAbilitySpec>& ActivatableAbilities = SourceComponent->GetActivatableAbilities();
+		const bool bHasActivatableAbility = ActivatableAbilities.Num() > 0;
+		const bool bBothArraysMatch = bHasAbility && bHasActivatableAbility && (GameplayAbilitySpecHandles[0] == ActivatableAbilities[0].Handle);
+		Test->TestTrue(TEXT("GetAllAbilities() == GetActivatableAbilities().Handle"), bBothArraysMatch);
+
+		if (!Test->HasAnyErrors())
+		{
+			TestAllAbilitySystemComponentCallbacks TestCallbacks{ SourceComponent, Test, AbilitySpec->Ability };
+
+			// Since we have a Controller for the SourceActor in this case, we're executing this all locally.
+			// Try to activate the ability, and see that we've gone through the correct flow.
+			constexpr bool bAllowRemoteActivation = false;
+			const bool bLocalActivation = SourceComponent->TryActivateAbility(GivenAbilitySpecHandle, bAllowRemoteActivation);
+			Test->TestTrue(TEXT("TryActivateAbility executes successfully (using FGameplayAbilitySpecHandle)"), bLocalActivation);
+			Test->TestTrue(TEXT(" AbilitySpec.IsActive() after TryActivateAbility (using FGameplayAbilitySpecHandle)"), AbilitySpec->IsActive());
+			Test->TestTrue(TEXT(" AbilityActivated after TryActivateAbility (using FGameplayAbilitySpecHandle)"), TestCallbacks.bReceivedAbilityActivated);
+			Test->TestTrue(TEXT(" AbilityCommitted after TryActivateAbility (using FGameplayAbilitySpecHandle)"), TestCallbacks.bReceivedAbilityCommitted);
+			Test->TestFalse(TEXT(" AbilityEnded (prematurely) after TryActivateAbility (using FGameplayAbilitySpecHandle)"), TestCallbacks.bReceiveAbilityEnded);
+			Test->TestFalse(TEXT(" AbilityFailed after TryActivateAbility (with an Ability that should succeed)"), TestCallbacks.bReceiveAbilityFailed);
+
+			// Now Cancel the ability and see if we handle that correctly.
+			SourceComponent->CancelAbilityHandle(GivenAbilitySpecHandle);
+			Test->TestTrue(TEXT(" AbilityEnded (after CancelAbilityHandle)"), TestCallbacks.bReceiveAbilityEnded);
+			Test->TestFalse(TEXT(" AbilitySpec.IsActive() (after CancelAbilityHandle)"), AbilitySpec->IsActive());
+
+			// We're the authority
+			Test->TestEqual(TEXT(" ActivationInfo.ActivationMode"), AbilitySpec->ActivationInfo.ActivationMode, EGameplayAbilityActivationMode::Authority);
+		}
+	}
+
+	void Test_FailedAbilityFlow()
+	{
+		// Tricky part here:  You pass in the FGameplayAbilitySpec to GiveAbility, but it makes a copy and leaves yours unchanged.
+		FGameplayAbilitySpec TempAbilitySpec{ UGameplayAbility::StaticClass(), 1 };
+		FGameplayAbilitySpecHandle GivenAbilitySpecHandle = SourceComponent->GiveAbility(TempAbilitySpec);
+		FGameplayAbilitySpec* AbilitySpec = SourceComponent->FindAbilitySpecFromHandle(GivenAbilitySpecHandle);
+
+		TArray<FGameplayAbilitySpecHandle> GameplayAbilitySpecHandles;
+		SourceComponent->GetAllAbilities(GameplayAbilitySpecHandles);
+		const bool bHasAbility = GameplayAbilitySpecHandles.Num() > 0;
+		const bool bHasCorrectAbility = bHasAbility && (GameplayAbilitySpecHandles[0] == GivenAbilitySpecHandle);
+		Test->TestTrue(TEXT("GiveAbility"), bHasCorrectAbility);
+
+		const TArray<FGameplayAbilitySpec>& ActivatableAbilities = SourceComponent->GetActivatableAbilities();
+		const bool bHasActivatableAbility = ActivatableAbilities.Num() > 0;
+		const bool bBothArraysMatch = bHasAbility && bHasActivatableAbility && (GameplayAbilitySpecHandles[0] == ActivatableAbilities[0].Handle);
+		Test->TestTrue(TEXT("GetAllAbilities() == GetActivatableAbilities().Handle"), bBothArraysMatch);
+
+		if (!Test->HasAnyErrors())
+		{
+			TestAllAbilitySystemComponentCallbacks TestCallbacks{ SourceComponent, Test, AbilitySpec->Ability };
+
+			// Since we have a Controller for the SourceActor in this case, we're executing this all locally.
+			// Inhibit Activation, then try to Activate it to trigger the failed flow.
+			SourceComponent->SetUserAbilityActivationInhibited(true);
+
+			constexpr bool bAllowRemoteActivation = false;
+			const bool bLocalActivation = SourceComponent->TryActivateAbility(GivenAbilitySpecHandle, bAllowRemoteActivation);
+
+			Test->TestFalse(TEXT("TryActivateAbility fails (using FGameplayAbilitySpecHandle)"), bLocalActivation);
+			Test->TestFalse(TEXT(" AbilitySpec.IsActive() after TryActivateAbility (using FGameplayAbilitySpecHandle)"), AbilitySpec->IsActive());
+			Test->TestFalse(TEXT(" AbilityActivated after TryActivateAbility (using FGameplayAbilitySpecHandle)"), TestCallbacks.bReceivedAbilityActivated);
+			Test->TestFalse(TEXT(" AbilityCommitted after TryActivateAbility (using FGameplayAbilitySpecHandle)"), TestCallbacks.bReceivedAbilityCommitted);
+			Test->TestFalse(TEXT(" AbilityEnded (prematurely) after TryActivateAbility (using FGameplayAbilitySpecHandle)"), TestCallbacks.bReceiveAbilityEnded);
+			Test->TestTrue(TEXT(" AbilityFailed after TryActivateAbility (with an Ability that should fail)"), TestCallbacks.bReceiveAbilityFailed);
+
+			// Even though we've been rejected, the ActivationMode should still say we're the authority
+			Test->TestEqual(TEXT(" ActivationInfo.ActivationMode"), AbilitySpec->ActivationInfo.ActivationMode, EGameplayAbilityActivationMode::Authority);
+		}
+	}
+
+private: // test helpers
+
+	void TickWorld(float Time)
+	{
+		const float step = 0.1f;
+		while (Time > 0.f)
+		{
+			World->Tick(ELevelTick::LEVELTICK_All, FMath::Min(Time, step));
+			Time -= step;
+
+			// This is terrible but required for subticking like this.
+			// we could always cache the real GFrameCounter at the start of our tests and restore it when finished.
+			GFrameCounter++;
+		}
+	}
+
+private:
+	UWorld* World;
+	FAutomationTestBase* Test;
+
+	AAbilitySystemTestPawn* SourceActor;
+	UAbilitySystemComponent* SourceComponent;
+	AController* SourceController;
+
+	AAbilitySystemTestPawn* DestActor;
+	UAbilitySystemComponent* DestComponent;
+}; // End AbilitySystemComponentTestSuite
+
+/**
+ * This is the Automation Test definition that lets the automation system know which tests we can run (and how to run them)
+ * The implementation will take care of setting up the World and tearing it down.
+ * It will instantiate a fresh AbilitySystemComponentTestSuite for each test.
+ */
+class FAbilitySystemComponentTest : public FAutomationTestBase
+{
+public:
+	typedef void (AbilitySystemComponentTestSuite::*TestFunc)();
+
+	FAbilitySystemComponentTest(const FString& InName)
+		: FAutomationTestBase(InName, false)
+	{
+		auto AddTest = [&](const TestFunc& InTestFunction, const FString& InTestName)
+		{
+			TestFunctions.Add(InTestFunction);
+			TestFunctionNames.Add(InTestName);
+		};
+
+		// list all test functions here
+		AddTest(&AbilitySystemComponentTestSuite::Test_ActivateAbilityFlow, TEXT("ActivateAbilityFlow"));
+		AddTest(&AbilitySystemComponentTestSuite::Test_FailedAbilityFlow, TEXT("FailedAbilityFlow"));
+	}
+
+	virtual uint32 GetTestFlags() const override { return EAutomationTestFlags::EditorContext | EAutomationTestFlags::EngineFilter; }
+	virtual bool IsStressTest() const { return false; }
+	virtual uint32 GetRequiredDeviceNum() const override { return 1; }
+
+protected:
+	virtual FString GetBeautifiedTestName() const override { return "System.AbilitySystem.AbilitySystemComponent"; }
+	virtual void GetTests(TArray<FString>& OutBeautifiedNames, TArray <FString>& OutTestCommands) const override
+	{
+		for (const FString& TestFuncName : TestFunctionNames)
+		{
+			OutBeautifiedNames.Add(TestFuncName);
+			OutTestCommands.Add(TestFuncName);
+		}
+	}
+
+	bool RunTest(const FString& Parameters) override
+	{
+		// find the matching test
+		TestFunc TestFunction = nullptr;
+		for (int32 i = 0; i < TestFunctionNames.Num(); ++i)
+		{
+			if (TestFunctionNames[i] == Parameters)
+			{
+				TestFunction = TestFunctions[i];
+				break;
+			}
+		}
+		if (TestFunction == nullptr)
+		{
+			return false;
+		}
+
+		// get the current data table (to restore later)
+		UDataTable *DataTable = IGameplayAbilitiesModule::Get().GetAbilitySystemGlobals()->GetGlobalAttributeMetaDataTable();
+
+		UWorld *World = UWorld::CreateWorld(EWorldType::Game, false);
+		FWorldContext &WorldContext = GEngine->CreateNewWorldContext(EWorldType::Game);
+		WorldContext.SetCurrentWorld(World);
+
+		FURL URL;
+		//World->GetGameInstanceChecked<UGameInstance>()->EnableListenServer(true);
+		World->InitializeActorsForPlay(URL);
+		World->BeginPlay();
+
+		// run the matching test
+		uint64 InitialFrameCounter = GFrameCounter;
+		{
+			AbilitySystemComponentTestSuite Tester(World, this);
+			(Tester.*TestFunction)();
+		}
+		GFrameCounter = InitialFrameCounter;
+
+		GEngine->DestroyWorldContext(World);
+		World->DestroyWorld(false);
+
+		IGameplayAbilitiesModule::Get().GetAbilitySystemGlobals()->AutomationTestOnly_SetGlobalAttributeDataTable(DataTable);
+		return true;
+	}
+
+	TArray<TestFunc> TestFunctions;
+	TArray<FString> TestFunctionNames;
+};
+
+namespace
+{
+	// Global instance that will register itself to participate in the automated tests
+	FAbilitySystemComponentTest FGameplayAbilitySystemComponentAutomationTestInstance(TEXT("FAbilitySystemComponentTest"));
+}
+
+#endif //WITH_EDITOR
diff --git a/Engine/Plugins/Runtime/GameplayAbilities/Source/GameplayAbilities/Private/Tests/GameplayEffectTests.cpp b/Engine/Plugins/Runtime/GameplayAbilities/Source/GameplayAbilities/Private/Tests/GameplayEffectTests.cpp
new file mode 100644
--- /dev/null
+++ b/Engine/Plugins/Runtime/GameplayAbilities/Source/GameplayAbilities/Private/Tests/GameplayEffectTests.cpp
@@ -0,0 +1,450 @@
+// Copyright Epic Games, Inc. All Rights Reserved.
+
+#include "CoreMinimal.h"
+#include "Misc/AutomationTest.h"
+#include "UObject/Package.h"
+#include "UObject/UnrealType.h"
+#include "Engine/EngineBaseTypes.h"
+#include "Engine/DataTable.h"
+#include "Engine/Engine.h"
+#include "AttributeSet.h"
+#include "GameplayEffectTypes.h"
+#include "GameplayEffect.h"
+#include "GameplayAbilitiesModule.h"
+#include "GameplayTagsManager.h"
+#include "GameplayTagsModule.h"
+#include "AbilitySystemComponent.h"
+#include "AbilitySystemTestPawn.h"
+#include "AbilitySystemTestAttributeSet.h"
+#include "AbilitySystemGlobals.h"
+
+#define SKILL_TEST_TEXT( Format, ... ) FString::Printf(TEXT("%s - %d: %s"), TEXT(__FILE__) , __LINE__ , *FString::Printf(TEXT(Format), ##__VA_ARGS__) )
+
+#if WITH_EDITOR
+
+static UDataTable* CreateGameplayDataTable()
+{
+	FString CSV(TEXT("---,Tag,DevComment"));
+	CSV.Append(TEXT("\r\n0,Damage,"));
+	CSV.Append(TEXT("\r\n1,Damage.Basic,"));
+	CSV.Append(TEXT("\r\n2,Damage.Type1,"));
+	CSV.Append(TEXT("\r\n3,Damage.Type2,"));
+	CSV.Append(TEXT("\r\n4,Damage.Reduce,"));
+	CSV.Append(TEXT("\r\n5,Damage.Buffable,"));
+	CSV.Append(TEXT("\r\n6,Damage.Buff,"));
+	CSV.Append(TEXT("\r\n7,Damage.Physical,"));
+	CSV.Append(TEXT("\r\n8,Damage.Fire,"));
+	CSV.Append(TEXT("\r\n9,Damage.Buffed.FireBuff,"));
+	CSV.Append(TEXT("\r\n10,Damage.Mitigated.Armor,"));
+	CSV.Append(TEXT("\r\n11,Lifesteal,"));
+	CSV.Append(TEXT("\r\n12,Shield,"));
+	CSV.Append(TEXT("\r\n13,Buff,"));
+	CSV.Append(TEXT("\r\n14,Immune,"));
+	CSV.Append(TEXT("\r\n15,FireDamage,"));
+	CSV.Append(TEXT("\r\n16,ShieldAbsorb,"));
+	CSV.Append(TEXT("\r\n17,Stackable,"));
+	CSV.Append(TEXT("\r\n18,Stack,"));
+	CSV.Append(TEXT("\r\n19,Stack.CappedNumber,"));
+	CSV.Append(TEXT("\r\n20,Stack.DiminishingReturns,"));
+	CSV.Append(TEXT("\r\n21,Protect.Damage,"));
+	CSV.Append(TEXT("\r\n22,SpellDmg.Buff,"));
+	CSV.Append(TEXT("\r\n23,GameplayCue.Burning,"));
+
+	UDataTable* DataTable = NewObject<UDataTable>(GetTransientPackage(), FName(TEXT("TempDataTable")));
+	DataTable->RowStruct = FGameplayTagTableRow::StaticStruct();
+	DataTable->CreateTableFromCSVString(CSV);
+
+	const FGameplayTagTableRow* Row = (const FGameplayTagTableRow*)DataTable->GetRowMap()["0"];
+	if (Row)
+	{
+		check(Row->Tag == TEXT("Damage"));
+	}
+	return DataTable;
+}
+
+#define GET_FIELD_CHECKED(Class, Field) FindFieldChecked<FProperty>(Class::StaticClass(), GET_MEMBER_NAME_CHECKED(Class, Field))
+#define CONSTRUCT_CLASS(Class, Name) Class* Name = NewObject<Class>(GetTransientPackage(), FName(TEXT(#Name)))
+
+class GameplayEffectsTestSuite
+{
+public:
+	GameplayEffectsTestSuite(UWorld* WorldIn, FAutomationTestBase* TestIn)
+	: World(WorldIn)
+	, Test(TestIn)
+	{
+		// run before each test
+
+		const float StartingHealth = 100.f;
+		const float StartingMana = 200.f;
+
+		// set up the source actor
+		SourceActor = World->SpawnActor<AAbilitySystemTestPawn>();
+		SourceComponent = SourceActor->GetAbilitySystemComponent();
+		SourceComponent->GetSet<UAbilitySystemTestAttributeSet>()->Health = StartingHealth;
+		SourceComponent->GetSet<UAbilitySystemTestAttributeSet>()->MaxHealth = StartingHealth;
+		SourceComponent->GetSet<UAbilitySystemTestAttributeSet>()->Mana = StartingMana;
+		SourceComponent->GetSet<UAbilitySystemTestAttributeSet>()->MaxMana = StartingMana;
+
+		// set up the destination actor
+		DestActor = World->SpawnActor<AAbilitySystemTestPawn>();
+		DestComponent = DestActor->GetAbilitySystemComponent();
+		DestComponent->GetSet<UAbilitySystemTestAttributeSet>()->Health = StartingHealth;
+		DestComponent->GetSet<UAbilitySystemTestAttributeSet>()->MaxHealth = StartingHealth;
+		DestComponent->GetSet<UAbilitySystemTestAttributeSet>()->Mana = StartingMana;
+		DestComponent->GetSet<UAbilitySystemTestAttributeSet>()->MaxMana = StartingMana;
+	}
+
+	~GameplayEffectsTestSuite()
+	{
+		// run after each test
+
+		// destroy the actors
+		if (SourceActor)
+		{
+			World->EditorDestroyActor(SourceActor, false);
+		}
+		if (DestActor)
+		{
+			World->EditorDestroyActor(DestActor, false);
+		}
+	}
+
+public: // the tests
+
+	void Test_InstantDamage()
+	{
+		const float DamageValue = 5.f;
+		const float StartingHealth = DestComponent->GetSet<UAbilitySystemTestAttributeSet>()->Health;
+
+		// just try and reduce the health attribute
+		{
+			
+			CONSTRUCT_CLASS(UGameplayEffect, BaseDmgEffect);
+			AddModifier(BaseDmgEffect, GET_FIELD_CHECKED(UAbilitySystemTestAttributeSet, Health), EGameplayModOp::Additive, FScalableFloat(-DamageValue));
+			BaseDmgEffect->DurationPolicy = EGameplayEffectDurationType::Instant;
+			
+			SourceComponent->ApplyGameplayEffectToTarget(BaseDmgEffect, DestComponent, 1.f);
+		}
+
+		// make sure health was reduced
+		Test->TestEqual(TEXT("Health Reduced"), DestComponent->GetSet<UAbilitySystemTestAttributeSet>()->Health, StartingHealth - DamageValue);
+	}
+
+	void Test_InstantDamageRemap()
+	{
+		const float DamageValue = 5.f;
+		const float StartingHealth = DestComponent->GetSet<UAbilitySystemTestAttributeSet>()->Health;
+
+		// This is the same as GameplayEffectsTest_InstantDamage but modifies the Damage attribute and confirms it is remapped to -Health by UAbilitySystemTestAttributeSet::PostAttributeModify
+		{
+			CONSTRUCT_CLASS(UGameplayEffect, BaseDmgEffect);
+			AddModifier(BaseDmgEffect, GET_FIELD_CHECKED(UAbilitySystemTestAttributeSet, Damage), EGameplayModOp::Additive, FScalableFloat(DamageValue));
+			BaseDmgEffect->DurationPolicy = EGameplayEffectDurationType::Instant;
+
+			SourceComponent->ApplyGameplayEffectToTarget(BaseDmgEffect, DestComponent, 1.f);
+		}
+
+		// Now we should have lost some health
+		Test->TestEqual(TEXT("Health Reduced"), DestComponent->GetSet<UAbilitySystemTestAttributeSet>()->Health, StartingHealth - DamageValue);
+
+		// Confirm the damage attribute itself was reset to 0 when it was applied to health
+		Test->TestEqual(TEXT("Damage Applied"), DestComponent->GetSet<UAbilitySystemTestAttributeSet>()->Damage, 0.f);
+	}
+
+	void Test_ManaBuff()
+	{
+		const float BuffValue = 30.f;
+		const float StartingMana = DestComponent->GetSet<UAbilitySystemTestAttributeSet>()->Mana;
+
+		FActiveGameplayEffectHandle BuffHandle;
+
+		// apply the buff
+		{
+			CONSTRUCT_CLASS(UGameplayEffect, DamageBuffEffect);
+			AddModifier(DamageBuffEffect, GET_FIELD_CHECKED(UAbilitySystemTestAttributeSet, Mana), EGameplayModOp::Additive, FScalableFloat(BuffValue));
+			DamageBuffEffect->DurationPolicy = EGameplayEffectDurationType::Infinite;
+
+			BuffHandle = SourceComponent->ApplyGameplayEffectToTarget(DamageBuffEffect, DestComponent, 1.f);
+		}
+
+		// check that the value changed
+		Test->TestEqual(TEXT("Mana Buffed"), DestComponent->GetSet<UAbilitySystemTestAttributeSet>()->Mana, StartingMana + BuffValue);
+
+		// remove the effect
+		{
+			DestComponent->RemoveActiveGameplayEffect(BuffHandle);
+		}
+
+		// check that the value changed back
+		Test->TestEqual(TEXT("Mana Restored"), DestComponent->GetSet<UAbilitySystemTestAttributeSet>()->Mana, StartingMana);
+	}
+
+	void Test_PeriodicDamage()
+	{
+		const int32 NumPeriods = 10;
+		const float PeriodSecs = 1.0f;
+		const float DamagePerPeriod = 5.f; 
+		const float StartingHealth = DestComponent->GetSet<UAbilitySystemTestAttributeSet>()->Health;
+
+		// just try and reduce the health attribute
+		{
+			CONSTRUCT_CLASS(UGameplayEffect, BaseDmgEffect);
+			AddModifier(BaseDmgEffect, GET_FIELD_CHECKED(UAbilitySystemTestAttributeSet, Health), EGameplayModOp::Additive, FScalableFloat(-DamagePerPeriod));
+			BaseDmgEffect->DurationPolicy = EGameplayEffectDurationType::HasDuration;
+			BaseDmgEffect->DurationMagnitude = FGameplayEffectModifierMagnitude(FScalableFloat(NumPeriods * PeriodSecs));
+			BaseDmgEffect->Period.Value = PeriodSecs;
+
+			SourceComponent->ApplyGameplayEffectToTarget(BaseDmgEffect, DestComponent, 1.f);
+		}
+
+		int32 NumApplications = 0;
+
+		// Tick a small number to verify the application tick
+		TickWorld(SMALL_NUMBER);
+		++NumApplications;
+
+		Test->TestEqual(TEXT("Health Reduced"), DestComponent->GetSet<UAbilitySystemTestAttributeSet>()->Health, StartingHealth - (DamagePerPeriod * NumApplications));
+
+		// Tick a bit more to address possible floating point issues
+		TickWorld(PeriodSecs * .1f);
+
+		for (int32 i = 0; i < NumPeriods; ++i)
+		{
+			// advance time by one period
+			TickWorld(PeriodSecs);
+
+			++NumApplications;
+
+			// check that health has been reduced
+			Test->TestEqual(TEXT("Health Reduced"), DestComponent->GetSet<UAbilitySystemTestAttributeSet>()->Health, StartingHealth - (DamagePerPeriod * NumApplications));
+		}
+
+		// advance time by one extra period
+		TickWorld(PeriodSecs);
+
+		// should not have reduced further
+		Test->TestEqual(TEXT("Health Reduced"), DestComponent->GetSet<UAbilitySystemTestAttributeSet>()->Health, StartingHealth - (DamagePerPeriod * NumApplications));
+
+		// TODO: test that the effect is no longer applied
+	}
+
+	void Test_StackLimit()
+	{
+		const float Duration = 10.0f;
+		const float HalfDuration = Duration / 2.f;
+		const float ChangePerGE = 5.f;
+		const uint32 StackLimit = 2;
+		const float StartingAttributeValue = DestComponent->GetSet<UAbilitySystemTestAttributeSet>()->StackingAttribute1;
+
+		// Apply one copy of the stacking GE
+		CONSTRUCT_CLASS(UGameplayEffect, StackingEffect);
+		AddModifier(StackingEffect, GET_FIELD_CHECKED(UAbilitySystemTestAttributeSet, StackingAttribute1), EGameplayModOp::Additive, FScalableFloat(ChangePerGE));
+		StackingEffect->DurationPolicy = EGameplayEffectDurationType::HasDuration;
+		StackingEffect->DurationMagnitude = FGameplayEffectModifierMagnitude(FScalableFloat(Duration));
+		StackingEffect->StackLimitCount = StackLimit;
+		StackingEffect->StackingType = EGameplayEffectStackingType::AggregateByTarget;
+		StackingEffect->StackDurationRefreshPolicy = EGameplayEffectStackingDurationPolicy::NeverRefresh;
+		StackingEffect->StackExpirationPolicy = EGameplayEffectStackingExpirationPolicy::ClearEntireStack;
+
+		// Apply the GE StackLimit + 1 times
+		for (uint32 Idx = 0; Idx <= StackLimit; ++Idx)
+		{
+			SourceComponent->ApplyGameplayEffectToTarget(StackingEffect, DestComponent, 1.f);
+		}
+
+		Test->TestEqual(TEXT("Stacking GEs"), DestComponent->GetSet<UAbilitySystemTestAttributeSet>()->StackingAttribute1, StartingAttributeValue + (StackLimit * ChangePerGE));
+	}
+
+	void Test_SetByCallerStackingDuration()
+	{
+		const float Duration = 10.0f;
+		const float HalfDuration = Duration / 2.f;
+		const float ChangePerGE = 5.f;
+		const uint32 StackLimit = 2;
+		const float StartingAttributeValue = DestComponent->GetSet<UAbilitySystemTestAttributeSet>()->StackingAttribute1;
+
+		const FName DurationName(TEXT("Duration"));
+		FSetByCallerFloat SetByCallerDuration;
+		SetByCallerDuration.DataName = DurationName;
+
+		// Apply one copy of the stacking GE
+		CONSTRUCT_CLASS(UGameplayEffect, StackingEffect);
+		AddModifier(StackingEffect, GET_FIELD_CHECKED(UAbilitySystemTestAttributeSet, StackingAttribute1), EGameplayModOp::Additive, FScalableFloat(ChangePerGE));
+		StackingEffect->DurationPolicy = EGameplayEffectDurationType::HasDuration;
+		StackingEffect->DurationMagnitude = FGameplayEffectModifierMagnitude(SetByCallerDuration);
+		StackingEffect->StackLimitCount = StackLimit;
+		StackingEffect->StackingType = EGameplayEffectStackingType::AggregateByTarget;
+		StackingEffect->StackDurationRefreshPolicy = EGameplayEffectStackingDurationPolicy::NeverRefresh;
+		StackingEffect->StackExpirationPolicy = EGameplayEffectStackingExpirationPolicy::RemoveSingleStackAndRefreshDuration;
+
+		// create a spec, set the magnitude and apply the GE
+		{
+			FGameplayEffectSpec	Spec(StackingEffect, FGameplayEffectContextHandle(), 1.f);
+			Spec.SetSetByCallerMagnitude(DurationName, Duration);
+			SourceComponent->ApplyGameplayEffectSpecToTarget(Spec, DestComponent, FPredictionKey());
+		}
+
+		// Tick to partway through the GE's duration and apply a second copy of the GE
+		TickWorld(HalfDuration);
+
+		Test->TestEqual(TEXT("Stacking GEs"), DestComponent->GetSet<UAbilitySystemTestAttributeSet>()->StackingAttribute1, StartingAttributeValue + ChangePerGE);
+
+		// apply second copy of GE
+		{
+			FGameplayEffectSpec	Spec(StackingEffect, FGameplayEffectContextHandle(), 1.f);
+			Spec.SetSetByCallerMagnitude(DurationName, Duration);
+			SourceComponent->ApplyGameplayEffectSpecToTarget(Spec, DestComponent, FPredictionKey());
+		}
+
+		Test->TestEqual(TEXT("Stacking GEs"), DestComponent->GetSet<UAbilitySystemTestAttributeSet>()->StackingAttribute1, StartingAttributeValue + (2 * ChangePerGE));
+
+		// Tick to just after the first GE should have expired
+		TickWorld(HalfDuration + 0.1f);
+
+		// we should have removed one copy and still have the second copy
+		Test->TestEqual(TEXT("Stacking GEs"), DestComponent->GetSet<UAbilitySystemTestAttributeSet>()->StackingAttribute1, StartingAttributeValue + ChangePerGE);
+
+		// check again near the end of the remaining GE's duration
+		TickWorld(Duration - 0.2f);
+
+		Test->TestEqual(TEXT("Stacking GEs"), DestComponent->GetSet<UAbilitySystemTestAttributeSet>()->StackingAttribute1, StartingAttributeValue + ChangePerGE);
+
+		// Tick to a point just after where the remaining GE should have been removed
+		TickWorld(0.2f);
+
+		Test->TestEqual(TEXT("Stacking GEs"), DestComponent->GetSet<UAbilitySystemTestAttributeSet>()->StackingAttribute1, StartingAttributeValue);
+	}
+
+private: // test helpers
+
+	template<typename MODIFIER_T>
+	FGameplayModifierInfo& AddModifier(UGameplayEffect* Effect, FProperty* Property, EGameplayModOp::Type Op, const MODIFIER_T& Magnitude)
+	{
+		int32 Idx = Effect->Modifiers.Num();
+		Effect->Modifiers.SetNum(Idx+1);
+		FGameplayModifierInfo& Info = Effect->Modifiers[Idx];
+		Info.ModifierMagnitude = Magnitude;
+		Info.ModifierOp = Op;
+		Info.Attribute.SetUProperty(Property);
+		return Info;
+	}
+
+	void TickWorld(float Time)
+	{
+		const float step = 0.1f;
+		while (Time > 0.f)
+		{
+			World->Tick(ELevelTick::LEVELTICK_All, FMath::Min(Time, step));
+			Time -= step;
+
+			// This is terrible but required for subticking like this.
+			// we could always cache the real GFrameCounter at the start of our tests and restore it when finished.
+			GFrameCounter++;
+		}
+	}
+
+private:
+	UWorld* World;
+	FAutomationTestBase* Test;
+
+	AAbilitySystemTestPawn* SourceActor;
+	UAbilitySystemComponent* SourceComponent;
+
+	AAbilitySystemTestPawn* DestActor;
+	UAbilitySystemComponent* DestComponent;
+};
+
+#define ADD_TEST(Name) \
+	TestFunctions.Add(&GameplayEffectsTestSuite::Name); \
+	TestFunctionNames.Add(TEXT(#Name))
+
+class FGameplayEffectsTest : public FAutomationTestBase
+{
+public:
+	typedef void (GameplayEffectsTestSuite::*TestFunc)();
+
+	FGameplayEffectsTest(const FString& InName)
+	: FAutomationTestBase(InName, false)
+	{
+		// list all test functions here
+		ADD_TEST(Test_InstantDamage);
+		ADD_TEST(Test_InstantDamageRemap);
+		ADD_TEST(Test_ManaBuff);
+		ADD_TEST(Test_PeriodicDamage);
+		ADD_TEST(Test_StackLimit);
+		ADD_TEST(Test_SetByCallerStackingDuration);
+	}
+
+	virtual uint32 GetTestFlags() const override { return EAutomationTestFlags::EditorContext | EAutomationTestFlags::EngineFilter; }
+	virtual bool IsStressTest() const { return false; }
+	virtual uint32 GetRequiredDeviceNum() const override { return 1; }
+
+protected:
+	virtual FString GetBeautifiedTestName() const override { return "System.AbilitySystem.GameplayEffects"; }
+	virtual void GetTests(TArray<FString>& OutBeautifiedNames, TArray <FString>& OutTestCommands) const override
+	{
+		for (const FString& TestFuncName : TestFunctionNames)
+		{
+			OutBeautifiedNames.Add(TestFuncName);
+			OutTestCommands.Add(TestFuncName);
+		}
+	}
+
+	bool RunTest(const FString& Parameters) override
+	{
+		// find the matching test
+		TestFunc TestFunction = nullptr;
+		for (int32 i = 0; i < TestFunctionNames.Num(); ++i)
+		{
+			if (TestFunctionNames[i] == Parameters)
+			{
+				TestFunction = TestFunctions[i];
+				break;
+			}
+		}
+		if (TestFunction == nullptr)
+		{
+			return false;
+		}
+
+		// get the current data table (to restore later)
+		UDataTable *DataTable = IGameplayAbilitiesModule::Get().GetAbilitySystemGlobals()->GetGlobalAttributeMetaDataTable();
+
+		// setup required GameplayTags
+		UDataTable* TagTable = CreateGameplayDataTable();
+
+		UGameplayTagsManager::Get().PopulateTreeFromDataTable(TagTable);
+
+		UWorld *World = UWorld::CreateWorld(EWorldType::Game, false);
+		FWorldContext &WorldContext = GEngine->CreateNewWorldContext(EWorldType::Game);
+		WorldContext.SetCurrentWorld(World);
+
+		FURL URL;
+		World->InitializeActorsForPlay(URL);
+		World->BeginPlay();
+
+		// run the matching test
+		uint64 InitialFrameCounter = GFrameCounter;
+		{
+			GameplayEffectsTestSuite Tester(World, this);
+			(Tester.*TestFunction)();
+		}
+		GFrameCounter = InitialFrameCounter;
+
+		GEngine->DestroyWorldContext(World);
+		World->DestroyWorld(false);
+
+		IGameplayAbilitiesModule::Get().GetAbilitySystemGlobals()->AutomationTestOnly_SetGlobalAttributeDataTable(DataTable);
+		return true;
+	}
+
+	TArray<TestFunc> TestFunctions;
+	TArray<FString> TestFunctionNames;
+};
+
+namespace
+{
+	FGameplayEffectsTest FGameplayEffectsTestAutomationTestInstance(TEXT("FGameplayEffectsTest"));
+}
+
+#endif //WITH_EDITOR
diff --git a/Engine/Plugins/Runtime/GameplayAbilities/Source/GameplayAbilities/Public/AbilitySystemComponent.h b/Engine/Plugins/Runtime/GameplayAbilities/Source/GameplayAbilities/Public/AbilitySystemComponent.h
--- a/Engine/Plugins/Runtime/GameplayAbilities/Source/GameplayAbilities/Public/AbilitySystemComponent.h
+++ b/Engine/Plugins/Runtime/GameplayAbilities/Source/GameplayAbilities/Public/AbilitySystemComponent.h
@@ -514,7 +514,7 @@ class GAMEPLAYABILITIES_API UAbilitySystemComponent : public UGameplayTasksCompo
 	/** A generic callback anytime an ability is committed (cost/cooldown applied) */
 	FGenericAbilityDelegate AbilityCommittedCallbacks;
 
-	/** Called with a failure reason when an ability fails to execute */
+	/** Called with a failure reason when an ability fails to activate */
 	FAbilityFailedDelegate AbilityFailedCallbacks;
 
 	/** Called when an ability spec's internals have changed */
