Fix for incorrect logic in GAS where a non-Controlled Pawn would be considered remote, even if it was spawned locally on the server (e.g. Vehicles).

diff --git a/Engine/Plugins/Runtime/GameplayAbilities/Source/GameplayAbilities/Private/GameplayAbilityTypes.cpp b/Engine/Plugins/Runtime/GameplayAbilities/Source/GameplayAbilities/Private/GameplayAbilityTypes.cpp
--- a/Engine/Plugins/Runtime/GameplayAbilities/Source/GameplayAbilities/Private/GameplayAbilityTypes.cpp
+++ b/Engine/Plugins/Runtime/GameplayAbilities/Source/GameplayAbilities/Private/GameplayAbilityTypes.cpp
@@ -110,15 +110,18 @@ bool FGameplayAbilityActorInfo::IsLocallyControlled() const
 	}
 	else if (const APawn* OwnerPawn = Cast<APawn>(OwnerActor))
 	{
-		return OwnerPawn->IsLocallyControlled();
-	}
-	else if (IsNetAuthority())
-	{
-		// Non-players are always locally controlled on the server
-		return true;
+		if (OwnerPawn->IsLocallyControlled())
+		{
+			return true;
+		}
+		else if (OwnerPawn->GetController())
+		{
+			// We're controlled, but we're not locally controlled.
+			return false;
+		}
 	}
-
-	return false;
+	
+	return IsNetAuthority();
 }
 
 bool FGameplayAbilityActorInfo::IsLocallyControlledPlayer() const
diff --git a/Engine/Plugins/Runtime/GameplayAbilities/Source/GameplayAbilities/Public/Abilities/GameplayAbilityTypes.h b/Engine/Plugins/Runtime/GameplayAbilities/Source/GameplayAbilities/Public/Abilities/GameplayAbilityTypes.h
--- a/Engine/Plugins/Runtime/GameplayAbilities/Source/GameplayAbilities/Public/Abilities/GameplayAbilityTypes.h
+++ b/Engine/Plugins/Runtime/GameplayAbilities/Source/GameplayAbilities/Public/Abilities/GameplayAbilityTypes.h
@@ -181,9 +181,10 @@ struct GAMEPLAYABILITIES_API FGameplayAbilityActorInfo
 	/** Accessor to get the affected anim instance from the SkeletalMeshComponent */
 	UAnimInstance* GetAnimInstance() const;
 	
-	/** Returns true if this actor is locally controlled. Only true for players on the client that owns them */
+	/** Returns true if this actor is locally controlled. Only true for players on the client that owns them (differs from APawn::IsLocallyControlled which requires a Controller) */
 	bool IsLocallyControlled() const;
 
+	/** Returns true if this actor has a PlayerController that is locally controlled. */
 	bool IsLocallyControlledPlayer() const;
 
 	/** Returns true if the owning actor has net authority */
