Fix for AbilitySystemComponent Unit Tests (was failing on Linux) #jira UE-175231

diff --git a/Engine/Plugins/Runtime/GameplayAbilities/Source/GameplayAbilities/Private/Tests/AbilitySystemComponentTests.cpp b/Engine/Plugins/Runtime/GameplayAbilities/Source/GameplayAbilities/Private/Tests/AbilitySystemComponentTests.cpp
--- a/Engine/Plugins/Runtime/GameplayAbilities/Source/GameplayAbilities/Private/Tests/AbilitySystemComponentTests.cpp
+++ b/Engine/Plugins/Runtime/GameplayAbilities/Source/GameplayAbilities/Private/Tests/AbilitySystemComponentTests.cpp
@@ -27,6 +27,18 @@ struct TestAllAbilitySystemComponentCallbacks
 	bool bReceiveAbilityFailed = false;
 	bool bReceiveAbilityEnded = false;
 
+	static bool IsSameAbility(const UGameplayAbility* GameplayAbility, const UGameplayAbility* ExpectedAbility)
+	{
+		if (ExpectedAbility->GetInstancingPolicy() == EGameplayAbilityInstancingPolicy::NonInstanced)
+		{
+			return GameplayAbility == ExpectedAbility;
+		}
+		else
+		{
+			return GameplayAbility->GetClass()->GetDefaultObject() == ExpectedAbility;
+		}
+	}
+
 	TestAllAbilitySystemComponentCallbacks(UAbilitySystemComponent* InAbilitySystemComponent, FAutomationTestBase* Test, const UGameplayAbility* ExpectedAbility)
 		: AbilitySystemComponent(InAbilitySystemComponent)
 	{
@@ -34,43 +46,31 @@ struct TestAllAbilitySystemComponentCallbacks
 		check(Test);
 		check(ExpectedAbility);
 
-		auto IsCorrectAbility = [&](const UGameplayAbility* InGameplayAbility)
-		{
-			if (ExpectedAbility->GetInstancingPolicy() == EGameplayAbilityInstancingPolicy::NonInstanced)
-			{
-				return InGameplayAbility == ExpectedAbility;
-			}
-			else
-			{
-				return InGameplayAbility->GetClass()->GetDefaultObject() == ExpectedAbility;
-			}
-		};
-
 		// Register All ASC Callbacks so we can detect the flow
-		OnAbilityActivated = AbilitySystemComponent->AbilityActivatedCallbacks.AddLambda([&](UGameplayAbility* InGameplayAbility)
+		OnAbilityActivated = AbilitySystemComponent->AbilityActivatedCallbacks.AddLambda([this, ExpectedAbility, Test](UGameplayAbility* InGameplayAbility)
 			{
-				const bool bIsCorrectAbility = IsCorrectAbility(InGameplayAbility);
+				const bool bIsCorrectAbility = IsSameAbility(InGameplayAbility, ExpectedAbility);
 				Test->TestTrue(TEXT(" AbilityActivatedCallbacks with Expected GameplayAbility Instance"), bIsCorrectAbility);
 				bReceivedAbilityActivated = true;
 			});
 
-		OnAbilityCommitted = AbilitySystemComponent->AbilityCommittedCallbacks.AddLambda([&](UGameplayAbility* InGameplayAbility)
+		OnAbilityCommitted = AbilitySystemComponent->AbilityCommittedCallbacks.AddLambda([this, ExpectedAbility, Test](UGameplayAbility* InGameplayAbility)
 			{
-				const bool bIsCorrectAbility = IsCorrectAbility(InGameplayAbility);
+				const bool bIsCorrectAbility = IsSameAbility(InGameplayAbility, ExpectedAbility);
 				Test->TestTrue(TEXT(" AbilityCommittedCallbacks with Expected GameplayAbility Instance"), bIsCorrectAbility);
 				bReceivedAbilityCommitted = true;
 			});
 
-		OnAbilityFailed = AbilitySystemComponent->AbilityFailedCallbacks.AddLambda([&](const UGameplayAbility* InGameplayAbility, const FGameplayTagContainer&)
+		OnAbilityFailed = AbilitySystemComponent->AbilityFailedCallbacks.AddLambda([this, ExpectedAbility, Test](const UGameplayAbility* InGameplayAbility, const FGameplayTagContainer&)
 			{
-				const bool bIsCorrectAbility = IsCorrectAbility(InGameplayAbility);
+				const bool bIsCorrectAbility = IsSameAbility(InGameplayAbility, ExpectedAbility);
 				Test->TestTrue(TEXT(" AbilityFailedCallbacks with Expected GameplayAbility Instance"), bIsCorrectAbility);
 				bReceiveAbilityFailed = true;
 			});
 
-		OnAbilityEnded = AbilitySystemComponent->AbilityEndedCallbacks.AddLambda([&](UGameplayAbility* InGameplayAbility)
+		OnAbilityEnded = AbilitySystemComponent->AbilityEndedCallbacks.AddLambda([this, ExpectedAbility, Test](UGameplayAbility* InGameplayAbility)
 			{
-				const bool bIsCorrectAbility = IsCorrectAbility(InGameplayAbility);
+				const bool bIsCorrectAbility = IsSameAbility(InGameplayAbility, ExpectedAbility);
 				Test->TestTrue(TEXT(" AbilityEndedCallbacks with Expected GameplayAbility Instance"), bIsCorrectAbility);
 				bReceiveAbilityEnded = true;
 			});
