Added public specificier to classes where recently GENERATED_UCLASS_BODY() macro which automatically comes with a public specificier has been replaced by GENERATED_BODY() macro.

diff --git a/Engine/Plugins/Runtime/GameplayAbilities/Source/GameplayAbilities/Public/Abilities/Tasks/AbilityTask_ApplyRootMotionMoveToForce.h b/Engine/Plugins/Runtime/GameplayAbilities/Source/GameplayAbilities/Public/Abilities/Tasks/AbilityTask_ApplyRootMotionMoveToForce.h
--- a/Engine/Plugins/Runtime/GameplayAbilities/Source/GameplayAbilities/Public/Abilities/Tasks/AbilityTask_ApplyRootMotionMoveToForce.h
+++ b/Engine/Plugins/Runtime/GameplayAbilities/Source/GameplayAbilities/Public/Abilities/Tasks/AbilityTask_ApplyRootMotionMoveToForce.h
@@ -23,7 +23,8 @@ UCLASS()
 class GAMEPLAYABILITIES_API UAbilityTask_ApplyRootMotionMoveToForce : public UAbilityTask_ApplyRootMotion_Base
 {
 	GENERATED_BODY()
-
+	
+public:
 	UPROPERTY(BlueprintAssignable)
 	FApplyRootMotionMoveToForceDelegate OnTimedOut;
 
@@ -41,11 +42,8 @@ class GAMEPLAYABILITIES_API UAbilityTask_ApplyRootMotionMoveToForce : public UAb
 	virtual void OnDestroy(bool AbilityIsEnding) override;
 
 protected:
-
 	virtual void SharedInitAndApply() override;
 
-protected:
-
 	UPROPERTY(Replicated)
 	FVector StartLocation;
 
diff --git a/Engine/Plugins/Runtime/GameplayAbilities/Source/GameplayAbilities/Public/Abilities/Tasks/AbilityTask_WaitGameplayTagQuery.h b/Engine/Plugins/Runtime/GameplayAbilities/Source/GameplayAbilities/Public/Abilities/Tasks/AbilityTask_WaitGameplayTagQuery.h
--- a/Engine/Plugins/Runtime/GameplayAbilities/Source/GameplayAbilities/Public/Abilities/Tasks/AbilityTask_WaitGameplayTagQuery.h
+++ b/Engine/Plugins/Runtime/GameplayAbilities/Source/GameplayAbilities/Public/Abilities/Tasks/AbilityTask_WaitGameplayTagQuery.h
@@ -20,7 +20,8 @@ UCLASS()
 class GAMEPLAYABILITIES_API UAbilityTask_WaitGameplayTagQuery : public UAbilityTask
 {
 	GENERATED_BODY()
-
+	
+public:
 	/** Activates this AbilityTask. */
 	virtual void Activate() override;
 	
