Removing unused function FGameplayAbilityTargetData::AddTargetDataToGameplayCueParameters.  Looking back, it appears it was introduced with nothing calling it (and having no functionality).  I think it's safer to remove than to insert a call to it somewhere.

diff --git a/Engine/Plugins/Runtime/GameplayAbilities/Source/GameplayAbilities/Private/GameplayAbilityTargetTypes.cpp b/Engine/Plugins/Runtime/GameplayAbilities/Source/GameplayAbilities/Private/GameplayAbilityTargetTypes.cpp
--- a/Engine/Plugins/Runtime/GameplayAbilities/Source/GameplayAbilities/Private/GameplayAbilityTargetTypes.cpp
+++ b/Engine/Plugins/Runtime/GameplayAbilities/Source/GameplayAbilities/Private/GameplayAbilityTargetTypes.cpp
@@ -73,10 +73,6 @@ void FGameplayAbilityTargetData::AddTargetDataToContext(FGameplayEffectContextHa
 	}
 }
 
-void FGameplayAbilityTargetData::AddTargetDataToGameplayCueParameters(FGameplayCueParameters& Parameters) const
-{
-}
-
 FString FGameplayAbilityTargetData::ToString() const
 {
 	return TEXT("BASE CLASS");
diff --git a/Engine/Plugins/Runtime/GameplayAbilities/Source/GameplayAbilities/Public/Abilities/GameplayAbilityTargetTypes.h b/Engine/Plugins/Runtime/GameplayAbilities/Source/GameplayAbilities/Public/Abilities/GameplayAbilityTargetTypes.h
--- a/Engine/Plugins/Runtime/GameplayAbilities/Source/GameplayAbilities/Public/Abilities/GameplayAbilityTargetTypes.h
+++ b/Engine/Plugins/Runtime/GameplayAbilities/Source/GameplayAbilities/Public/Abilities/GameplayAbilityTargetTypes.h
@@ -83,9 +83,6 @@ struct GAMEPLAYABILITIES_API FGameplayAbilityTargetData
 	/** Modifies the context and adds this target data to the target data handle stored within */
 	virtual void AddTargetDataToContext(FGameplayEffectContextHandle& Context, bool bIncludeActorArray) const;
 
-	/** Modifies the cue parameters and adds this target data to the target data handle stored within */
-	virtual void AddTargetDataToGameplayCueParameters(FGameplayCueParameters& Parameters) const;
-
 	/** Returns all actors targeted, almost always overridden */
 	virtual TArray<TWeakObjectPtr<AActor>> GetActors() const
 	{
