Modified version of PR #9297: Set CurrentEventData in GameplayAbility (Contributed by [REDACTED_PERSON])

diff --git a/Engine/Plugins/Runtime/GameplayAbilities/Source/GameplayAbilities/Private/Abilities/GameplayAbility.cpp b/Engine/Plugins/Runtime/GameplayAbilities/Source/GameplayAbilities/Private/Abilities/GameplayAbility.cpp
--- a/Engine/Plugins/Runtime/GameplayAbilities/Source/GameplayAbilities/Private/Abilities/GameplayAbility.cpp
+++ b/Engine/Plugins/Runtime/GameplayAbilities/Source/GameplayAbilities/Private/Abilities/GameplayAbility.cpp
@@ -714,6 +714,11 @@ void UGameplayAbility::EndAbility(const FGameplayAbilitySpecHandle Handle, const
 			// Tell owning AbilitySystemComponent that we ended so it can do stuff (including MarkPendingKill us)
 			AbilitySystemComponent->NotifyAbilityEnded(Handle, this, bWasCancelled);
 		}
+
+		if (IsInstantiated())
+		{
+			CurrentEventData = FGameplayEventData{};
+		}
 	}
 }
 
@@ -788,6 +793,11 @@ void UGameplayAbility::PreActivate(const FGameplayAbilitySpecHandle Handle, cons
 	// We could set off a chain that results in calling functions on this ability that rely on the current info being set.
 	SetCurrentInfo(Handle, ActorInfo, ActivationInfo);
 
+	if (TriggerEventData && IsInstantiated())
+	{
+		CurrentEventData = *TriggerEventData;
+	}
+
 	Comp->HandleChangeAbilityCanBeCanceled(AbilityTags, this, true);
 
 	Comp->AddLooseGameplayTags(ActivationOwnedTags);
diff --git a/Engine/Plugins/Runtime/GameplayAbilities/Source/GameplayAbilities/Private/AbilitySystemComponent_Abilities.cpp b/Engine/Plugins/Runtime/GameplayAbilities/Source/GameplayAbilities/Private/AbilitySystemComponent_Abilities.cpp
--- a/Engine/Plugins/Runtime/GameplayAbilities/Source/GameplayAbilities/Private/AbilitySystemComponent_Abilities.cpp
+++ b/Engine/Plugins/Runtime/GameplayAbilities/Source/GameplayAbilities/Private/AbilitySystemComponent_Abilities.cpp
@@ -46,6 +46,7 @@ DECLARE_CYCLE_STAT(TEXT("AbilitySystemComp ServerEndAbility"), STAT_AbilitySyste
 CSV_DECLARE_CATEGORY_MODULE_EXTERN(CORE_API, Basic);
 
 static TAutoConsoleVariable<float> CVarReplayMontageErrorThreshold(TEXT("replay.MontageErrorThreshold"), 0.5f, TEXT("Tolerance level for when montage playback position correction occurs in replays"));
+static TAutoConsoleVariable<bool> CVarAbilitySystemSetActivationInfoMultipleTimes(TEXT("AbilitySystem.SetActivationInfoMultipleTimes"), false, TEXT("Set this to true if some replicated Gameplay Abilities aren't setting their owning actors correctly"));
 
 void UAbilitySystemComponent::InitializeComponent()
 {
@@ -1703,7 +1704,11 @@ bool UAbilitySystemComponent::InternalTryActivateAbility(FGameplayAbilitySpecHan
 			*OutInstancedAbility = InstancedAbility;
 		}
 
-		InstancedAbility->SetCurrentActivationInfo(ActivationInfo);	// Need to push this to the ability if it was instanced.
+		// UGameplayAbility::PreActivate actually sets this internally (via SetCurrentInfo) which happens after replication (this is only set locally).  Let's cautiously remove this code.
+		if (CVarAbilitySystemSetActivationInfoMultipleTimes.GetValueOnGameThread())
+		{
+			InstancedAbility->SetCurrentActivationInfo(ActivationInfo);	// Need to push this to the ability if it was instanced.
+		}
 	}
 
 	MarkAbilitySpecDirty(*Spec);
