Slightly modified PR #10008: Make ModMagnitudeCalc wrapper functions const (Contributed by [REDACTED_PERSON])

diff --git a/Engine/Plugins/Runtime/GameplayAbilities/Source/GameplayAbilities/Private/GameplayModMagnitudeCalculation.cpp b/Engine/Plugins/Runtime/GameplayAbilities/Source/GameplayAbilities/Private/GameplayModMagnitudeCalculation.cpp
--- a/Engine/Plugins/Runtime/GameplayAbilities/Source/GameplayAbilities/Private/GameplayModMagnitudeCalculation.cpp
+++ b/Engine/Plugins/Runtime/GameplayAbilities/Source/GameplayAbilities/Private/GameplayModMagnitudeCalculation.cpp
@@ -44,7 +44,7 @@ bool UGameplayModMagnitudeCalculation::GetCapturedAttributeMagnitude(const FGame
 	return true;
 }
 
-float UGameplayModMagnitudeCalculation::K2_GetCapturedAttributeMagnitude(const FGameplayEffectSpec& EffectSpec, FGameplayAttribute Attribute, const FGameplayTagContainer& SourceTags, const FGameplayTagContainer& TargetTags)
+float UGameplayModMagnitudeCalculation::K2_GetCapturedAttributeMagnitude(const FGameplayEffectSpec& EffectSpec, FGameplayAttribute Attribute, const FGameplayTagContainer& SourceTags, const FGameplayTagContainer& TargetTags) const
 {
 	float Magnitude = 0.0f;
 
@@ -71,17 +71,17 @@ float UGameplayModMagnitudeCalculation::K2_GetCapturedAttributeMagnitude(const F
 	return Magnitude;
 }
 
-float UGameplayModMagnitudeCalculation::GetSetByCallerMagnitudeByTag(const FGameplayEffectSpec& EffectSpec, const FGameplayTag& Tag)
+float UGameplayModMagnitudeCalculation::GetSetByCallerMagnitudeByTag(const FGameplayEffectSpec& EffectSpec, const FGameplayTag& Tag) const
 {
 	return EffectSpec.GetSetByCallerMagnitude(Tag, true, 0.0f);
 }
 
-float UGameplayModMagnitudeCalculation::GetSetByCallerMagnitudeByName(const FGameplayEffectSpec& EffectSpec, const FName& MagnitudeName)
+float UGameplayModMagnitudeCalculation::GetSetByCallerMagnitudeByName(const FGameplayEffectSpec& EffectSpec, const FName& MagnitudeName) const
 {
 	return EffectSpec.GetSetByCallerMagnitude(MagnitudeName, true, 0.0f);
 }
 
-FGameplayTagContainer UGameplayModMagnitudeCalculation::GetSourceAggregatedTags(const FGameplayEffectSpec& EffectSpec)
+FGameplayTagContainer UGameplayModMagnitudeCalculation::GetSourceAggregatedTags(const FGameplayEffectSpec& EffectSpec) const
 {
 	const FGameplayTagContainer* Tags = EffectSpec.CapturedSourceTags.GetAggregatedTags();
 
@@ -93,17 +93,17 @@ FGameplayTagContainer UGameplayModMagnitudeCalculation::GetSourceAggregatedTags(
 	return FGameplayTagContainer();
 }
 
-const FGameplayTagContainer& UGameplayModMagnitudeCalculation::GetSourceActorTags(const FGameplayEffectSpec& EffectSpec)
+const FGameplayTagContainer& UGameplayModMagnitudeCalculation::GetSourceActorTags(const FGameplayEffectSpec& EffectSpec) const
 {
 	return EffectSpec.CapturedSourceTags.GetActorTags();
 }
 
-const FGameplayTagContainer& UGameplayModMagnitudeCalculation::GetSourceSpecTags(const FGameplayEffectSpec& EffectSpec)
+const FGameplayTagContainer& UGameplayModMagnitudeCalculation::GetSourceSpecTags(const FGameplayEffectSpec& EffectSpec) const
 {
 	return EffectSpec.CapturedSourceTags.GetSpecTags();
 }
 
-FGameplayTagContainer UGameplayModMagnitudeCalculation::GetTargetAggregatedTags(const FGameplayEffectSpec& EffectSpec)
+FGameplayTagContainer UGameplayModMagnitudeCalculation::GetTargetAggregatedTags(const FGameplayEffectSpec& EffectSpec) const
 {
 	const FGameplayTagContainer* Tags = EffectSpec.CapturedTargetTags.GetAggregatedTags();
 
@@ -115,12 +115,12 @@ FGameplayTagContainer UGameplayModMagnitudeCalculation::GetTargetAggregatedTags(
 	return FGameplayTagContainer();
 }
 
-const FGameplayTagContainer& UGameplayModMagnitudeCalculation::GetTargetActorTags(const FGameplayEffectSpec& EffectSpec)
+const FGameplayTagContainer& UGameplayModMagnitudeCalculation::GetTargetActorTags(const FGameplayEffectSpec& EffectSpec) const
 {
 	return EffectSpec.CapturedTargetTags.GetActorTags();
 }
 
-const FGameplayTagContainer& UGameplayModMagnitudeCalculation::GetTargetSpecTags(const FGameplayEffectSpec& EffectSpec)
+const FGameplayTagContainer& UGameplayModMagnitudeCalculation::GetTargetSpecTags(const FGameplayEffectSpec& EffectSpec) const
 {
 	return EffectSpec.CapturedTargetTags.GetSpecTags();
 }
diff --git a/Engine/Plugins/Runtime/GameplayAbilities/Source/GameplayAbilities/Public/GameplayModMagnitudeCalculation.h b/Engine/Plugins/Runtime/GameplayAbilities/Source/GameplayAbilities/Public/GameplayModMagnitudeCalculation.h
--- a/Engine/Plugins/Runtime/GameplayAbilities/Source/GameplayAbilities/Public/GameplayModMagnitudeCalculation.h
+++ b/Engine/Plugins/Runtime/GameplayAbilities/Source/GameplayAbilities/Public/GameplayModMagnitudeCalculation.h
@@ -73,8 +73,8 @@ protected:
 	 *
 	 * @return The magnitude value if found, zero otherwise
 	 */
-	UFUNCTION(BlueprintCallable, Category = "Ability|GameplayEffect", meta=(DisplayName="Get Captured Attribute Magnitude", ScriptName="GetCapturedAttributeMagnitude"))
-	float K2_GetCapturedAttributeMagnitude(const FGameplayEffectSpec& EffectSpec, FGameplayAttribute Attribute, const FGameplayTagContainer& SourceTags, const FGameplayTagContainer& TargetTags);
+	UFUNCTION(BlueprintCallable, BlueprintPure = false, Category = "Ability|GameplayEffect", meta=(DisplayName="Get Captured Attribute Magnitude", ScriptName="GetCapturedAttributeMagnitude"))
+	float K2_GetCapturedAttributeMagnitude(const FGameplayEffectSpec& EffectSpec, FGameplayAttribute Attribute, const FGameplayTagContainer& SourceTags, const FGameplayTagContainer& TargetTags) const;
 
 	/**
 	 * Extracts the Set by Caller Magnitude from a Gameplay Effect Spec
@@ -84,8 +84,8 @@ protected:
 	 * 
 	 * @return The magnitude value if found, zero otherwise
 	 */
-	UFUNCTION(BlueprintCallable, Category = "Ability|GameplayEffect")
-	float GetSetByCallerMagnitudeByTag(const FGameplayEffectSpec& EffectSpec, const FGameplayTag& Tag);
+	UFUNCTION(BlueprintCallable, BlueprintPure = false, Category = "Ability|GameplayEffect")
+	float GetSetByCallerMagnitudeByTag(const FGameplayEffectSpec& EffectSpec, const FGameplayTag& Tag) const;
 
 	/**
 	 * Extracts the Set by Caller Magnitude from a Gameplay Effect Spec
@@ -95,8 +95,8 @@ protected:
 	 * 
 	 * @return The magnitude value if found, zero otherwise
 	 */
-	UFUNCTION(BlueprintCallable, Category = "Ability|GameplayEffect")
-	float GetSetByCallerMagnitudeByName(const FGameplayEffectSpec& EffectSpec, const FName& MagnitudeName);
+	UFUNCTION(BlueprintCallable, BlueprintPure = false, Category = "Ability|GameplayEffect")
+	float GetSetByCallerMagnitudeByName(const FGameplayEffectSpec& EffectSpec, const FName& MagnitudeName) const;
 
 	/**
 	 * Copies and returns the source aggregated tags from a Gameplay Effect Spec
@@ -105,8 +105,8 @@ protected:
 	 * 
 	 * @return Gameplay Tag Container with the copied tags. The container will be empty if no captured source tags exist.
 	 */
-	UFUNCTION(BlueprintCallable, Category = "Ability|GameplayEffect")
-	FGameplayTagContainer GetSourceAggregatedTags(const FGameplayEffectSpec& EffectSpec);
+	UFUNCTION(BlueprintCallable, BlueprintPure = false, Category = "Ability|GameplayEffect")
+	FGameplayTagContainer GetSourceAggregatedTags(const FGameplayEffectSpec& EffectSpec) const;
 
 	/**
 	 * Returns the source actor tags from a Gameplay Effect Spec
@@ -116,7 +116,7 @@ protected:
 	 * @return Gameplay Tag Container with the copied tags. The container will be empty if no captured source tags exist.
 	 */
 	UFUNCTION(BlueprintCallable, Category = "Ability|GameplayEffect")
-	const FGameplayTagContainer& GetSourceActorTags(const FGameplayEffectSpec& EffectSpec);
+	const FGameplayTagContainer& GetSourceActorTags(const FGameplayEffectSpec& EffectSpec) const;
 
 	/**
 	 * Returns the source spec tags from a Gameplay Effect Spec
@@ -126,7 +126,7 @@ protected:
 	 * @return Gameplay Tag Container with the copied tags. The container will be empty if no captured source tags exist.
 	 */
 	UFUNCTION(BlueprintCallable, Category = "Ability|GameplayEffect")
-	const FGameplayTagContainer& GetSourceSpecTags(const FGameplayEffectSpec& EffectSpec);
+	const FGameplayTagContainer& GetSourceSpecTags(const FGameplayEffectSpec& EffectSpec) const;
 
 	/**
 	 * Copies and returns the target aggregated tags from a Gameplay Effect Spec
@@ -135,8 +135,8 @@ protected:
 	 * 
 	 * @return Gameplay Tag Container with the copied tags. The container will be empty if no captured source tags exist.
 	 */
-	UFUNCTION(BlueprintCallable, Category = "Ability|GameplayEffect")
-	FGameplayTagContainer GetTargetAggregatedTags(const FGameplayEffectSpec& EffectSpec);
+	UFUNCTION(BlueprintCallable, BlueprintPure = false, Category = "Ability|GameplayEffect")
+	FGameplayTagContainer GetTargetAggregatedTags(const FGameplayEffectSpec& EffectSpec) const;
 
 	/**
 	 * Returns the target actor tags from a Gameplay Effect Spec
@@ -147,7 +147,7 @@ protected:
 	 * @return Gameplay Tag Container with the copied tags. The container will be empty if no captured source tags exist.
 	 */
 	UFUNCTION(BlueprintCallable, Category = "Ability|GameplayEffect")
-	const FGameplayTagContainer& GetTargetActorTags(const FGameplayEffectSpec& EffectSpec);
+	const FGameplayTagContainer& GetTargetActorTags(const FGameplayEffectSpec& EffectSpec) const;
 
 	/**
 	 * Returns the target spec tags from a Gameplay Effect Spec
@@ -158,5 +158,5 @@ protected:
 	 * @return Gameplay Tag Container with the copied tags. The container will be empty if no captured source tags exist.
 	 */
 	UFUNCTION(BlueprintCallable, Category = "Ability|GameplayEffect")
-	const FGameplayTagContainer& GetTargetSpecTags(const FGameplayEffectSpec& EffectSpec);
+	const FGameplayTagContainer& GetTargetSpecTags(const FGameplayEffectSpec& EffectSpec) const;
 };
