Cache trigger event data if it was passed in with GiveAbilityAndActivateOnce and the ability list was locked

diff --git a/Engine/Plugins/Runtime/GameplayAbilities/Source/GameplayAbilities/Private/AbilitySystemComponent_Abilities.cpp b/Engine/Plugins/Runtime/GameplayAbilities/Source/GameplayAbilities/Private/AbilitySystemComponent_Abilities.cpp
--- a/Engine/Plugins/Runtime/GameplayAbilities/Source/GameplayAbilities/Private/AbilitySystemComponent_Abilities.cpp
+++ b/Engine/Plugins/Runtime/GameplayAbilities/Source/GameplayAbilities/Private/AbilitySystemComponent_Abilities.cpp
@@ -327,6 +327,15 @@ FGameplayAbilitySpecHandle UAbilitySystemComponent::GiveAbilityAndActivateOnce(F
 			return FGameplayAbilitySpecHandle();
 		}
 	}
+	else if (GameplayEventData)
+	{
+		// Cache the GameplayEventData in the pending spec (if it was correctly queued)
+		FGameplayAbilitySpec& PendingSpec = AbilityPendingAdds.Last();
+		if (PendingSpec.Handle == AddedAbilityHandle)
+		{
+			PendingSpec.GameplayEventData = MakeShared<FGameplayEventData>(*GameplayEventData);
+		}
+	}
 
 	return AddedAbilityHandle;
 }
@@ -720,7 +729,7 @@ void UAbilitySystemComponent::DecrementAbilityListLock()
 		{
 			if (Spec.bActivateOnce)
 			{
-				GiveAbilityAndActivateOnce(Spec);
+				GiveAbilityAndActivateOnce(Spec, Spec.GameplayEventData.Get());
 			}
 			else
 			{
diff --git a/Engine/Plugins/Runtime/GameplayAbilities/Source/GameplayAbilities/Public/GameplayAbilitySpec.h b/Engine/Plugins/Runtime/GameplayAbilities/Source/GameplayAbilities/Public/GameplayAbilitySpec.h
--- a/Engine/Plugins/Runtime/GameplayAbilities/Source/GameplayAbilities/Public/GameplayAbilitySpec.h
+++ b/Engine/Plugins/Runtime/GameplayAbilities/Source/GameplayAbilities/Public/GameplayAbilitySpec.h
@@ -16,7 +16,7 @@
 
 class UAbilitySystemComponent;
 class UGameplayAbility;
-
+struct FGameplayEventData;
 
 /** Describes the status of activating this ability, this is updated as prediction is handled */
 UENUM(BlueprintType)
@@ -216,6 +216,9 @@ struct GAMEPLAYABILITIES_API FGameplayAbilitySpec : public FFastArraySerializerI
 	UPROPERTY(NotReplicated)
 	uint8 bActivateOnce : 1;
 
+	/** Cached GameplayEventData if this ability was pending for add and activate due to scope lock */
+	TSharedPtr<FGameplayEventData> GameplayEventData = nullptr;
+
 	/** Activation state of this ability. This is not replicated since it needs to be overwritten locally on clients during prediction. */
 	UPROPERTY(NotReplicated)
 	FGameplayAbilityActivationInfo	ActivationInfo;
