Added comment to FAggregatorEvaluateParams

diff --git a/Engine/Plugins/Runtime/GameplayAbilities/Source/GameplayAbilities/Public/GameplayEffectAggregator.h b/Engine/Plugins/Runtime/GameplayAbilities/Source/GameplayAbilities/Public/GameplayEffectAggregator.h
--- a/Engine/Plugins/Runtime/GameplayAbilities/Source/GameplayAbilities/Public/GameplayEffectAggregator.h
+++ b/Engine/Plugins/Runtime/GameplayAbilities/Source/GameplayAbilities/Public/GameplayEffectAggregator.h
@@ -18,7 +18,9 @@ struct GAMEPLAYABILITIES_API FAggregatorEvaluateParameters
 		, IncludePredictiveMods(false) 
 	{}
 
+	/** This tag container is expected to hold all aggregrated tags of the gameplay effect's source: from ability system component, effect spec, ability, calculation specific tags, etc. */
 	const FGameplayTagContainer* SourceTags;
+	/** This tag container is expected to hold all aggregrated tags of the gameplay effect's target: from ability system component, effect spec, ability, calculation specific tags, etc. */
 	const FGameplayTagContainer* TargetTags;
 
 	/** Any mods with one of these handles will be ignored during evaluation */
