A GameplayEffect's bRequireModifierSuccessToTriggerCues setting was already being ignored when there are no Modifiers and no Executions. Now, the setting is also ignored if there were Executions but those didn't attempt to apply any OutModifiers.

diff --git a/Engine/Plugins/Runtime/GameplayAbilities/Source/GameplayAbilities/Private/GameplayEffect.cpp b/Engine/Plugins/Runtime/GameplayAbilities/Source/GameplayAbilities/Private/GameplayEffect.cpp
--- a/Engine/Plugins/Runtime/GameplayAbilities/Source/GameplayAbilities/Private/GameplayEffect.cpp
+++ b/Engine/Plugins/Runtime/GameplayAbilities/Source/GameplayAbilities/Private/GameplayEffect.cpp
@@ -2020,14 +2020,14 @@ void FActiveGameplayEffectsContainer::PredictivelyExecuteEffectSpec(FGameplayEff
 	//		These will modify the base value of attributes
 	// ------------------------------------------------------
 
-	bool ModifierSuccessfullyExecuted = false;
+	bool bAnyModifierSuccessfullyExecuted = false;
 
 	for (int32 ModIdx = 0; ModIdx < SpecToUse.Modifiers.Num(); ++ModIdx)
 	{
 		const FGameplayModifierInfo& ModDef = SpecToUse.Def->Modifiers[ModIdx];
 
 		FGameplayModifierEvaluatedData EvalData(ModDef.Attribute, ModDef.ModifierOp, SpecToUse.GetModifierMagnitude(ModIdx, true));
-		ModifierSuccessfullyExecuted |= InternalExecuteMod(SpecToUse, EvalData);
+		bAnyModifierSuccessfullyExecuted |= InternalExecuteMod(SpecToUse, EvalData);
 	}
 
 	// ------------------------------------------------------
@@ -2035,7 +2035,8 @@ void FActiveGameplayEffectsContainer::PredictivelyExecuteEffectSpec(FGameplayEff
 	//		This will run custom code to 'do stuff'
 	// ------------------------------------------------------
 
-	bool GameplayCuesWereManuallyHandled = false;
+	bool bGameplayCuesWereManuallyHandled = false;
+	bool bExecutionsProducedOutModifiers = false;
 
 	for (const FGameplayEffectExecutionDefinition& CurExecDef : SpecToUse.Def->Executions)
 	{
@@ -2051,6 +2052,7 @@ void FActiveGameplayEffectsContainer::PredictivelyExecuteEffectSpec(FGameplayEff
 
 			// Execute any mods the custom execution yielded
 			TArray<FGameplayModifierEvaluatedData>& OutModifiers = ExecutionOutput.GetOutputModifiersRef();
+			bExecutionsProducedOutModifiers |= OutModifiers.Num() > 0;
 
 			const bool bApplyStackCountToEmittedMods = !ExecutionOutput.IsStackCountHandledManually();
 			const int32 SpecStackCount = SpecToUse.StackCount;
@@ -2062,13 +2064,13 @@ void FActiveGameplayEffectsContainer::PredictivelyExecuteEffectSpec(FGameplayEff
 				{
 					CurExecMod.Magnitude = GameplayEffectUtilities::ComputeStackedModifierMagnitude(CurExecMod.Magnitude, SpecStackCount, CurExecMod.ModifierOp);
 				}
-				ModifierSuccessfullyExecuted |= InternalExecuteMod(SpecToUse, CurExecMod);
+				bAnyModifierSuccessfullyExecuted |= InternalExecuteMod(SpecToUse, CurExecMod);
 			}
 
 			// If execution handled GameplayCues, we dont have to.
 			if (ExecutionOutput.AreGameplayCuesHandledManually())
 			{
-				GameplayCuesWereManuallyHandled = true;
+				bGameplayCuesWereManuallyHandled = true;
 			}
 		}
 	}
@@ -2079,25 +2081,18 @@ void FActiveGameplayEffectsContainer::PredictivelyExecuteEffectSpec(FGameplayEff
 	if (bPredictGameplayCues)
 	{
 		// If there are no modifiers or we don't require modifier success to trigger, we apply the GameplayCue.
-		const bool bHasModifiers = SpecToUse.Modifiers.Num() > 0;
-		const bool bHasExecutions = SpecToUse.Def->Executions.Num() > 0;
-		const bool bHasModifiersOrExecutions = bHasModifiers || bHasExecutions;
+		const bool bHasModifiersFromData = SpecToUse.Modifiers.Num() > 0;
 
-		// If there are no modifiers or we don't require modifier success to trigger, we apply the GameplayCue.
-		bool InvokeGameplayCueExecute = (!bHasModifiersOrExecutions) || !Spec.Def->bRequireModifierSuccessToTriggerCues;
-
-		if (bHasModifiersOrExecutions && ModifierSuccessfullyExecuted)
-		{
-			InvokeGameplayCueExecute = true;
-		}
+		// Apply the GameplayCue if:
+		// - No execution handled it manually AND
+		// -- It didn't require a successful modifier OR there was a successful modifier
+		// -- Or it didn't have any modifiers at all
+		const bool bInvokeGameplayCueExecute = !bGameplayCuesWereManuallyHandled && (
+			(!Spec.Def->bRequireModifierSuccessToTriggerCues || bAnyModifierSuccessfullyExecuted) ||
+			(!bHasModifiersFromData && !bExecutionsProducedOutModifiers)
+		);
 
-		// Don't trigger gameplay cues if one of the executions says it manually handled them
-		if (GameplayCuesWereManuallyHandled)
-		{
-			InvokeGameplayCueExecute = false;
-		}
-
-		if (InvokeGameplayCueExecute && SpecToUse.Def->GameplayCues.Num())
+		if (bInvokeGameplayCueExecute && SpecToUse.Def->GameplayCues.Num())
 		{
 			// TODO: check replication policy. Right now we will replicate every execute via a multicast RPC
 
@@ -2135,14 +2130,14 @@ void FActiveGameplayEffectsContainer::ExecuteActiveEffectsFrom(FGameplayEffectSp
 	//		These will modify the base value of attributes
 	// ------------------------------------------------------
 	
-	bool ModifierSuccessfullyExecuted = false;
+	bool bAnyModifierSuccessfullyExecuted = false;
 
 	for (int32 ModIdx = 0; ModIdx < SpecToUse.Modifiers.Num(); ++ModIdx)
 	{
 		const FGameplayModifierInfo& ModDef = SpecToUse.Def->Modifiers[ModIdx];
 		
 		FGameplayModifierEvaluatedData EvalData(ModDef.Attribute, ModDef.ModifierOp, SpecToUse.GetModifierMagnitude(ModIdx, true));
-		ModifierSuccessfullyExecuted |= InternalExecuteMod(SpecToUse, EvalData);
+		bAnyModifierSuccessfullyExecuted |= InternalExecuteMod(SpecToUse, EvalData);
 	}
 
 	// ------------------------------------------------------
@@ -2152,7 +2147,8 @@ void FActiveGameplayEffectsContainer::ExecuteActiveEffectsFrom(FGameplayEffectSp
 	
 	TArray< FGameplayEffectSpecHandle, TInlineAllocator<4> > ConditionalEffectSpecs;
 
-	bool GameplayCuesWereManuallyHandled = false;
+	bool bGameplayCuesWereManuallyHandled = false;
+	bool bExecutionsProducedOutModifiers = false;
 
 	for (const FGameplayEffectExecutionDefinition& CurExecDef : SpecToUse.Def->Executions)
 	{
@@ -2172,6 +2168,7 @@ void FActiveGameplayEffectsContainer::ExecuteActiveEffectsFrom(FGameplayEffectSp
 
 			// Execute any mods the custom execution yielded
 			TArray<FGameplayModifierEvaluatedData>& OutModifiers = ExecutionOutput.GetOutputModifiersRef();
+			bExecutionsProducedOutModifiers |= OutModifiers.Num() > 0;
 
 			const bool bApplyStackCountToEmittedMods = !ExecutionOutput.IsStackCountHandledManually();
 			const int32 SpecStackCount = SpecToUse.StackCount;
@@ -2183,13 +2180,13 @@ void FActiveGameplayEffectsContainer::ExecuteActiveEffectsFrom(FGameplayEffectSp
 				{
 					CurExecMod.Magnitude = GameplayEffectUtilities::ComputeStackedModifierMagnitude(CurExecMod.Magnitude, SpecStackCount, CurExecMod.ModifierOp);
 				}
-				ModifierSuccessfullyExecuted |= InternalExecuteMod(SpecToUse, CurExecMod);
+				bAnyModifierSuccessfullyExecuted |= InternalExecuteMod(SpecToUse, CurExecMod);
 			}
 
 			// If execution handled GameplayCues, we dont have to.
 			if (ExecutionOutput.AreGameplayCuesHandledManually())
 			{
-				GameplayCuesWereManuallyHandled = true;
+				bGameplayCuesWereManuallyHandled = true;
 			}
 		}
 
@@ -2215,25 +2212,18 @@ void FActiveGameplayEffectsContainer::ExecuteActiveEffectsFrom(FGameplayEffectSp
 	// ------------------------------------------------------
 	
 	// If there are no modifiers or we don't require modifier success to trigger, we apply the GameplayCue.
-	const bool bHasModifiers = SpecToUse.Modifiers.Num() > 0;
-	const bool bHasExecutions = SpecToUse.Def->Executions.Num() > 0;
-	const bool bHasModifiersOrExecutions = bHasModifiers || bHasExecutions;
-
-	// If there are no modifiers or we don't require modifier success to trigger, we apply the GameplayCue.
-	bool InvokeGameplayCueExecute = (!bHasModifiersOrExecutions) || !Spec.Def->bRequireModifierSuccessToTriggerCues;
-
-	if (bHasModifiersOrExecutions && ModifierSuccessfullyExecuted)
-	{
-		InvokeGameplayCueExecute = true;
-	}
-
-	// Don't trigger gameplay cues if one of the executions says it manually handled them
-	if (GameplayCuesWereManuallyHandled)
-	{
-		InvokeGameplayCueExecute = false;
-	}
+	const bool bHasModifiersFromData = SpecToUse.Modifiers.Num() > 0;
+
+	// Apply the GameplayCue if:
+	// - No execution handled it manually AND
+	// -- It didn't require a successful modifier OR there was a successful modifier
+	// -- Or it didn't have any modifiers at all
+	const bool bInvokeGameplayCueExecute = !bGameplayCuesWereManuallyHandled && (
+		(!Spec.Def->bRequireModifierSuccessToTriggerCues || bAnyModifierSuccessfullyExecuted) ||
+		(!bHasModifiersFromData && !bExecutionsProducedOutModifiers)
+	);
 
-	if (InvokeGameplayCueExecute && SpecToUse.Def->GameplayCues.Num())
+	if (bInvokeGameplayCueExecute && SpecToUse.Def->GameplayCues.Num())
 	{
 		// TODO: check replication policy. Right now we will replicate every execute via a multicast RPC
 
diff --git a/Engine/Plugins/Runtime/GameplayAbilities/Source/GameplayAbilities/Public/GameplayEffect.h b/Engine/Plugins/Runtime/GameplayAbilities/Source/GameplayAbilities/Public/GameplayEffect.h
--- a/Engine/Plugins/Runtime/GameplayAbilities/Source/GameplayAbilities/Public/GameplayEffect.h
+++ b/Engine/Plugins/Runtime/GameplayAbilities/Source/GameplayAbilities/Public/GameplayEffect.h
@@ -2022,7 +2022,7 @@ public:
 	UPROPERTY(EditDefaultsOnly, BlueprintReadOnly, Category=Expiration)
 	TArray<TSubclassOf<UGameplayEffect>> RoutineExpirationEffectClasses;
 
-	/** If true, cues will only trigger when GE modifiers succeed being applied (whether through modifiers or executions) */
+	/** If true, cues will only trigger when any GE modifiers succeed being applied (whether through modifiers or executions), or if there were no modifiers application attempts. */
 	UPROPERTY(EditDefaultsOnly, BlueprintReadOnly, Category = Display)
 	bool bRequireModifierSuccessToTriggerCues;
 
