FGameplayTagCountContainer and FMinimalReplicationTagCountMap differed in behavior of negative tag counts. FGameplayTagCountContainer would clamp values and remove any <= 0 counts while FMinimalReplicationTagCountMap had 3 different behaviors of zero and negative tag counts:
* RemoveTag() will store 0 when the current count is <= 0.
*SetTagCount() will store any negative count, but not 0.
*NetSerialize() serializes all TagMap keys regardless of count, so any zero counts from RemoveTag and negative counts from SetTagCount result in a bad HasMatchingGameplayTag on the remote client.
Updated both RemoveTag and SetTagCount in FMinimalReplicationTagCountMap to only store positive counts, consistent with FGameplayTagCountContainer.

diff --git a/Engine/Plugins/Runtime/GameplayAbilities/Source/GameplayAbilities/Private/GameplayEffectTypes.cpp b/Engine/Plugins/Runtime/GameplayAbilities/Source/GameplayAbilities/Private/GameplayEffectTypes.cpp
--- a/Engine/Plugins/Runtime/GameplayAbilities/Source/GameplayAbilities/Private/GameplayEffectTypes.cpp
+++ b/Engine/Plugins/Runtime/GameplayAbilities/Source/GameplayAbilities/Private/GameplayEffectTypes.cpp
@@ -1185,17 +1185,19 @@ const UObject* FGameplayCueParameters::GetSourceObject() const
 void FMinimalReplicationTagCountMap::RemoveTag(const FGameplayTag& Tag)
 {
 	MapID++;
-	int32& Count = TagMap.FindOrAdd(Tag);
-	Count--;
-	if (Count == 0)
+	if (int32* CountPtr = TagMap.Find(Tag))
 	{
-		// Remove from map so that we do not replicate
-		TagMap.Remove(Tag);
+		int32& Count = *CountPtr;
+		Count--;
+		if (Count <= 0)
+		{
+			// Remove from map so that we do not replicate
+			TagMap.Remove(Tag);
+		}
 	}
-	else if (Count < 0)
+	else
 	{
-		ABILITY_LOG(Error, TEXT("FMinimalReplicationTagCountMap::RemoveTag called on Tag %s and count is now < 0"), *Tag.ToString());
-		Count = 0;
+		ABILITY_LOG(Error, TEXT("FMinimalReplicationTagCountMap::RemoveTag called on Tag %s that wasn't in the tag map."), *Tag.ToString());
 	}
 }
 
diff --git a/Engine/Plugins/Runtime/GameplayAbilities/Source/GameplayAbilities/Public/GameplayEffectTypes.h b/Engine/Plugins/Runtime/GameplayAbilities/Source/GameplayAbilities/Public/GameplayEffectTypes.h
--- a/Engine/Plugins/Runtime/GameplayAbilities/Source/GameplayAbilities/Public/GameplayEffectTypes.h
+++ b/Engine/Plugins/Runtime/GameplayAbilities/Source/GameplayAbilities/Public/GameplayEffectTypes.h
@@ -1536,7 +1536,7 @@ struct GAMEPLAYABILITIES_API FMinimalReplicationTagCountMap
 	{
 		MapID++;
 		int32& Count = TagMap.FindOrAdd(Tag);
-		Count = NewCount;
+		Count = FMath::Max(0, NewCount);
 		
 		if (Count == 0)
 		{
