PR #9700: Update FGameplayAttribute::PostSerialize to mark the contained attribute as a searchable name (Contributed by [REDACTED_PERSON])

diff --git a/Engine/Plugins/Runtime/GameplayAbilities/Source/GameplayAbilities/Private/AttributeSet.cpp b/Engine/Plugins/Runtime/GameplayAbilities/Source/GameplayAbilities/Private/AttributeSet.cpp
--- a/Engine/Plugins/Runtime/GameplayAbilities/Source/GameplayAbilities/Private/AttributeSet.cpp
+++ b/Engine/Plugins/Runtime/GameplayAbilities/Source/GameplayAbilities/Private/AttributeSet.cpp
@@ -293,6 +293,11 @@ void FGameplayAttribute::PostSerialize(const FArchive& Ar)
 			}
 		}
 	}
+	if (Ar.IsSaving() && IsValid())
+	{
+		// This marks the attribute "address" for later searching
+		Ar.MarkSearchableName(FGameplayAttribute::StaticStruct(), FName(FString::Printf(TEXT("%s.%s"), *GetUProperty()->GetOwnerVariant().GetName(), *GetUProperty()->GetName())));
+	}
 }
 #endif
 
