PR #9699: Add support for properly replicating different types of FGameplayEffectContext (Contributed by [REDACTED_PERSON])

diff --git a/Engine/Plugins/Runtime/GameplayAbilities/Source/GameplayAbilities/Private/AbilitySystemGlobals.cpp b/Engine/Plugins/Runtime/GameplayAbilities/Source/GameplayAbilities/Private/AbilitySystemGlobals.cpp
--- a/Engine/Plugins/Runtime/GameplayAbilities/Source/GameplayAbilities/Private/AbilitySystemGlobals.cpp
+++ b/Engine/Plugins/Runtime/GameplayAbilities/Source/GameplayAbilities/Private/AbilitySystemGlobals.cpp
@@ -264,6 +264,7 @@ UFunction* UAbilitySystemGlobals::GetGameplayCueFunction(const FGameplayTag& Chi
 void UAbilitySystemGlobals::InitTargetDataScriptStructCache()
 {
 	TargetDataStructCache.InitForType(FGameplayAbilityTargetData::StaticStruct());
+	EffectContextStructCache.InitForType(FGameplayEffectContext::StaticStruct());
 #if UE_WITH_IRIS
 	InitGameplayAbilityTargetDataHandleNetSerializerTypeCache();
 	InitGameplayEffectContextHandleNetSerializerTypeCache();
@@ -689,7 +690,7 @@ bool FNetSerializeScriptStructCache::NetSerialize(FArchive& Ar, UScriptStruct*&
 			return true;
 		}
 
-		ABILITY_LOG(Error, TEXT("Could not script struct at idx %d"), b);
+		ABILITY_LOG(Error, TEXT("Could not find script struct at idx %d"), b);
 		return false;
 	}
 }
diff --git a/Engine/Plugins/Runtime/GameplayAbilities/Source/GameplayAbilities/Private/GameplayEffectTypes.cpp b/Engine/Plugins/Runtime/GameplayAbilities/Source/GameplayAbilities/Private/GameplayEffectTypes.cpp
--- a/Engine/Plugins/Runtime/GameplayAbilities/Source/GameplayAbilities/Private/GameplayEffectTypes.cpp
+++ b/Engine/Plugins/Runtime/GameplayAbilities/Source/GameplayAbilities/Private/GameplayEffectTypes.cpp
@@ -358,6 +358,18 @@ void FGameplayEffectContext::GetOwnedGameplayTags(OUT FGameplayTagContainer& Act
 	}
 }
 
+struct FGameplayEffectContextDeleter
+{
+	FORCEINLINE void operator()(FGameplayEffectContext* Object) const
+	{
+		check(Object);
+		UScriptStruct* ScriptStruct = Object->GetScriptStruct();
+		check(ScriptStruct);
+		ScriptStruct->DestroyStruct(Object);
+		FMemory::Free(Object);
+	}
+};
+
 bool FGameplayEffectContextHandle::NetSerialize(FArchive& Ar, class UPackageMap* Map, bool& bOutSuccess)
 {
 	bool ValidData = Data.IsValid();
@@ -365,33 +377,51 @@ bool FGameplayEffectContextHandle::NetSerialize(FArchive& Ar, class UPackageMap*
 
 	if (ValidData)
 	{
-		if (Ar.IsLoading())
+		TCheckedObjPtr<UScriptStruct> ScriptStruct = Data.IsValid() ? Data->GetScriptStruct() : nullptr;
+		
+		UAbilitySystemGlobals::Get().EffectContextStructCache.NetSerialize(Ar, ScriptStruct.Get());
+
+		if (ScriptStruct.IsValid())
 		{
-			// For now, just always reset/reallocate the data when loading.
-			// Longer term if we want to generalize this and use it for property replication, we should support
-			// only reallocating when necessary
-			
-			if (Data.IsValid() == false)
+			if (Ar.IsLoading())
 			{
-				Data = TSharedPtr<FGameplayEffectContext>(UAbilitySystemGlobals::Get().AllocGameplayEffectContext());
+				// If data is invalid, or a different type, allocate
+				if (!Data.IsValid() || (Data->GetScriptStruct() != ScriptStruct.Get()))
+				{
+					FGameplayEffectContext* NewData = (FGameplayEffectContext*)FMemory::Malloc(ScriptStruct->GetStructureSize());
+					ScriptStruct->InitializeStruct(NewData);
+
+					Data = TSharedPtr<FGameplayEffectContext>(NewData, FGameplayEffectContextDeleter());
+				}
 			}
-		}
 
-		UScriptStruct* ScriptStruct = Data->GetScriptStruct();
+			check(Data.IsValid());
+			if (ScriptStruct->StructFlags & STRUCT_NetSerializeNative)
+			{
+				ScriptStruct->GetCppStructOps()->NetSerialize(Ar, Map, bOutSuccess, Data.Get());
+			}
+			else
+			{
+				// This won't work since FStructProperty::NetSerializeItem is deprecrated.
+				//	1) we have to manually crawl through the topmost struct's fields since we don't have a FStructProperty for it (just the UScriptProperty)
+				//	2) if there are any UStructProperties in the topmost struct's fields, we will assert in FStructProperty::NetSerializeItem.
 
-		if (ScriptStruct->StructFlags & STRUCT_NetSerializeNative)
-		{
-			ScriptStruct->GetCppStructOps()->NetSerialize(Ar, Map, bOutSuccess, Data.Get());
+				ABILITY_LOG(Fatal, TEXT("FGameplayEffectContextHandle::NetSerialize called on data struct %s without a native NetSerialize"), *ScriptStruct->GetName());
+			}
 		}
-		else
+		else if (ScriptStruct.IsError())
 		{
-			// This won't work since FStructProperty::NetSerializeItem is deprecrated.
-			//	1) we have to manually crawl through the topmost struct's fields since we don't have a FStructProperty for it (just the UScriptProperty)
-			//	2) if there are any UStructProperties in the topmost struct's fields, we will assert in FStructProperty::NetSerializeItem.
-
-			ABILITY_LOG(Fatal, TEXT("FGameplayEffectContextHandle::NetSerialize called on data struct %s without a native NetSerialize"), *ScriptStruct->GetName());
+			ABILITY_LOG(Error, TEXT("FGameplayEffectContextHandle::NetSerialize: Bad ScriptStruct serialized, can't recover."));
+			Ar.SetError();
+			Data.Reset();
+			bOutSuccess = false;
+			return false;
 		}
 	}
+	else
+	{
+		Data.Reset();
+	}
 
 	bOutSuccess = true;
 	return true;
diff --git a/Engine/Plugins/Runtime/GameplayAbilities/Source/GameplayAbilities/Public/AbilitySystemGlobals.h b/Engine/Plugins/Runtime/GameplayAbilities/Source/GameplayAbilities/Public/AbilitySystemGlobals.h
--- a/Engine/Plugins/Runtime/GameplayAbilities/Source/GameplayAbilities/Public/AbilitySystemGlobals.h
+++ b/Engine/Plugins/Runtime/GameplayAbilities/Source/GameplayAbilities/Public/AbilitySystemGlobals.h
@@ -308,6 +308,9 @@ class GAMEPLAYABILITIES_API UAbilitySystemGlobals : public UObject
 	UPROPERTY()
 	FNetSerializeScriptStructCache	TargetDataStructCache;
 
+	UPROPERTY()
+	FNetSerializeScriptStructCache	EffectContextStructCache;
+
 	void AddAttributeDefaultTables(const TArray<FSoftObjectPath>& AttribDefaultTableNames);
 
 protected:
