Make UProperty SpawnedAttributes and ActivatableAbilities Transient

#TechNotes: We're making sure these properties values are not saved to the .uasset since they are generated at runtime.

#rnx

diff --git a/Engine/Plugins/Runtime/GameplayAbilities/Source/GameplayAbilities/Public/AbilitySystemComponent.h b/Engine/Plugins/Runtime/GameplayAbilities/Source/GameplayAbilities/Public/AbilitySystemComponent.h
--- a/Engine/Plugins/Runtime/GameplayAbilities/Source/GameplayAbilities/Public/AbilitySystemComponent.h
+++ b/Engine/Plugins/Runtime/GameplayAbilities/Source/GameplayAbilities/Public/AbilitySystemComponent.h
@@ -1625,7 +1625,7 @@ protected:
 	 *	instancing or anything else that the AbilitySystemComponent would provide, then it doesn't need the component to function.
 	 */
 
-	UPROPERTY(ReplicatedUsing=OnRep_ActivateAbilities, BlueprintReadOnly, Category = "Abilities")
+	UPROPERTY(ReplicatedUsing = OnRep_ActivateAbilities, BlueprintReadOnly, Transient, Category = "Abilities")
 	FGameplayAbilitySpecContainer ActivatableAbilities;
 
 	/** Maps from an ability spec to the target data. Used to track replicated data and callbacks */
@@ -1878,7 +1878,7 @@ private:
 private:
 
 	/** List of attribute sets */
-	UPROPERTY(Replicated, ReplicatedUsing=OnRep_SpawnedAttributes)
+	UPROPERTY(Replicated, ReplicatedUsing = OnRep_SpawnedAttributes, Transient)
 	TArray<TObjectPtr<UAttributeSet>>	SpawnedAttributes;
 
 	UFUNCTION()
