ReplicatedPredictionKeys now explicitly copies ReplicatonID and ReplicationKeys to propertly support being initialized with preserved ReplicationIDs

diff --git a/Engine/Plugins/Runtime/GameplayAbilities/Source/GameplayAbilities/Private/GameplayPrediction.cpp b/Engine/Plugins/Runtime/GameplayAbilities/Source/GameplayAbilities/Private/GameplayPrediction.cpp
--- a/Engine/Plugins/Runtime/GameplayAbilities/Source/GameplayAbilities/Private/GameplayPrediction.cpp
+++ b/Engine/Plugins/Runtime/GameplayAbilities/Source/GameplayAbilities/Private/GameplayPrediction.cpp
@@ -6,6 +6,45 @@
 
 #include UE_INLINE_GENERATED_CPP_BY_NAME(GameplayPrediction)
 
+FReplicatedPredictionKeyItem::FReplicatedPredictionKeyItem()
+{
+};
+
+FReplicatedPredictionKeyItem::FReplicatedPredictionKeyItem(const FReplicatedPredictionKeyItem& Other)
+{
+	*this = Other;
+}
+
+FReplicatedPredictionKeyItem& FReplicatedPredictionKeyItem::operator=(const FReplicatedPredictionKeyItem& Other)
+{
+	if (&Other != this)
+	{
+		ReplicationID = Other.ReplicationID;
+		ReplicationKey = Other.ReplicationKey;
+		MostRecentArrayReplicationKey = Other.MostRecentArrayReplicationKey;
+		PredictionKey = Other.PredictionKey;
+	}
+	return *this;
+}
+
+FReplicatedPredictionKeyItem::FReplicatedPredictionKeyItem(FReplicatedPredictionKeyItem&& Other)
+: PredictionKey(MoveTemp(Other.PredictionKey))
+{
+	ReplicationID = Other.ReplicationID;
+	ReplicationKey = Other.ReplicationKey;
+	MostRecentArrayReplicationKey = Other.MostRecentArrayReplicationKey;
+}
+
+FReplicatedPredictionKeyItem& FReplicatedPredictionKeyItem::operator=(FReplicatedPredictionKeyItem&& Other)
+{
+	ReplicationID = Other.ReplicationID;
+	ReplicationKey = Other.ReplicationKey;
+	MostRecentArrayReplicationKey = Other.MostRecentArrayReplicationKey;
+	PredictionKey = MoveTemp(Other.PredictionKey);
+
+	return *this;
+}
+
 /** The key to understanding this function is that when a key is received by the server, we note which connection gave it to us. We only serialize the key back to that client.  */
 bool FPredictionKey::NetSerialize(FArchive& Ar, class UPackageMap* Map, bool& bOutSuccess)
 {
@@ -355,6 +394,10 @@ const int32 FReplicatedPredictionKeyMap::KeyRingBufferSize = 32;
 FReplicatedPredictionKeyMap::FReplicatedPredictionKeyMap()
 {
 	PredictionKeys.SetNum(KeyRingBufferSize);
+	for (FReplicatedPredictionKeyItem& Item : PredictionKeys)
+	{
+		MarkItemDirty(Item);
+	}
 }
 
 bool FReplicatedPredictionKeyMap::NetDeltaSerialize(FNetDeltaSerializeInfo& DeltaParms)
diff --git a/Engine/Plugins/Runtime/GameplayAbilities/Source/GameplayAbilities/Public/AbilitySystemComponent.h b/Engine/Plugins/Runtime/GameplayAbilities/Source/GameplayAbilities/Public/AbilitySystemComponent.h
--- a/Engine/Plugins/Runtime/GameplayAbilities/Source/GameplayAbilities/Public/AbilitySystemComponent.h
+++ b/Engine/Plugins/Runtime/GameplayAbilities/Source/GameplayAbilities/Public/AbilitySystemComponent.h
@@ -1898,7 +1898,7 @@ public:
 	void CacheIsNetSimulated();
 
 	/** PredictionKeys, see more info in GameplayPrediction.h. This has to come *last* in all replicated properties on the AbilitySystemComponent to ensure OnRep/callback order. */
-	UPROPERTY(Replicated)
+	UPROPERTY(Replicated, Transient)
 	FReplicatedPredictionKeyMap ReplicatedPredictionKeyMap;
 
 protected:
diff --git a/Engine/Plugins/Runtime/GameplayAbilities/Source/GameplayAbilities/Public/GameplayPrediction.h b/Engine/Plugins/Runtime/GameplayAbilities/Source/GameplayAbilities/Public/GameplayPrediction.h
--- a/Engine/Plugins/Runtime/GameplayAbilities/Source/GameplayAbilities/Public/GameplayPrediction.h
+++ b/Engine/Plugins/Runtime/GameplayAbilities/Source/GameplayAbilities/Public/GameplayPrediction.h
@@ -531,10 +531,12 @@ struct FReplicatedPredictionKeyItem : public FFastArraySerializerItem
 {
 	GENERATED_USTRUCT_BODY()
 
-	FReplicatedPredictionKeyItem()	
-	{
-		
-	}
+	// As we know that FReplicatedPredictionKeyItem is well behaved and does not leak outside of FReplicatedPredictionKeyMap we allow ReplicationID and ReplicationKey to be copied around to avoid issues when instantiating from archetype or CDO
+	FReplicatedPredictionKeyItem();
+	FReplicatedPredictionKeyItem(const FReplicatedPredictionKeyItem& Other);
+	FReplicatedPredictionKeyItem(FReplicatedPredictionKeyItem&& Other);
+	FReplicatedPredictionKeyItem& operator=(FReplicatedPredictionKeyItem&& other);
+	FReplicatedPredictionKeyItem& operator=(const FReplicatedPredictionKeyItem& other);
 
 	UPROPERTY()
 	FPredictionKey PredictionKey;
