Fixed FGameplayAbilityTargetDataHandleNetSerializer Validate implementation.

Validate operate on source data, not quantized data. The polymorphic serializer knows how to extract the array from the source and validate it.

#rnx

diff --git a/Engine/Plugins/Runtime/GameplayAbilities/Source/GameplayAbilities/Private/Serialization/GameplayAbilityTargetDataHandleNetSerializer.cpp b/Engine/Plugins/Runtime/GameplayAbilities/Source/GameplayAbilities/Private/Serialization/GameplayAbilityTargetDataHandleNetSerializer.cpp
--- a/Engine/Plugins/Runtime/GameplayAbilities/Source/GameplayAbilities/Private/Serialization/GameplayAbilityTargetDataHandleNetSerializer.cpp
+++ b/Engine/Plugins/Runtime/GameplayAbilities/Source/GameplayAbilities/Private/Serialization/GameplayAbilityTargetDataHandleNetSerializer.cpp
@@ -177,12 +177,7 @@ bool FGameplayAbilityTargetDataHandleNetSerializer::IsEqual(FNetSerializationCon
 
 bool FGameplayAbilityTargetDataHandleNetSerializer::Validate(FNetSerializationContext& Context, const FNetValidateArgs& Args)
 {
-	FQuantizedType& SourceValue = *reinterpret_cast<FQuantizedType*>(Args.Source);
-
-	// Forward
-	FNetValidateArgs InternalArgs = Args;
-	InternalArgs.Source = NetSerializerValuePointer(&SourceValue.InternalQuantizedType);
-	return InternalSerializerType::Validate(Context, InternalArgs);
+	return InternalSerializerType::Validate(Context, Args);
 }
 
 void FGameplayAbilityTargetDataHandleNetSerializer::CloneDynamicState(FNetSerializationContext& Context, const FNetCloneDynamicStateArgs& Args)
