#UE Remove unncessary const on return value for GetOwningAbilitySystemComponent. This function was only meant to be a const accessor, not a way to add const to a pointer that was non-const when it was received.

diff --git a/Engine/Plugins/Runtime/GameplayAbilities/Source/GameplayAbilities/Private/ActiveGameplayEffectHandle.cpp b/Engine/Plugins/Runtime/GameplayAbilities/Source/GameplayAbilities/Private/ActiveGameplayEffectHandle.cpp
--- a/Engine/Plugins/Runtime/GameplayAbilities/Source/GameplayAbilities/Private/ActiveGameplayEffectHandle.cpp
+++ b/Engine/Plugins/Runtime/GameplayAbilities/Source/GameplayAbilities/Private/ActiveGameplayEffectHandle.cpp
@@ -27,18 +27,7 @@ FActiveGameplayEffectHandle FActiveGameplayEffectHandle::GenerateNewHandle(UAbil
 	return NewHandle;
 }
 
-UAbilitySystemComponent* FActiveGameplayEffectHandle::GetOwningAbilitySystemComponent()
-{
-	TWeakObjectPtr<UAbilitySystemComponent>* Ptr = GlobalActiveGameplayEffectHandles::Map.Find(*this);
-	if (Ptr)
-	{
-		return Ptr->Get();
-	}
-
-	return nullptr;	
-}
-
-const UAbilitySystemComponent* FActiveGameplayEffectHandle::GetOwningAbilitySystemComponent() const
+UAbilitySystemComponent* FActiveGameplayEffectHandle::GetOwningAbilitySystemComponent() const
 {
 	TWeakObjectPtr<UAbilitySystemComponent>* Ptr = GlobalActiveGameplayEffectHandles::Map.Find(*this);
 	if (Ptr)
diff --git a/Engine/Plugins/Runtime/GameplayAbilities/Source/GameplayAbilities/Public/ActiveGameplayEffectHandle.h b/Engine/Plugins/Runtime/GameplayAbilities/Source/GameplayAbilities/Public/ActiveGameplayEffectHandle.h
--- a/Engine/Plugins/Runtime/GameplayAbilities/Source/GameplayAbilities/Public/ActiveGameplayEffectHandle.h
+++ b/Engine/Plugins/Runtime/GameplayAbilities/Source/GameplayAbilities/Public/ActiveGameplayEffectHandle.h
@@ -47,8 +47,7 @@ struct GAMEPLAYABILITIES_API FActiveGameplayEffectHandle
 	static void ResetGlobalHandleMap();
 
 	/** Returns the ability system component that created this handle */
-	UAbilitySystemComponent* GetOwningAbilitySystemComponent();
-	const UAbilitySystemComponent* GetOwningAbilitySystemComponent() const;
+	UAbilitySystemComponent* GetOwningAbilitySystemComponent() const;
 
 	/** Remove this from the GetOwningAbilitySystemComponent map */
 	void RemoveFromGlobalMap();
