Use COND_Dynamic for the AbilitySystemComponent ActiveGameplayEffects. Disable/enable the replication of MinimalReplicationGameplayCues based on replication mode.
Use cvar AbilitySystem.UseReplicationConditionForActiveGameplayEffects false to fall back to old behavior for ActiveGameplayEffects.

diff --git a/Engine/Plugins/Runtime/GameplayAbilities/Source/GameplayAbilities/Private/AbilitySystemComponent.cpp b/Engine/Plugins/Runtime/GameplayAbilities/Source/GameplayAbilities/Private/AbilitySystemComponent.cpp
--- a/Engine/Plugins/Runtime/GameplayAbilities/Source/GameplayAbilities/Private/AbilitySystemComponent.cpp
+++ b/Engine/Plugins/Runtime/GameplayAbilities/Source/GameplayAbilities/Private/AbilitySystemComponent.cpp
@@ -17,6 +17,7 @@
 #include "GameplayEffectCustomApplicationRequirement.h"
 #include "TimerManager.h"
 #include "Net/Core/PushModel/PushModel.h"
+#include "HAL/IConsoleManager.h"
 
 #include UE_INLINE_GENERATED_CPP_BY_NAME(AbilitySystemComponent)
 
@@ -35,6 +36,9 @@ DECLARE_CYCLE_STAT(TEXT("AbilitySystemComp ExecuteGameplayEffect"), STAT_Ability
 /** Enable to log out all render state create, destroy and updatetransform events */
 #define LOG_RENDER_STATE 0
 
+static bool bUseReplicationConditionForActiveGameplayEffects = true;
+static FAutoConsoleVariableRef CVarUseReplicationConditionForActiveGameplayEffects(TEXT("AbilitySystem.UseReplicationConditionForActiveGameplayEffects"), bUseReplicationConditionForActiveGameplayEffects, TEXT("Whether to be able to determine the replication condition for AbilitySystemComponent::ActiveGameplayEffects at runtime. Removes the need for executing custom logic in FActiveGameplayEffects::NetDeltaSerialize. Default is true."));
+
 UAbilitySystemComponent::UAbilitySystemComponent(const FObjectInitializer& ObjectInitializer)
 	: Super(ObjectInitializer)
 	, GameplayTagCountContainer()
@@ -202,12 +206,16 @@ void UAbilitySystemComponent::OnRegister()
 	}
 
 	ActiveGameplayEffects.RegisterWithOwner(this);
+	ActiveGameplayEffects.SetIsUsingReplicationCondition(bUseReplicationConditionForActiveGameplayEffects);
 	ActivatableAbilities.RegisterWithOwner(this);
 	ActiveGameplayCues.bMinimalReplication = false;
 	ActiveGameplayCues.SetOwner(this);	
 	MinimalReplicationGameplayCues.bMinimalReplication = true;
 	MinimalReplicationGameplayCues.SetOwner(this);
 
+	UpdateActiveGameplayEffectsReplicationCondition();
+	UpdateMinimalReplicationGameplayCuesCondition();
+
 PRAGMA_DISABLE_DEPRECATION_WARNINGS
 	// This field is not replicated (MinimalReplicationTags has a custom serializer),
 	// so we don't need to mark it dirty.
@@ -266,6 +274,7 @@ void UAbilitySystemComponent::CacheIsNetSimulated()
 {
 	bCachedIsNetSimulated = IsNetSimulating();
 	ActiveGameplayEffects.OwnerIsNetAuthority = IsOwnerActorAuthoritative();
+	UpdateActiveGameplayEffectsReplicationCondition();
 }
 
 const FActiveGameplayEffect* UAbilitySystemComponent::GetActiveGameplayEffect(const FActiveGameplayEffectHandle Handle) const
@@ -1609,8 +1618,11 @@ void UAbilitySystemComponent::GetLifetimeReplicatedProps(TArray< FLifetimeProper
 	FDoRepLifetimeParams Params;
 	Params.bIsPushBased = true;
 
-	DOREPLIFETIME_WITH_PARAMS_FAST(UAbilitySystemComponent, SpawnedAttributes, Params);
+	Params.Condition = (bUseReplicationConditionForActiveGameplayEffects ? COND_Dynamic : COND_None);
 	DOREPLIFETIME_WITH_PARAMS_FAST(UAbilitySystemComponent, ActiveGameplayEffects, Params);
+
+	Params.Condition = COND_None;
+	DOREPLIFETIME_WITH_PARAMS_FAST(UAbilitySystemComponent, SpawnedAttributes, Params);
 	DOREPLIFETIME_WITH_PARAMS_FAST(UAbilitySystemComponent, ActiveGameplayCues, Params);
 	DOREPLIFETIME_WITH_PARAMS_FAST(UAbilitySystemComponent, RepAnimMontageInfo, Params);
 	DOREPLIFETIME_WITH_PARAMS_FAST(UAbilitySystemComponent, OwnerActor, Params);
@@ -1635,6 +1647,28 @@ void UAbilitySystemComponent::GetLifetimeReplicatedProps(TArray< FLifetimeProper
 	Super::GetLifetimeReplicatedProps(OutLifetimeProps);
 }
 
+void UAbilitySystemComponent::GetReplicatedCustomConditionState(FCustomPropertyConditionState& OutActiveState) const
+{
+	if (ActiveGameplayEffects.IsUsingReplicationCondition())
+	{
+		DOREPDYNAMICCONDITION_INITCONDITION_FAST(ThisClass, ActiveGameplayEffects, ActiveGameplayEffects.GetReplicationCondition());
+	}
+	DOREPCUSTOMCONDITION_ACTIVE_FAST(ThisClass, MinimalReplicationGameplayCues, MinimalReplicationGameplayCues.ShouldReplicate());
+}
+
+void UAbilitySystemComponent::UpdateActiveGameplayEffectsReplicationCondition()
+{
+	if (ActiveGameplayEffects.IsUsingReplicationCondition())
+	{
+		DOREPDYNAMICCONDITION_SETCONDITION_FAST(ThisClass, ActiveGameplayEffects, ActiveGameplayEffects.GetReplicationCondition());
+	}
+}
+
+void UAbilitySystemComponent::UpdateMinimalReplicationGameplayCuesCondition()
+{
+	DOREPCUSTOMCONDITION_SETACTIVE_FAST(ThisClass, MinimalReplicationGameplayCues, MinimalReplicationGameplayCues.ShouldReplicate());
+}
+
 void UAbilitySystemComponent::ForceReplication()
 {
 	bIsNetDirty = true;
@@ -1722,6 +1756,11 @@ bool UAbilitySystemComponent::HasAuthorityOrPredictionKey(const FGameplayAbility
 void UAbilitySystemComponent::SetReplicationMode(EGameplayEffectReplicationMode NewReplicationMode)
 {
 	ReplicationMode = NewReplicationMode;
+
+	// The changing of replication mode can affect replication conditions for ActiveGameplayEffects and MinimalReplicationGameplayCues.
+	// It's ok to call these before the component is replicated, GetReplicatedCustomConditionState will make sure the conditions are up-to-date.
+	UpdateActiveGameplayEffectsReplicationCondition();
+	UpdateMinimalReplicationGameplayCuesCondition();
 }
 
 IAbilitySystemReplicationProxyInterface* UAbilitySystemComponent::GetReplicationInterface()
diff --git a/Engine/Plugins/Runtime/GameplayAbilities/Source/GameplayAbilities/Private/GameplayCueInterface.cpp b/Engine/Plugins/Runtime/GameplayAbilities/Source/GameplayAbilities/Private/GameplayCueInterface.cpp
--- a/Engine/Plugins/Runtime/GameplayAbilities/Source/GameplayAbilities/Private/GameplayCueInterface.cpp
+++ b/Engine/Plugins/Runtime/GameplayAbilities/Source/GameplayAbilities/Private/GameplayCueInterface.cpp
@@ -323,13 +323,23 @@ bool FActiveGameplayCueContainer::HasCue(const FGameplayTag& Tag) const
 	return false;
 }
 
-bool FActiveGameplayCueContainer::NetDeltaSerialize(FNetDeltaSerializeInfo & DeltaParms)
+bool FActiveGameplayCueContainer::ShouldReplicate() const
 {
 	if (bMinimalReplication && (Owner && Owner->ReplicationMode == EGameplayEffectReplicationMode::Full))
 	{
 		return false;
 	}
 
+	return true;
+}
+
+bool FActiveGameplayCueContainer::NetDeltaSerialize(FNetDeltaSerializeInfo & DeltaParms)
+{
+	if (!ShouldReplicate())
+	{
+		return false;
+	}
+
 	return FastArrayDeltaSerialize<FActiveGameplayCue>(GameplayCues, DeltaParms, *this);
 }
 
diff --git a/Engine/Plugins/Runtime/GameplayAbilities/Source/GameplayAbilities/Private/GameplayEffect.cpp b/Engine/Plugins/Runtime/GameplayAbilities/Source/GameplayAbilities/Private/GameplayEffect.cpp
--- a/Engine/Plugins/Runtime/GameplayAbilities/Source/GameplayAbilities/Private/GameplayEffect.cpp
+++ b/Engine/Plugins/Runtime/GameplayAbilities/Source/GameplayAbilities/Private/GameplayEffect.cpp
@@ -1955,6 +1955,7 @@ FActiveGameplayEffectsContainer::FActiveGameplayEffectsContainer()
 	, ScopedLockCount(0)
 	, PendingRemoves(0)
 	, PendingGameplayEffectHead(nullptr)
+	, bIsUsingReplicationCondition(false)
 {
 	PendingGameplayEffectNext = &PendingGameplayEffectHead;
 	SetDeltaSerializationEnabled(true);
@@ -4024,10 +4025,47 @@ bool FActiveGameplayEffectsContainer::HasApplicationImmunityToSpec(const FGamepl
 	return false;
 }
 
+ELifetimeCondition FActiveGameplayEffectsContainer::GetReplicationCondition() const
+{
+	if (Owner)
+	{
+		const EGameplayEffectReplicationMode ReplicationMode = Owner->ReplicationMode;
+		switch (ReplicationMode)
+		{
+			case EGameplayEffectReplicationMode::Minimal:
+			{
+				return COND_Never;
+			}
+
+			case EGameplayEffectReplicationMode::Mixed:
+			{
+				if (IsNetAuthority())
+				{
+					return COND_OwnerOnly;
+				}
+				else
+				{
+					return COND_ReplayOrOwner;
+				}
+			}
+
+			case EGameplayEffectReplicationMode::Full:
+			default:
+			{
+				return COND_None;
+			}
+		}
+	}
+
+	// If there's no owner we assume full replication is required.
+	return COND_None;
+}
+
 bool FActiveGameplayEffectsContainer::NetDeltaSerialize(FNetDeltaSerializeInfo& DeltaParms)
 {
-	// these tests are only necessary when sending
-	if (DeltaParms.Writer != nullptr && Owner != nullptr)
+	// NOTE: Changes to this testing needs to be reflected in GetReplicatedCondition(), which is what is used if the container has a dynamic lifetime condition, COND_Dynamic, set.
+	// These tests are only needed when sending and the instance isn't using the COND_Dynamic lifetime condition, in which case NetDeltaSerialize wouldn't be called unless it should replicate.
+	if (!IsUsingReplicationCondition() && DeltaParms.Writer != nullptr && Owner != nullptr)
 	{
 		EGameplayEffectReplicationMode ReplicationMode = Owner->ReplicationMode;
 		if (ReplicationMode == EGameplayEffectReplicationMode::Minimal)
diff --git a/Engine/Plugins/Runtime/GameplayAbilities/Source/GameplayAbilities/Public/AbilitySystemComponent.h b/Engine/Plugins/Runtime/GameplayAbilities/Source/GameplayAbilities/Public/AbilitySystemComponent.h
--- a/Engine/Plugins/Runtime/GameplayAbilities/Source/GameplayAbilities/Public/AbilitySystemComponent.h
+++ b/Engine/Plugins/Runtime/GameplayAbilities/Source/GameplayAbilities/Public/AbilitySystemComponent.h
@@ -1614,6 +1614,11 @@ public:
 	virtual void BeginPlay() override;
 
 protected:
+	virtual void GetLifetimeReplicatedProps(TArray<FLifetimeProperty>& OutLifetimeProps) const override;
+	virtual void GetReplicatedCustomConditionState(FCustomPropertyConditionState& OutActiveState) const override;
+
+	void UpdateActiveGameplayEffectsReplicationCondition();
+	void UpdateMinimalReplicationGameplayCuesCondition();
 
 	/**
 	 *	The abilities we can activate. 
diff --git a/Engine/Plugins/Runtime/GameplayAbilities/Source/GameplayAbilities/Public/GameplayCueInterface.h b/Engine/Plugins/Runtime/GameplayAbilities/Source/GameplayAbilities/Public/GameplayCueInterface.h
--- a/Engine/Plugins/Runtime/GameplayAbilities/Source/GameplayAbilities/Public/GameplayCueInterface.h
+++ b/Engine/Plugins/Runtime/GameplayAbilities/Source/GameplayAbilities/Public/GameplayCueInterface.h
@@ -151,6 +151,9 @@ struct GAMEPLAYABILITIES_API FActiveGameplayCueContainer : public FFastArraySeri
 	/** Does explicit check for gameplay cue tag */
 	bool HasCue(const FGameplayTag& Tag) const;
 
+	/** Returns true if the instance should be replicated. If false the property is allowed to be disabled for replication. */
+	bool ShouldReplicate() const;
+
 	bool NetDeltaSerialize(FNetDeltaSerializeInfo & DeltaParms);
 
 	// Will broadcast the OnRemove event for all currently active cues
diff --git a/Engine/Plugins/Runtime/GameplayAbilities/Source/GameplayAbilities/Public/GameplayEffect.h b/Engine/Plugins/Runtime/GameplayAbilities/Source/GameplayAbilities/Public/GameplayEffect.h
--- a/Engine/Plugins/Runtime/GameplayAbilities/Source/GameplayAbilities/Public/GameplayEffect.h
+++ b/Engine/Plugins/Runtime/GameplayAbilities/Source/GameplayAbilities/Public/GameplayEffect.h
@@ -1662,6 +1662,15 @@ struct GAMEPLAYABILITIES_API FActiveGameplayEffectsContainer : public FFastArray
 
 	void CheckDuration(FActiveGameplayEffectHandle Handle);
 
+	/** Returns which ELifetimeCondition can be used for this instance to replicate to relevant connections. This can be used to change the condition if the property is set to use COND_Dynamic in an object's GetLifetimeReplicatedProps implementation. */
+	ELifetimeCondition GetReplicationCondition() const;
+
+	/** Set whether the container is using COND_Dynamic and setting the proper condition at runtime. */
+	void SetIsUsingReplicationCondition(bool bInIsUsingReplicationCondition) { bIsUsingReplicationCondition = bInIsUsingReplicationCondition; }
+
+	/** Return whether the container is using COND_Dynamic and setting the proper condition at runtime. */
+	bool IsUsingReplicationCondition() const { return bIsUsingReplicationCondition; }
+
 	bool NetDeltaSerialize(FNetDeltaSerializeInfo& DeltaParms);
 
 	void Uninitialize();
@@ -1877,6 +1886,8 @@ private:
 	FActiveGameplayEffect*	PendingGameplayEffectHead;	// Head of pending GE linked list
 	FActiveGameplayEffect** PendingGameplayEffectNext;	// Points to the where to store the next pending GE (starts pointing at head, as more are added, points further down the list).
 
+	uint8 bIsUsingReplicationCondition : 1;
+
 	/**
 	 * DO NOT USE DIRECTLY
 	 * STL-like iterators to enable range-based for loop support.
