WaitGameplayTagQuery will now respect its Query even if we have no Gameplay Tags on the Actor.  The issue at hand is that we can create a Query that specifically asks if there are no Tags, or any Tags thus we must evaluate the Query regardless.

diff --git a/Engine/Plugins/Runtime/GameplayAbilities/Source/GameplayAbilities/Private/Abilities/Async/AbilityAsync_WaitGameplayTagQuery.cpp b/Engine/Plugins/Runtime/GameplayAbilities/Source/GameplayAbilities/Private/Abilities/Async/AbilityAsync_WaitGameplayTagQuery.cpp
--- a/Engine/Plugins/Runtime/GameplayAbilities/Source/GameplayAbilities/Private/Abilities/Async/AbilityAsync_WaitGameplayTagQuery.cpp
+++ b/Engine/Plugins/Runtime/GameplayAbilities/Source/GameplayAbilities/Private/Abilities/Async/AbilityAsync_WaitGameplayTagQuery.cpp
@@ -94,7 +94,7 @@ void UAbilityAsync_WaitGameplayTagQuery::EvaluateTagQuery()
 		return;
 	}
 
-	const bool bMatchesQuery = !TargetTags.IsEmpty() && TagQuery.Matches(TargetTags);
+	const bool bMatchesQuery = TagQuery.Matches(TargetTags);
 	const bool bStateChanged = bMatchesQuery != bQueryState;
 	bQueryState = bMatchesQuery;
 
diff --git a/Engine/Plugins/Runtime/GameplayAbilities/Source/GameplayAbilities/Private/Abilities/Tasks/AbilityTask_WaitGameplayTagQuery.cpp b/Engine/Plugins/Runtime/GameplayAbilities/Source/GameplayAbilities/Private/Abilities/Tasks/AbilityTask_WaitGameplayTagQuery.cpp
--- a/Engine/Plugins/Runtime/GameplayAbilities/Source/GameplayAbilities/Private/Abilities/Tasks/AbilityTask_WaitGameplayTagQuery.cpp
+++ b/Engine/Plugins/Runtime/GameplayAbilities/Source/GameplayAbilities/Private/Abilities/Tasks/AbilityTask_WaitGameplayTagQuery.cpp
@@ -58,7 +58,7 @@ void UAbilityTask_WaitGameplayTagQuery::EvaluateTagQuery()
 		return;
 	}
 
-	const bool bMatchesQuery = !TargetTags.IsEmpty() && TagQuery.Matches(TargetTags);
+	const bool bMatchesQuery = TagQuery.Matches(TargetTags);
 	const bool bStateChanged = bMatchesQuery != bQueryState;
 	bQueryState = bMatchesQuery;
 
