GAMEPLAYATTRIBUTE_REPNOTIFY works its way into the Attribute system and then executes its callbacks (PostAttributeChange and AttributeValueChangeDelegates) with incorrect OldValue (due to the variables already being replicated by the time the callbacks are triggered).

diff --git a/Engine/Plugins/Runtime/GameplayAbilities/Source/GameplayAbilities/Private/GameplayEffect.cpp b/Engine/Plugins/Runtime/GameplayAbilities/Source/GameplayAbilities/Private/GameplayEffect.cpp
--- a/Engine/Plugins/Runtime/GameplayAbilities/Source/GameplayAbilities/Private/GameplayEffect.cpp
+++ b/Engine/Plugins/Runtime/GameplayAbilities/Source/GameplayAbilities/Private/GameplayEffect.cpp
@@ -2338,32 +2338,30 @@ void FActiveGameplayEffectsContainer::OnAttributeAggregatorDirty(FAggregator* Ag
 			// 
 			// GlobalFromNetworkUpdate - only set to true when we are coming from an OnRep or when we are coming from an ActiveGameplayEffect container net update.
 			// NetUpdateID - updated once whenever an AttributeSet is received over the network. It will be incremented one time per actor that gets an update.
-					
-			float BaseValue = 0.f;
+			//
+			// See UAttributeSet::PostNetReceive().
+
 			if (!FGameplayAttribute::IsGameplayAttributeDataProperty(Attribute.GetUProperty()))
 			{
 				// Legacy float attribute case requires the base value to be deduced from the final value, as it is not replicated
-				float FinalValue = Owner->GetNumericAttribute(Attribute);
-				BaseValue = Aggregator->ReverseEvaluate(FinalValue, EvaluationParameters);
-				ABILITY_LOG(Log, TEXT("Reverse Evaluated %s. FinalValue: %.2f  BaseValue: %.2f "), *Attribute.GetName(), FinalValue, BaseValue);
-			}
-			else
-			{
-				BaseValue = Owner->GetNumericAttributeBase(Attribute);
+				const float FinalValue = Owner->GetNumericAttribute(Attribute);
+				const float BaseValue = Aggregator->ReverseEvaluate(FinalValue, EvaluationParameters);
+				ABILITY_LOG(Log, TEXT("Reverse Evaluated %s. FinalValue: %.2f  BaseValue: %2.f.  Setting BaseValue.  (Role: %s)"), *Attribute.GetName(), FinalValue, BaseValue, *UEnum::GetValueAsString(Owner->GetOwnerRole()));
+
+				Aggregator->SetBaseValue(BaseValue, false);
 			}
-			
-			Aggregator->SetBaseValue(BaseValue, false);
+
 			Aggregator->NetUpdateID = FScopedAggregatorOnDirtyBatch::NetUpdateID;
 		}
 
 		EvaluationParameters.IncludePredictiveMods = true;
 	}
 
-	float NewValue = Aggregator->Evaluate(EvaluationParameters);
-
+	const float NewValue = Aggregator->Evaluate(EvaluationParameters);
 	if (EvaluationParameters.IncludePredictiveMods)
 	{
-		ABILITY_LOG(Log, TEXT("After Prediction, FinalValue: %.2f"), NewValue);
+		const float OldValue = Owner->GetNumericAttribute(Attribute);
+		ABILITY_LOG(Log, TEXT("[%s] Aggregator Evaluated %s. OldValue: %.2f  NewValue: %.2f"), *UEnum::GetValueAsString(Owner->GetOwnerRole()), *Attribute.GetName(), OldValue, NewValue);
 	}
 
 	InternalUpdateNumericalAttribute(Attribute, NewValue, nullptr, bFromRecursiveCall);
@@ -2561,13 +2559,29 @@ bool FActiveGameplayEffectsContainer::ShouldUseMinimalReplication()
 	return IsNetAuthority() && (Owner->ReplicationMode == EGameplayEffectReplicationMode::Minimal || Owner->ReplicationMode == EGameplayEffectReplicationMode::Mixed);
 }
 
-void FActiveGameplayEffectsContainer::SetBaseAttributeValueFromReplication(const FGameplayAttribute& Attribute, float ServerValue, float OldValue)
+void FActiveGameplayEffectsContainer::SetBaseAttributeValueFromReplication(const FGameplayAttribute& Attribute, float ServerBaseValue, float OldBaseValue)
 {
 	FAggregatorRef* RefPtr = AttributeAggregatorMap.Find(Attribute);
 	if (RefPtr && RefPtr->Get())
 	{
 		FAggregator* Aggregator = RefPtr->Get();
-		
+		if (FGameplayAttribute::IsGameplayAttributeDataProperty(Attribute.GetUProperty()))
+		{
+			// Reset to the server's old value
+			constexpr bool bDoNotExecuteCallbacksValue = false;
+			Aggregator->SetBaseValue(OldBaseValue, bDoNotExecuteCallbacksValue);
+
+			// Evaluate what the old value would have resulted in...  We do this to ensure the correct "old value" for the callbacks.
+			FAggregatorEvaluateParameters EvaluationParameters;
+			EvaluationParameters.IncludePredictiveMods = true;
+			float OldEvaluatedValue = Aggregator->Evaluate(EvaluationParameters);
+			Owner->SetNumericAttribute_Internal(Attribute, OldEvaluatedValue);
+
+			// Now set the new value and go through all of the aggregations...
+			Aggregator->SetBaseValue(ServerBaseValue, bDoNotExecuteCallbacksValue);
+			ABILITY_LOG(Log, TEXT("SetBaseAttributeValueFromReplication [%s]: %s rewound to state NewBaseValue: %.2f  OldCurrentValue: %.2f"), OwnerIsNetAuthority ? TEXT("Authority") : TEXT("Client"), *Attribute.AttributeName, ServerBaseValue, OldEvaluatedValue);
+		}
+
 		FScopedAggregatorOnDirtyBatch::GlobalFromNetworkUpdate = true;
 		OnAttributeAggregatorDirty(Aggregator, Attribute);
 		FScopedAggregatorOnDirtyBatch::GlobalFromNetworkUpdate = false;
@@ -2578,7 +2592,7 @@ void FActiveGameplayEffectsContainer::SetBaseAttributeValueFromReplication(const
 		PRAGMA_DISABLE_DEPRECATION_WARNINGS
 		if (FOnGameplayAttributeChange* LegacyDelegate = AttributeChangeDelegates.Find(Attribute))
 		{
-			LegacyDelegate->Broadcast(ServerValue, nullptr);
+			LegacyDelegate->Broadcast(ServerBaseValue, nullptr);
 		}
 		PRAGMA_ENABLE_DEPRECATION_WARNINGS
 
@@ -2586,8 +2600,8 @@ void FActiveGameplayEffectsContainer::SetBaseAttributeValueFromReplication(const
 		{
 			FOnAttributeChangeData CallbackData;
 			CallbackData.Attribute = Attribute;
-			CallbackData.NewValue = ServerValue;
-			CallbackData.OldValue = OldValue;
+			CallbackData.NewValue = ServerBaseValue;
+			CallbackData.OldValue = OldBaseValue;
 			CallbackData.GEModData = nullptr;
 
 			Delegate->Broadcast(CallbackData);
@@ -2743,9 +2757,8 @@ void FActiveGameplayEffectsContainer::CaptureAttributeForGameplayEffect(OUT FGam
 
 void FActiveGameplayEffectsContainer::InternalUpdateNumericalAttribute(FGameplayAttribute Attribute, float NewValue, const FGameplayEffectModCallbackData* ModData, bool bFromRecursiveCall)
 {
-	ABILITY_LOG(Log, TEXT("Property %s new value is: %.2f"), *Attribute.GetName(), NewValue);
-
-	float OldValue = Owner->GetNumericAttribute(Attribute);
+	const float OldValue = Owner->GetNumericAttribute(Attribute);
+	ABILITY_LOG(Log, TEXT("[%s] InternalUpdateNumericalAttribute %s OldValue = %.2f  NewValue = %.2f."), *UEnum::GetValueAsString(Owner->GetOwnerRole()), *Attribute.GetName(), OldValue, NewValue);
 	Owner->SetNumericAttribute_Internal(Attribute, NewValue);
 	
 	if (!bFromRecursiveCall)
@@ -2789,7 +2802,6 @@ void FActiveGameplayEffectsContainer::SetAttributeBaseValue(FGameplayAttribute A
 	}
 
 	float OldBaseValue = 0.0f;
-	bool bBaseValueSet = false;
 
 	Set->PreAttributeBaseChange(Attribute, NewBaseValue);
 
@@ -2804,9 +2816,13 @@ void FActiveGameplayEffectsContainer::SetAttributeBaseValue(FGameplayAttribute A
 		{
 			OldBaseValue = DataPtr->GetBaseValue();
 			DataPtr->SetBaseValue(NewBaseValue);
-			bBaseValueSet = true;
 		}
 	}
+	else
+	{
+		// If it's a simple float value, then base value == current value (unless there's an aggregator, below)
+		OldBaseValue = Owner->GetNumericAttribute(Attribute);
+	}
 
 	FAggregatorRef* RefPtr = AttributeAggregatorMap.Find(Attribute);
 	if (RefPtr)
@@ -2818,20 +2834,14 @@ void FActiveGameplayEffectsContainer::SetAttributeBaseValue(FGameplayAttribute A
 		// will update the actual AttributeSet property value for us.		
 		OldBaseValue = Aggregator->GetBaseValue();
 		Aggregator->SetBaseValue(NewBaseValue);
-		bBaseValueSet = true;
 	}
 	// if there is no aggregator set the current value (base == current in this case)
 	else
 	{
-		OldBaseValue = Owner->GetNumericAttribute(Attribute);
 		InternalUpdateNumericalAttribute(Attribute, NewBaseValue, nullptr);
-		bBaseValueSet = true;
 	}
 
-	if (bBaseValueSet)
-	{
-		Set->PostAttributeBaseChange(Attribute, OldBaseValue, NewBaseValue);
-	}
+	Set->PostAttributeBaseChange(Attribute, OldBaseValue, NewBaseValue);
 }
 
 float FActiveGameplayEffectsContainer::GetAttributeBaseValue(FGameplayAttribute Attribute) const
