[NO NON-PII COMMIT MESSAGE]

diff --git a/Engine/Plugins/Runtime/GameplayAbilities/Source/GameplayAbilities/Private/GameplayEffect.cpp b/Engine/Plugins/Runtime/GameplayAbilities/Source/GameplayAbilities/Private/GameplayEffect.cpp
--- a/Engine/Plugins/Runtime/GameplayAbilities/Source/GameplayAbilities/Private/GameplayEffect.cpp
+++ b/Engine/Plugins/Runtime/GameplayAbilities/Source/GameplayAbilities/Private/GameplayEffect.cpp
@@ -3576,13 +3576,17 @@ bool FActiveGameplayEffectsContainer::InternalRemoveActiveGameplayEffect(int32 I
 		// Mark the effect pending remove, and remove all side effects from the effect
 		InternalOnActiveGameplayEffectRemoved(Effect, ShouldInvokeGameplayCueEvent, GameplayEffectRemovalInfo);
 
-		if (Effect.DurationHandle.IsValid())
+		// Check world validity in case RemoveActiveGameplayEffect is called during world teardown
+		if (UWorld* World = Owner->GetWorld())
 		{
-			Owner->GetWorld()->GetTimerManager().ClearTimer(Effect.DurationHandle);
-		}
-		if (Effect.PeriodHandle.IsValid())
-		{
-			Owner->GetWorld()->GetTimerManager().ClearTimer(Effect.PeriodHandle);
+			if (Effect.DurationHandle.IsValid())
+			{
+				World->GetTimerManager().ClearTimer(Effect.DurationHandle);
+			}
+			if (Effect.PeriodHandle.IsValid())
+			{
+				World->GetTimerManager().ClearTimer(Effect.PeriodHandle);
+			}
 		}
 
 		// Remove this handle from the global map
