Iris - Added missing validation of quantized size and alignment. Adjusted size.

#rnx

diff --git a/Engine/Plugins/Runtime/GameplayAbilities/Source/GameplayAbilities/Private/Serialization/GameplayAbilityTargetingLocationInfoNetSerializer.cpp b/Engine/Plugins/Runtime/GameplayAbilities/Source/GameplayAbilities/Private/Serialization/GameplayAbilityTargetingLocationInfoNetSerializer.cpp
--- a/Engine/Plugins/Runtime/GameplayAbilities/Source/GameplayAbilities/Private/Serialization/GameplayAbilityTargetingLocationInfoNetSerializer.cpp
+++ b/Engine/Plugins/Runtime/GameplayAbilities/Source/GameplayAbilities/Private/Serialization/GameplayAbilityTargetingLocationInfoNetSerializer.cpp
@@ -32,7 +32,7 @@ struct FGameplayAbilityTargetingLocationInfoNetSerializer
 	struct FQuantizedData
 	{
 		// Keep TargetingLocationInfo first. It simplifies things.
-		alignas(16) uint8 TargetingLocationInfo[128 - 1];
+		alignas(16) uint8 TargetingLocationInfo[192 - 1];
 
 		uint8 LocationType;
 	};
@@ -407,6 +407,13 @@ void FGameplayAbilityTargetingLocationInfoNetSerializer::FNetSerializerRegistryD
 	const FReplicationStateDescriptor* Descriptor = StructNetSerializerConfig.StateDescriptor.GetReference();
 	ValidateForwardingNetSerializerTraits(&UE_NET_GET_SERIALIZER(FGameplayAbilityTargetingLocationInfoNetSerializer), Descriptor->Traits);
 
+	// Validate our assumptions regarding quantized state size and alignment.
+	static_assert(offsetof(FQuantizedData, TargetingLocationInfo) == 0, "");
+	if ((sizeof(FQuantizedData::TargetingLocationInfo) < Descriptor->InternalSize) || alignof(FQuantizedData) < Descriptor->InternalAlignment)
+	{
+		LowLevelFatalError(TEXT("FQuantizedData::TargetingLocationInfo has size %u and alignment %u but requires size %u and alignment %u."), uint32(sizeof(FQuantizedData::TargetingLocationInfo)), uint32(alignof(FQuantizedData)), uint32(Descriptor->InternalSize), uint32(Descriptor->InternalAlignment));
+	}
+
 	// Find the properties we're interested in and build change masks for the different location types.
 	{
 		const ANSICHAR* const PropertyNames[] =
