Fix compile errors
#rnx

diff --git a/Engine/Plugins/Runtime/GameplayAbilities/Source/GameplayAbilitiesEditor/Private/GameplayAbilityAudit.cpp b/Engine/Plugins/Runtime/GameplayAbilities/Source/GameplayAbilitiesEditor/Private/GameplayAbilityAudit.cpp
--- a/Engine/Plugins/Runtime/GameplayAbilities/Source/GameplayAbilitiesEditor/Private/GameplayAbilityAudit.cpp
+++ b/Engine/Plugins/Runtime/GameplayAbilities/Source/GameplayAbilitiesEditor/Private/GameplayAbilityAudit.cpp
@@ -3,12 +3,13 @@
 #include "GameplayAbilityAudit.h"
 
 #include "Abilities/GameplayAbility.h"
+#include "Engine/BlueprintGeneratedClass.h"
 #include "GameplayAbilityBlueprint.h"
 
-#include "BlueprintGraph/Classes/K2Node_CallFunction.h"
-#include "BlueprintGraph/Classes/K2Node_BaseAsyncTask.h"
-#include "BlueprintGraph/Classes/K2Node_MacroInstance.h"
-#include "BlueprintGraph/Classes/K2Node_VariableSet.h"
+#include "K2Node_CallFunction.h"
+#include "K2Node_BaseAsyncTask.h"
+#include "K2Node_MacroInstance.h"
+#include "K2Node_VariableSet.h"
 
 #include "Factories/DataTableFactory.h"
 
diff --git a/Engine/Plugins/Runtime/GameplayAbilities/Source/GameplayAbilitiesEditor/Public/GameplayAbilityAudit.h b/Engine/Plugins/Runtime/GameplayAbilities/Source/GameplayAbilitiesEditor/Public/GameplayAbilityAudit.h
--- a/Engine/Plugins/Runtime/GameplayAbilities/Source/GameplayAbilitiesEditor/Public/GameplayAbilityAudit.h
+++ b/Engine/Plugins/Runtime/GameplayAbilities/Source/GameplayAbilitiesEditor/Public/GameplayAbilityAudit.h
@@ -6,6 +6,8 @@
 #include "Abilities/GameplayAbility.h"
 #include "GameplayAbilityAudit.generated.h"
 
+class UEdGraph;
+
 /**
  * This file has the implementation of the Gameplay Ability Audit code.  You can right-click on assets and select "Audit" which will produce a DataTable with useful information.
  * The trick to do a full audit is to select all of the GameplayAbility Blueprints in the Content Browser (use filter 'NativeParentClass=GameplayAbility'), then right-click and select the Audit Context Menu Item
