Non-Unity Compile fix

diff --git a/Engine/Plugins/Runtime/GameplayAbilities/Source/GameplayAbilitiesEditor/Public/GameplayAbilityAudit.h b/Engine/Plugins/Runtime/GameplayAbilities/Source/GameplayAbilitiesEditor/Public/GameplayAbilityAudit.h
--- a/Engine/Plugins/Runtime/GameplayAbilities/Source/GameplayAbilitiesEditor/Public/GameplayAbilityAudit.h
+++ b/Engine/Plugins/Runtime/GameplayAbilities/Source/GameplayAbilitiesEditor/Public/GameplayAbilityAudit.h
@@ -6,6 +6,7 @@
 #include "Abilities/GameplayAbility.h"
 #include "GameplayAbilityAudit.generated.h"
 
+class UBlueprint;
 class UEdGraph;
 
 /**
@@ -45,7 +46,7 @@ struct GAMEPLAYABILITIESEDITOR_API FGameplayAbilityAuditRow : public FTableRowBa
 	GENERATED_BODY()
 	
 	/** Fill this structure with data from a UBlueprint (GameplayAbility asset).  Return true if successful. */
-	virtual void FillDataFromGameplayAbilityBlueprint(const class UBlueprint& GameplayAbilityBlueprint);
+	virtual void FillDataFromGameplayAbilityBlueprint(const UBlueprint& GameplayAbilityBlueprint);
 
 	/** Fill this structure with data from a UGameplayAbility (likely compiled from a UBlueprint).  Return true if successful. */
 	virtual void FillDataFromGameplayAbility(const UGameplayAbility& GameplayAbility);
