Removed cpu scopes in GameplayAbilitiesModule static functions.

diff --git a/Engine/Plugins/Runtime/GameplayAbilities/Source/GameplayAbilities/Private/GameplayAbilitiesModule.cpp b/Engine/Plugins/Runtime/GameplayAbilities/Source/GameplayAbilities/Private/GameplayAbilitiesModule.cpp
--- a/Engine/Plugins/Runtime/GameplayAbilities/Source/GameplayAbilities/Private/GameplayAbilitiesModule.cpp
+++ b/Engine/Plugins/Runtime/GameplayAbilities/Source/GameplayAbilities/Private/GameplayAbilitiesModule.cpp
@@ -22,10 +22,10 @@ class FGameplayAbilitiesModule : public IGameplayAbilitiesModule
 
 	virtual UAbilitySystemGlobals* GetAbilitySystemGlobals() override
 	{
-		QUICK_SCOPE_CYCLE_COUNTER(STAT_IGameplayAbilitiesModule_GetAbilitySystemGlobals);
 		// Defer loading of globals to the first time it is requested
 		if (!AbilitySystemGlobals)
 		{
+			QUICK_SCOPE_CYCLE_COUNTER(STAT_IGameplayAbilitiesModule_GetAbilitySystemGlobals_LoadModule);
 			FSoftClassPath AbilitySystemClassName = (UAbilitySystemGlobals::StaticClass()->GetDefaultObject<UAbilitySystemGlobals>())->AbilitySystemGlobalsClassName;
 
 			UClass* SingletonClass = AbilitySystemClassName.TryLoadClass<UObject>();
@@ -42,7 +42,6 @@ class FGameplayAbilitiesModule : public IGameplayAbilitiesModule
 
 	virtual bool IsAbilitySystemGlobalsAvailable() override
 	{
-		QUICK_SCOPE_CYCLE_COUNTER(STAT_IGameplayAbilitiesModule_IsAbilitySystemGlobalsAvailable);
 		return AbilitySystemGlobals != nullptr;
 	}
 
diff --git a/Engine/Plugins/Runtime/GameplayAbilities/Source/GameplayAbilities/Public/GameplayAbilitiesModule.h b/Engine/Plugins/Runtime/GameplayAbilities/Source/GameplayAbilities/Public/GameplayAbilitiesModule.h
--- a/Engine/Plugins/Runtime/GameplayAbilities/Source/GameplayAbilities/Public/GameplayAbilitiesModule.h
+++ b/Engine/Plugins/Runtime/GameplayAbilities/Source/GameplayAbilities/Public/GameplayAbilitiesModule.h
@@ -23,7 +23,6 @@ public:
 	 */
 	static inline IGameplayAbilitiesModule& Get()
 	{
-		QUICK_SCOPE_CYCLE_COUNTER(STAT_IGameplayAbilitiesModule_Get);
 		static IGameplayAbilitiesModule& Singleton = FModuleManager::LoadModuleChecked< IGameplayAbilitiesModule >("GameplayAbilities");
 		return Singleton;
 	}
@@ -35,7 +34,6 @@ public:
 	 */
 	static inline bool IsAvailable()
 	{
-		QUICK_SCOPE_CYCLE_COUNTER(STAT_IGameplayAbilitiesModule_IsAvailable);
 		return FModuleManager::Get().IsModuleLoaded( "GameplayAbilities" );
 	}
 
