Fixed an issue where Gameplay Cue Notify Actors executing on Child Actors could leak memory when gc.PendingKill is disabled.

diff --git a/Engine/Plugins/Runtime/GameplayAbilities/Source/GameplayAbilities/Private/GameplayCueManager.cpp b/Engine/Plugins/Runtime/GameplayAbilities/Source/GameplayAbilities/Private/GameplayCueManager.cpp
--- a/Engine/Plugins/Runtime/GameplayAbilities/Source/GameplayAbilities/Private/GameplayCueManager.cpp
+++ b/Engine/Plugins/Runtime/GameplayAbilities/Source/GameplayAbilities/Private/GameplayCueManager.cpp
@@ -69,8 +69,7 @@ UGameplayCueManager::UGameplayCueManager(const FObjectInitializer& PCIP)
 
 void UGameplayCueManager::OnCreated()
 {
-	FWorldDelegates::OnWorldCleanup.AddUObject(this, &UGameplayCueManager::OnWorldCleanup);
-	FWorldDelegates::OnPreWorldFinishDestroy.AddUObject(this, &UGameplayCueManager::OnWorldCleanup, true, true);
+	FWorldDelegates::OnPostWorldCleanup.AddUObject(this, &UGameplayCueManager::OnPostWorldCleanup);
 	FNetworkReplayDelegates::OnPreScrub.AddUObject(this, &UGameplayCueManager::OnPreReplayScrub);
 		
 #if WITH_EDITOR
@@ -548,7 +547,7 @@ void UGameplayCueManager::NotifyGameplayCueActorFinished(AGameplayCueNotify_Acto
 				// This could happen if a GC notify actor has a delayed removal and another GC event happens before the delayed removal happens (the old GC actor could replace the latest one in the map)
 				if (WeakPtrPtr->Get() == Actor)
 				{
-					WeakPtrPtr->Reset();
+					NotifyMapActor.Remove(Actor->NotifyKey);
 				}
 			}
 
@@ -1607,43 +1606,75 @@ FPreallocationInfo& UGameplayCueManager::GetPreallocationInfo(UWorld* World)
 	return PreallocationInfoList_Internal.Last();
 }
 
-void UGameplayCueManager::OnWorldCleanup(UWorld* World, bool bSessionEnded, bool bCleanupResources)
+void UGameplayCueManager::OnPostWorldCleanup(UWorld* World, bool bSessionEnded, bool bCleanupResources)
 {
-#if !(UE_BUILD_SHIPPING || UE_BUILD_TEST)
-	DumpPreallocationStats(World);
-#endif
+	const FObjectKey WorldObjectKey(World);
 
 	for (int32 idx=0; idx < PreallocationInfoList_Internal.Num(); ++idx)
 	{
-		if (PreallocationInfoList_Internal[idx].OwningWorldKey == FObjectKey(World))
+		FPreallocationInfo& PreallocationInfo = PreallocationInfoList_Internal[idx];
+		if (PreallocationInfo.OwningWorldKey != WorldObjectKey)
 		{
-			ABILITY_LOG(Verbose, TEXT("UGameplayCueManager::OnWorldCleanup Removing PreallocationInfoList_Internal element %d"), idx);
-			PreallocationInfoList_Internal.RemoveAtSwap(idx, 1, false);
-			idx--;
+			continue;
 		}
+		
+		ABILITY_LOG(Verbose, TEXT("UGameplayCueManager::OnPostWorldCleanup %s Removing PreallocationInfoList_Internal element %d"), *GetNameSafe(World), idx);
+
+		// Spit out some debug information to help us track down memory issues
+		constexpr bool bWarnOnActiveActors = true;
+		DumpPreallocationStats(PreallocationInfo, bWarnOnActiveActors);
+
+		// Actually remove the entry which can contain hard references
+		PreallocationInfoList_Internal.RemoveAtSwap(idx, 1, false);
+		idx--;
 	}
 
 	IGameplayCueInterface::ClearTagToFunctionMap();
+
+	// Let's help debug any slow memory leaks in the NotifyMapActor
+	for (auto It = NotifyMapActor.CreateIterator(); It; ++It)
+	{
+		if (!It.Value().IsValid())
+		{
+			const UObject* CueClass = It.Key().CueClass.ResolveObjectPtrEvenIfUnreachable();
+			const UObject* TargetActor = It.Key().TargetActor.ResolveObjectPtrEvenIfUnreachable();
+
+			UE_LOG(LogAbilitySystem, Error, TEXT("GameplayCueManager contained invalid entry in NotifyMapActor. CueClass '%s'. TargetActor '%s'.  This will result in a slow memory leak."), *GetNameSafe(CueClass), *GetNameSafe(TargetActor));
+			It.RemoveCurrent();
+		}
+	}
 }
 
-void UGameplayCueManager::DumpPreallocationStats(UWorld* World)
+void UGameplayCueManager::DumpPreallocationStats(const FPreallocationInfo& PreallocationInfo, bool bWarnOnActiveActors)
 {
-	if (World == nullptr)
+	constexpr bool bAllowLoggingInBuild = !(UE_BUILD_SHIPPING || UE_BUILD_TEST);
+	if (!bAllowLoggingInBuild || !UE_LOG_ACTIVE(LogAbilitySystem, Display))
 	{
 		return;
 	}
 
-	FPreallocationInfo& Info = GetPreallocationInfo(World);
-	for (auto &It : Info.PreallocatedInstances)
+	for (auto& It : PreallocationInfo.PreallocatedInstances)
 	{
-		if (UClass* ThisClass = It.Key)
+		if (const UClass* ThisClass = It.Key)
 		{
-			if (AGameplayCueNotify_Actor* CDO = ThisClass->GetDefaultObject<AGameplayCueNotify_Actor>())
+			if (const AGameplayCueNotify_Actor* CDO = ThisClass->GetDefaultObject<AGameplayCueNotify_Actor>())
 			{
-				TArray<AGameplayCueNotify_Actor*>& List = It.Value.Actors;
+				const TArray<AGameplayCueNotify_Actor*>& List = It.Value.Actors;
 				if (List.Num() > CDO->NumPreallocatedInstances)
 				{
-					ABILITY_LOG(Display, TEXT("Notify class: %s was used simultaneously %d times. The CDO default is %d preallocated instanced."), *ThisClass->GetName(), List.Num(),  CDO->NumPreallocatedInstances); 
+					ABILITY_LOG(Display, TEXT("  GameplayCueNotify Class '%s' was used simultaneously %d times. The CDO default is only %d preallocated instances."), *ThisClass->GetName(), List.Num(), CDO->NumPreallocatedInstances);
+				}
+
+				if (bWarnOnActiveActors)
+				{
+					int StillActive = 0;
+					for (const AGameplayCueNotify_Actor* Actor : List)
+					{
+						StillActive += (!Actor->bInRecycleQueue);
+					}
+
+					// NotifyGameplayCueActorFinished should have been called on all of these Actors by the time we're done tearing down the world.
+					UE_CLOG(StillActive > 0, LogAbilitySystem, Error, TEXT("  GameplayCueNotify Class '%s' had %d instances still active.  This shouldn't happen."), *ThisClass->GetName(), StillActive);
 				}
 			}
 		}
diff --git a/Engine/Plugins/Runtime/GameplayAbilities/Source/GameplayAbilities/Public/GameplayCueManager.h b/Engine/Plugins/Runtime/GameplayAbilities/Source/GameplayAbilities/Public/GameplayCueManager.h
--- a/Engine/Plugins/Runtime/GameplayAbilities/Source/GameplayAbilities/Public/GameplayCueManager.h
+++ b/Engine/Plugins/Runtime/GameplayAbilities/Source/GameplayAbilities/Public/GameplayCueManager.h
@@ -194,12 +194,12 @@ class GAMEPLAYABILITIES_API UGameplayCueManager : public UDataAsset
 	/** Prespawns a single actor for gameplaycue notify actor classes that need prespawning (should be called by outside gamecode, such as gamestate) */
 	void UpdatePreallocation(UWorld* World);
 
-	void OnWorldCleanup(UWorld* World, bool bSessionEnded, bool bCleanupResources);
+	void OnPostWorldCleanup(UWorld* World, bool bSessionEnded, bool bCleanupResources);
 
 	void OnPreReplayScrub(UWorld* World);
 
-	/** Prints what classess exceeded their preallocation sizes during runtime */
-	void DumpPreallocationStats(UWorld* World);
+	/** Prints what classes exceeded their preallocation sizes during runtime */
+	void DumpPreallocationStats(const FPreallocationInfo& PreallocationInfo, bool bWarnOnActiveActors = false);
 
 	// -------------------------------------------------------------
 	//  Loading GameplayCueNotifies from ObjectLibraries
