AbilitySystemComponent;

Checking the Spec's ability pointer before use.

-- Code --
AbilitySystemComponent_Abilities:
FindAbilitySpecFromClass() - Checking the Spec's ability pointer.

#tests PIE, got on a motocycle and cycled weapons, none of the abilities were nullptr but this will fix it if there were

diff --git a/Engine/Plugins/Runtime/GameplayAbilities/Source/GameplayAbilities/Private/AbilitySystemComponent_Abilities.cpp b/Engine/Plugins/Runtime/GameplayAbilities/Source/GameplayAbilities/Private/AbilitySystemComponent_Abilities.cpp
--- a/Engine/Plugins/Runtime/GameplayAbilities/Source/GameplayAbilities/Private/AbilitySystemComponent_Abilities.cpp
+++ b/Engine/Plugins/Runtime/GameplayAbilities/Source/GameplayAbilities/Private/AbilitySystemComponent_Abilities.cpp
@@ -777,6 +777,11 @@ FGameplayAbilitySpec* UAbilitySystemComponent::FindAbilitySpecFromClass(TSubclas
 
 	for (FGameplayAbilitySpec& Spec : ActivatableAbilities.Items)
 	{
+		if (Spec.Ability == nullptr)
+		{
+			continue;
+		}
+
 		if (Spec.Ability->GetClass() == InAbilityClass)
 		{
 			return &Spec;
