Fixing an issue in GameplayCueManager where we would assume we were keeping track of which GameplayCueNotify_Actors were executing on which TargetActor, but they could actually be lost due to hash collisions.  It turns out we don't need to keep track, as the GameplayCueNotify_Actors will be in the TargetActor's Children list.

diff --git a/Engine/Plugins/Runtime/GameplayAbilities/Source/GameplayAbilities/Private/GameplayCueManager.cpp b/Engine/Plugins/Runtime/GameplayAbilities/Source/GameplayAbilities/Private/GameplayCueManager.cpp
--- a/Engine/Plugins/Runtime/GameplayAbilities/Source/GameplayAbilities/Private/GameplayCueManager.cpp
+++ b/Engine/Plugins/Runtime/GameplayAbilities/Source/GameplayAbilities/Private/GameplayCueManager.cpp
@@ -273,17 +273,14 @@ void UGameplayCueManager::ExecuteGameplayCue_NonReplicated(AActor* Target, const
 
 void UGameplayCueManager::EndGameplayCuesFor(AActor* TargetActor)
 {
-	for (auto It = NotifyMapActor.CreateIterator(); It; ++It)
+	// Make a copy so that OnOwnerDestroyed can remove itself (if it so chooses)
+	TArray<TObjectPtr<AActor>> Children = TargetActor->Children;
+	for (AActor* Child : Children)
 	{
-		FGCNotifyActorKey& Key = It.Key();
-		if (Key.TargetActor == TargetActor)
+		AGameplayCueNotify_Actor* NotifyActor = Cast<AGameplayCueNotify_Actor>(Child);
+		if (NotifyActor)
 		{
-			AGameplayCueNotify_Actor* InstancedCue = It.Value().Get();
-			if (InstancedCue)
-			{
-				InstancedCue->OnOwnerDestroyed(TargetActor);
-			}
-			It.RemoveCurrent();
+			NotifyActor->OnOwnerDestroyed(TargetActor);
 		}
 	}
 }
@@ -373,155 +370,151 @@ void UGameplayCueManager::OnMissingCueAsyncLoadComplete(FSoftObjectPath LoadedOb
 	}
 }
 
-AGameplayCueNotify_Actor* UGameplayCueManager::GetInstancedCueActor(AActor* TargetActor, UClass* CueClass, const FGameplayCueParameters& Parameters)
+AGameplayCueNotify_Actor* UGameplayCueManager::FindExistingCueOnActor(const AActor& TargetActor, const TSubclassOf<AGameplayCueNotify_Actor>& CueClass, const FGameplayCueParameters& Parameters) const
 {
-	QUICK_SCOPE_CYCLE_COUNTER(STAT_GameplayCueManager_GetInstancedCueActor);
-
-
-	// First, see if this actor already have a GameplayCueNotifyActor already going for this CueClass
-	AGameplayCueNotify_Actor* CDO = Cast<AGameplayCueNotify_Actor>(CueClass->ClassDefaultObject);
-	FGCNotifyActorKey	NotifyKey(TargetActor, CueClass, 
-							CDO->bUniqueInstancePerInstigator ? Parameters.GetInstigator() : nullptr, 
-							CDO->bUniqueInstancePerSourceObject ? Parameters.GetSourceObject() : nullptr);
-
-	AGameplayCueNotify_Actor* SpawnedCue = nullptr;
-	if (TWeakObjectPtr<AGameplayCueNotify_Actor>* WeakPtrPtr = NotifyMapActor.Find(NotifyKey))
-	{		
-		SpawnedCue = WeakPtrPtr->Get();
-		// If the cue is scheduled to be destroyed, don't reuse it, create a new one instead
-		if (SpawnedCue && SpawnedCue->GameplayCuePendingRemove() == false)
+	for (AActor* Child : TargetActor.Children)
+	{
+		if (Child->IsA(CueClass))
 		{
-			if (SpawnedCue->GetOwner() != TargetActor)
-			{
-#if WITH_EDITOR	
-				if (TargetActor && TargetActor->HasAnyFlags(RF_ClassDefaultObject))
-				{
-					// Animation preview hack, reuse this one even though the owner doesnt match the CDO
-					return SpawnedCue;
-				}
-#endif
+			AGameplayCueNotify_Actor* ChildNotify = CastChecked<AGameplayCueNotify_Actor>(Child);
 
-				// This should not happen. This means we think we can recycle and GC actor that is currently being used by someone else.
-				ABILITY_LOG(Warning, TEXT("GetInstancedCueActor attempting to reuse GC Actor with a different owner! %s (Target: %s). Using GC Actor: %s. Current Owner: %s"), *GetNameSafe(CueClass), *GetNameSafe(TargetActor), *GetNameSafe(SpawnedCue), *GetNameSafe(SpawnedCue->GetOwner()));
+			// Somehow the LifeSpan can end up being zero, meaning we're about to be destroyed (so don't reuse)
+			if (ChildNotify->GameplayCuePendingRemove())
+			{
+				UE_LOG(LogAbilitySystem, Verbose, TEXT("FindExistingCueActor considered %s, but it was pending remove"), *GetNameSafe(ChildNotify));
+				continue;
 			}
-			else
+
+			const bool bInstigatorMatches = !ChildNotify->bUniqueInstancePerInstigator || ChildNotify->CueInstigator == Parameters.GetInstigator();
+			const bool bSourceMatches = !ChildNotify->bUniqueInstancePerSourceObject || ChildNotify->CueSourceObject == Parameters.GetSourceObject();
+			if (bInstigatorMatches && bSourceMatches)
 			{
-				UE_CLOG((GameplayCueActorRecycleDebug>0), LogAbilitySystem, Display, TEXT("::GetInstancedCueActor Using Existing %s (Target: %s). Using GC Actor: %s"), *GetNameSafe(CueClass), *GetNameSafe(TargetActor), *GetNameSafe(SpawnedCue));
-				return SpawnedCue;
+				return ChildNotify;
 			}
 		}
-
-		// We aren't going to use this existing cue notify actor, so clear it.
-		SpawnedCue = nullptr;
 	}
 
-	UWorld* World = GetWorld();
+	return nullptr;
+}
 
-	// We don't have an instance for this, and we need one, so make one
-	if (ensure(TargetActor) && ensure(CueClass) && ensure(World))
+AGameplayCueNotify_Actor* UGameplayCueManager::GetInstancedCueActor(AActor* TargetActor, UClass* GameplayCueNotifyActorClass, const FGameplayCueParameters& Parameters)
+{
+	QUICK_SCOPE_CYCLE_COUNTER(STAT_GameplayCueManager_GetInstancedCueActor);
+
+	const TSubclassOf<AGameplayCueNotify_Actor> CueClass = GameplayCueNotifyActorClass;
+	if (!ensure(TargetActor) || !ensure(CueClass))
 	{
-		AActor* NewOwnerActor = TargetActor;
-		bool UseActorRecycling = (GameplayCueActorRecycle > 0);
-		
-#if WITH_EDITOR	
-		// Animtion preview hack. If we are trying to play the GC on a CDO, then don't use actor recycling and don't set the owner (to the CDO, which would cause problems)
-		if (TargetActor->HasAnyFlags(RF_ClassDefaultObject))
-		{
-			NewOwnerActor = nullptr;
-			UseActorRecycling = false;
-		}
-#endif
-		// Look to reuse an existing one that is stored on the CDO:
-		if (UseActorRecycling)
-		{
-			FPreallocationInfo& Info = GetPreallocationInfo(World);
-			FGameplayCueNotifyActorArray* PreallocatedList = Info.PreallocatedInstances.Find(CueClass);
-			if (PreallocatedList && PreallocatedList->Actors.Num() > 0)
-			{
-				SpawnedCue = nullptr;
-				while (true)
-				{
-					SpawnedCue = PreallocatedList->Actors.Pop(false);
+		UE_LOG(LogAbilitySystem, Error, TEXT("GetInstancedCueActor called with invalid parameters (TargetActor = %s, CueClass = %s)"), *GetNameSafe(TargetActor), *GetNameSafe(GameplayCueNotifyActorClass));
+		return nullptr;
+	}
 
-					// Temp: tracking down possible memory corruption
-					// null is maybe ok. But invalid low level is bad and we want to crash hard to find out who/why.
-					if (SpawnedCue && (SpawnedCue->IsValidLowLevelFast() == false))
-					{
-						checkf(false, TEXT("UGameplayCueManager::GetInstancedCueActor found an invalid SpawnedCue for class %s"), *GetNameSafe(CueClass));
-					}
+	// There used to be special code here to handle the case where the TargetActor was a CDO.  I'm not sure why that would be (or how that's even possible -- perhaps in the default Blueprint Viewport? But I can't trigger it.)
+	// Let's log it in case a user comes across this issue.
+	//	Animtion preview hack. If we are trying to play the GC on a CDO, then don't use actor recycling and don't set the owner (to the CDO, which would cause problems)
+	//	And we'll try to manually find (and reuse) the existing instance on the TargetActor.
+	UE_CLOG(WITH_EDITOR && TargetActor->HasAnyFlags(RF_ClassDefaultObject), LogAbilitySystem, Warning, TEXT("Adding %s to CDO %s. This used to be explicitly disallowed."), *GetNameSafe(CueClass), *GetNameSafe(TargetActor));
 
-					// Normal check: if cue was destroyed or is pending kill, then don't use it.
-					if (IsValid(SpawnedCue))
-					{
-						break;
-					}
-					
-					// outside of replays, this should not happen. GC Notifies should not be actually destroyed.
-					checkf(World->IsPlayingReplay(), TEXT("Spawned Cue is pending kill, garbage or null: %s."), *GetNameSafe(SpawnedCue));
+	UWorld* World = TargetActor->GetWorld();
+	if (!World)
+	{
+		UE_LOG(LogAbilitySystem, Warning, TEXT("GetInstancedCueActor called on TargetActor %s which did not belong to a world (is it a CDO or being destroyed?)"), *GetNameSafe(TargetActor));
+		return nullptr;
+	}
+	UE_CLOG(CurrentWorld != World, LogAbilitySystem, Error, TEXT("GetInstancedCueActor had CurrentWorld set to %s but TargetActor is in World %s"), *GetNameSafe(CurrentWorld), *GetNameSafe(World));
 
-					if (PreallocatedList->Actors.Num() <= 0)
-					{
-						// Ran out of preallocated instances... break and create a new one.
-						break;
-					}
-				}
+	// We found the exact Cue we're looking for already on the TargetActor.  Use that.
+	AGameplayCueNotify_Actor* ExistingCueOnActor = FindExistingCueOnActor(*TargetActor, CueClass, Parameters);
+	if (ExistingCueOnActor)
+	{
+		ExistingCueOnActor->CueInstigator = Parameters.GetInstigator();
+		ExistingCueOnActor->CueSourceObject = Parameters.GetSourceObject();
+		return ExistingCueOnActor;
+	}
 
-				if (SpawnedCue)
-				{
-					SpawnedCue->bInRecycleQueue = false;
-					SpawnedCue->SetOwner(NewOwnerActor);
-					SpawnedCue->SetActorLocationAndRotation(TargetActor->GetActorLocation(), TargetActor->GetActorRotation());
-					SpawnedCue->ReuseAfterRecycle();
-				}
+	const bool bUseActorRecycling = (GameplayCueActorRecycle > 0);
+	if (bUseActorRecycling)
+	{
+		if (AGameplayCueNotify_Actor* RecycledCue = FindRecycledCue(CueClass, *World))
+		{
+			RecycledCue->bInRecycleQueue = false;
+			RecycledCue->SetOwner(TargetActor);
+			RecycledCue->SetActorLocationAndRotation(TargetActor->GetActorLocation(), TargetActor->GetActorRotation());
+			RecycledCue->ReuseAfterRecycle();
+			RecycledCue->CueInstigator = Parameters.GetInstigator();
+			RecycledCue->CueSourceObject = Parameters.GetSourceObject();
+
+			UE_CLOG((GameplayCueActorRecycleDebug > 0), LogAbilitySystem, Display, TEXT("GetInstancedCueActor reusing Recycled CueActor: %s"), *GetNameSafe(RecycledCue));
 
-				UE_CLOG((GameplayCueActorRecycleDebug>0), LogAbilitySystem, Display, TEXT("GetInstancedCueActor Popping Recycled %s (Target: %s). Using GC Actor: %s"), *GetNameSafe(CueClass), *GetNameSafe(TargetActor), *GetNameSafe(SpawnedCue));
 #if WITH_EDITOR
-				// let things know that we 'spawned'
-				ISequenceRecorder& SequenceRecorder	= FModuleManager::LoadModuleChecked<ISequenceRecorder>("SequenceRecorder");
-				SequenceRecorder.NotifyActorStartRecording(SpawnedCue);
+			// let things know that we 'spawned'
+			ISequenceRecorder& SequenceRecorder = FModuleManager::LoadModuleChecked<ISequenceRecorder>("SequenceRecorder");
+			SequenceRecorder.NotifyActorStartRecording(RecycledCue);
 #endif
-			}
+
+			return RecycledCue;
 		}
+	}
 
-		// If we can't reuse, then spawn a new one
-		if (SpawnedCue == nullptr)
-		{
-			FActorSpawnParameters SpawnParams;
-			SpawnParams.Owner = NewOwnerActor;
-			if (LogGameplayCueActorSpawning)
-			{
-				ABILITY_LOG(Warning, TEXT("Spawning GameplaycueActor: %s"), *CueClass->GetName());
-			}
+	// If we can't reuse, then spawn a new one. Since TargetActor is the Owner, a reference to this CueNotify Actor will live in TargetActor::Children.
+	FActorSpawnParameters SpawnParams;
+	SpawnParams.Owner = TargetActor;
+	SpawnParams.OverrideLevel = World->PersistentLevel;
+	AGameplayCueNotify_Actor* SpawnedCue = World->SpawnActor<AGameplayCueNotify_Actor>(CueClass, TargetActor->GetActorLocation(), TargetActor->GetActorRotation(), SpawnParams);
+	SpawnedCue->CueInstigator = Parameters.GetInstigator();
+	SpawnedCue->CueSourceObject = Parameters.GetSourceObject();
 
-			SpawnParams.OverrideLevel = World->PersistentLevel;
-			SpawnedCue = World->SpawnActor<AGameplayCueNotify_Actor>(CueClass, TargetActor->GetActorLocation(), TargetActor->GetActorRotation(), SpawnParams);
-		}
+	UE_CLOG(LogGameplayCueActorSpawning > 0, LogAbilitySystem, Warning, TEXT("Spawned Gameplay Cue Notify Actor: %s (instance %s)"), *CueClass->GetName(), *GetNameSafe(SpawnedCue));
 
-		// Associate this GameplayCueNotifyActor with this target actor/key
-		if (ensure(SpawnedCue))
+	return SpawnedCue;
+
+}
+
+AGameplayCueNotify_Actor* UGameplayCueManager::FindRecycledCue(const TSubclassOf<AGameplayCueNotify_Actor>& CueClass, const UWorld& FindInWorld)
+{
+	FPreallocationInfo& Info = GetPreallocationInfo(&FindInWorld);
+	FGameplayCueNotifyActorArray* PreallocatedList = Info.PreallocatedInstances.Find(CueClass);
+	if (!PreallocatedList)
+	{
+		// No preallocated instances yet
+		return nullptr;
+	}
+
+	while (PreallocatedList->Actors.Num() > 0)
+	{
+		AGameplayCueNotify_Actor* RecycledCue = PreallocatedList->Actors.Pop(false);
+
+		// Normal check: if cue was destroyed or is pending kill, then don't use it.
+		if (IsValid(RecycledCue))
 		{
-			SpawnedCue->NotifyKey = NotifyKey;
-			NotifyMapActor.Add(NotifyKey, SpawnedCue);
+			return RecycledCue;
 		}
+					
+		// outside of replays, this should not happen. GC Notifies should not be actually destroyed.
+		ensureMsgf(FindInWorld.IsPlayingReplay(), TEXT("RecycledCue is pending kill, garbage or null: %s."), *GetNameSafe(RecycledCue));
 	}
 
-	UE_CLOG((GameplayCueActorRecycleDebug>0), LogAbilitySystem, Display, TEXT("GetInstancedCueActor  Returning %s (Target: %s). Using GC Actor: %s"), *GetNameSafe(CueClass), *GetNameSafe(TargetActor), *GetNameSafe(SpawnedCue));
-	return SpawnedCue;
+	return nullptr;
 }
 
 void UGameplayCueManager::NotifyGameplayCueActorFinished(AGameplayCueNotify_Actor* Actor)
 {
-	bool UseActorRecycling = (GameplayCueActorRecycle > 0);
+	if (!IsValid(Actor))
+	{
+		ensureMsgf(GetWorld() && GetWorld()->IsPlayingReplay(), TEXT("GameplayCueNotify %s is pending kill or garbage in ::NotifyGameplayCueActorFinished (and not in network demo)"), *GetNameSafe(Actor));
+		return;
+	}
+
+	bool bUseActorRecycling = (GameplayCueActorRecycle > 0);
 
 #if WITH_EDITOR	
 	// Don't recycle in preview worlds
 	if (Actor->GetWorld()->IsPreviewWorld())
 	{
-		UseActorRecycling = false;
+		bUseActorRecycling = false;
 	}
 #endif
 
-	if (UseActorRecycling)
+	if (bUseActorRecycling)
 	{
 		if (Actor->bInRecycleQueue)
 		{
@@ -530,28 +523,12 @@ void UGameplayCueManager::NotifyGameplayCueActorFinished(AGameplayCueNotify_Acto
 			return;
 		}
 
-		AGameplayCueNotify_Actor* CDO = Actor->GetClass()->GetDefaultObject<AGameplayCueNotify_Actor>();
-		if (CDO && Actor->Recycle())
+		UE_CLOG((GameplayCueActorRecycleDebug > 0), LogAbilitySystem, Warning, TEXT("Recycling CueActor %s"), *GetNameSafe(Actor));
+		
+		if (Actor->Recycle())
 		{
-			if (!IsValid(Actor))
-			{
-				ensureMsgf(GetWorld()->IsPlayingReplay(), TEXT("GameplayCueNotify %s is pending kill or garbage in ::NotifyGameplayCueActorFinished (and not in network demo)"), *GetNameSafe(Actor));
-				return;
-			}
-			Actor->bInRecycleQueue = true;
-
-			// Remove this now from our internal map so that it doesn't get reused like a currently active cue would
-			if (TWeakObjectPtr<AGameplayCueNotify_Actor>* WeakPtrPtr = NotifyMapActor.Find(Actor->NotifyKey))
-			{
-				// Only remove if this is the current actor in the map!
-				// This could happen if a GC notify actor has a delayed removal and another GC event happens before the delayed removal happens (the old GC actor could replace the latest one in the map)
-				if (WeakPtrPtr->Get() == Actor)
-				{
-					NotifyMapActor.Remove(Actor->NotifyKey);
-				}
-			}
 
-			UE_CLOG((GameplayCueActorRecycleDebug>0), LogAbilitySystem, Display, TEXT("NotifyGameplayCueActorFinished %s"), *GetNameSafe(Actor));
+			Actor->bInRecycleQueue = true;
 
 			FPreallocationInfo& Info = GetPreallocationInfo(Actor->GetWorld());
 			FGameplayCueNotifyActorArray& PreAllocatedList = Info.PreallocatedInstances.FindOrAdd(Actor->GetClass());
@@ -569,6 +546,8 @@ void UGameplayCueManager::NotifyGameplayCueActorFinished(AGameplayCueNotify_Acto
 #endif
 			return;
 		}
+
+		UE_CLOG((GameplayCueActorRecycleDebug > 0), LogAbilitySystem, Error, TEXT("Could not Recycle CueActor %s, Destroying"), *GetNameSafe(Actor));
 	}	
 
 	// We didn't recycle, so just destroy
@@ -1587,7 +1566,7 @@ void UGameplayCueManager::UpdatePreallocation(UWorld* World)
 	}
 }
 
-FPreallocationInfo& UGameplayCueManager::GetPreallocationInfo(UWorld* World)
+FPreallocationInfo& UGameplayCueManager::GetPreallocationInfo(const UWorld* World)
 {
 	FObjectKey ObjKey(World);
 
@@ -1630,19 +1609,6 @@ void UGameplayCueManager::OnPostWorldCleanup(UWorld* World, bool bSessionEnded,
 	}
 
 	IGameplayCueInterface::ClearTagToFunctionMap();
-
-	// Let's help debug any slow memory leaks in the NotifyMapActor
-	for (auto It = NotifyMapActor.CreateIterator(); It; ++It)
-	{
-		if (!It.Value().IsValid())
-		{
-			const UObject* CueClass = It.Key().CueClass.ResolveObjectPtrEvenIfUnreachable();
-			const UObject* TargetActor = It.Key().TargetActor.ResolveObjectPtrEvenIfUnreachable();
-
-			UE_LOG(LogAbilitySystem, Error, TEXT("GameplayCueManager contained invalid entry in NotifyMapActor. CueClass '%s'. TargetActor '%s'.  This will result in a slow memory leak."), *GetNameSafe(CueClass), *GetNameSafe(TargetActor));
-			It.RemoveCurrent();
-		}
-	}
 }
 
 void UGameplayCueManager::DumpPreallocationStats(const FPreallocationInfo& PreallocationInfo, bool bWarnOnActiveActors)
diff --git a/Engine/Plugins/Runtime/GameplayAbilities/Source/GameplayAbilities/Private/GameplayCueNotify_Actor.cpp b/Engine/Plugins/Runtime/GameplayAbilities/Source/GameplayAbilities/Private/GameplayCueNotify_Actor.cpp
--- a/Engine/Plugins/Runtime/GameplayAbilities/Source/GameplayAbilities/Private/GameplayCueNotify_Actor.cpp
+++ b/Engine/Plugins/Runtime/GameplayAbilities/Source/GameplayAbilities/Private/GameplayCueNotify_Actor.cpp
@@ -351,6 +351,9 @@ bool AGameplayCueNotify_Actor::Recycle()
 	bHasHandledOnActiveEvent = false;
 	bHasHandledWhileActiveEvent = false;
 	bHasHandledOnRemoveEvent = false;
+	CueInstigator = nullptr;
+	CueSourceObject = nullptr;
+
 	ClearOwnerDestroyedDelegate();
 	if (FinishTimerHandle.IsValid())
 	{
diff --git a/Engine/Plugins/Runtime/GameplayAbilities/Source/GameplayAbilities/Private/GameplayCueSet.cpp b/Engine/Plugins/Runtime/GameplayAbilities/Source/GameplayAbilities/Private/GameplayCueSet.cpp
--- a/Engine/Plugins/Runtime/GameplayAbilities/Source/GameplayAbilities/Private/GameplayCueSet.cpp
+++ b/Engine/Plugins/Runtime/GameplayAbilities/Source/GameplayAbilities/Private/GameplayCueSet.cpp
@@ -41,9 +41,11 @@ namespace GameplayCueDebug
 				}
 			})
 	);
-
 }
 
+// Let's use this to create a path for testing what happens when GameplayCues are failing to load (or are loading slowly due to IO contention)
+static TAutoConsoleVariable<bool> CVarGameplayCueFailLoads(TEXT("AbilitySystem.GameplayCueFailLoads"), false, TEXT("Pretend all GameplayCues are unloaded (and keep requesting loads) while true. Set to false to allow GameplayCues to play."), ECVF_Default);
+
 // --------------------------------------------------------------------------------------------------------------------------------------------------------
 //
 //	UGameplayCueSet
@@ -272,12 +274,14 @@ bool UGameplayCueSet::HandleGameplayCueNotify_Internal(AActor* TargetActor, int3
 
 		Parameters.MatchedTagName = CueData.GameplayCueTag;
 
+		const bool bDebugFailLoads = CVarGameplayCueFailLoads.GetValueOnGameThread();
+
 		// If object is not loaded yet
-		if (CueData.LoadedGameplayCueClass == nullptr)
+		if (CueData.LoadedGameplayCueClass == nullptr || bDebugFailLoads)
 		{
 			// See if the object is loaded but just not hooked up here
 			CueData.LoadedGameplayCueClass = Cast<UClass>(CueData.GameplayCueNotifyObj.ResolveObject());
-			if (CueData.LoadedGameplayCueClass == nullptr)
+			if (CueData.LoadedGameplayCueClass == nullptr || bDebugFailLoads)
 			{
 				if (!CueManager->HandleMissingGameplayCue(this, CueData, TargetActor, EventType, Parameters))
 				{
@@ -318,8 +322,10 @@ bool UGameplayCueSet::HandleGameplayCueNotify_Internal(AActor* TargetActor, int3
 			{
 				if (TargetActor)
 				{
+					TSubclassOf<AGameplayCueNotify_Actor> InstancedClass = InstancedCue->GetClass();
+
 					//Get our instance. We should probably have a flag or something to determine if we want to reuse or stack instances. That would mean changing our map to have a list of active instances.
-					AGameplayCueNotify_Actor* SpawnedInstancedCue = CueManager->GetInstancedCueActor(TargetActor, CueData.LoadedGameplayCueClass, Parameters);
+					AGameplayCueNotify_Actor* SpawnedInstancedCue = CueManager->GetInstancedCueActor(TargetActor, InstancedClass, Parameters);
 					if (ensure(SpawnedInstancedCue))
 					{
 						SpawnedInstancedCue->HandleGameplayCue(TargetActor, EventType, Parameters);
diff --git a/Engine/Plugins/Runtime/GameplayAbilities/Source/GameplayAbilities/Public/GameplayCueManager.h b/Engine/Plugins/Runtime/GameplayAbilities/Source/GameplayAbilities/Public/GameplayCueManager.h
--- a/Engine/Plugins/Runtime/GameplayAbilities/Source/GameplayAbilities/Public/GameplayCueManager.h
+++ b/Engine/Plugins/Runtime/GameplayAbilities/Source/GameplayAbilities/Public/GameplayCueManager.h
@@ -179,8 +179,8 @@ class GAMEPLAYABILITIES_API UGameplayCueManager : public UDataAsset
 	/** Force any instanced GameplayCueNotifies to stop */
 	virtual void EndGameplayCuesFor(AActor* TargetActor);
 
-	/** Returns the cached instance cue. Creates it if it doesn't exist */
-	virtual AGameplayCueNotify_Actor* GetInstancedCueActor(AActor* TargetActor, UClass* CueClass, const FGameplayCueParameters& Parameters);
+	/** Returns the cached instance cue. Creates it if it doesn't exist. */
+	virtual AGameplayCueNotify_Actor* GetInstancedCueActor(AActor* TargetActor, UClass* GameplayCueNotifyActorClass, const FGameplayCueParameters& Parameters);
 
 	/** Notify that this actor is finished and should be destroyed or recycled */
 	virtual void NotifyGameplayCueActorFinished(AGameplayCueNotify_Actor* Actor);
@@ -263,6 +263,12 @@ protected:
 	/** returns list of valid gameplay cue paths. Subclasses may override this to specify locations that aren't part of the "always loaded" LoadedPaths array */
 	virtual TArray<FString>	GetValidGameplayCuePaths() { return GetAlwaysLoadedGameplayCuePaths(); }
 
+	/** Given a TargetActor and a CueClass, find an existing instance of the CueNotify Actor that we can reuse */
+	AGameplayCueNotify_Actor* FindExistingCueOnActor(const AActor& TargetActor, const TSubclassOf<AGameplayCueNotify_Actor>& CueClass, const FGameplayCueParameters& Parameters) const;
+
+	/** Given a CueClass, find an already spawned (but currently unused) recycled instance that exists in FindInWorld so that we may reuse it. Note: This function also compacts the recycled instances, removing stale ones. */
+	AGameplayCueNotify_Actor* FindRecycledCue(const TSubclassOf<AGameplayCueNotify_Actor>& CueClass, const UWorld& FindInWorld);
+
 	UPROPERTY(transient)
 	FGameplayCueObjectLibrary RuntimeGameplayCueObjectLibrary;
 
@@ -301,8 +307,6 @@ public:
 
 	FStreamableManager	StreamableManager;
 	
-	TMap<FGCNotifyActorKey, TWeakObjectPtr<AGameplayCueNotify_Actor> > NotifyMapActor;
-
 	void PrintGameplayCueNotifyMap();
 
 	virtual void PrintLoadedGameplayCueNotifyClasses();
@@ -397,7 +401,7 @@ protected:
 	/** Cached world we are currently handling cues for. Used for non instanced GC Notifies that need world. */
 	static UWorld* CurrentWorld;
 
-	FPreallocationInfo& GetPreallocationInfo(UWorld* World);
+	FPreallocationInfo& GetPreallocationInfo(const UWorld* World);
 
 	UPROPERTY(transient)
 	TArray<FPreallocationInfo>	PreallocationInfoList_Internal;
diff --git a/Engine/Plugins/Runtime/GameplayAbilities/Source/GameplayAbilities/Public/GameplayCueNotify_Actor.h b/Engine/Plugins/Runtime/GameplayAbilities/Source/GameplayAbilities/Public/GameplayCueNotify_Actor.h
--- a/Engine/Plugins/Runtime/GameplayAbilities/Source/GameplayAbilities/Public/GameplayCueNotify_Actor.h
+++ b/Engine/Plugins/Runtime/GameplayAbilities/Source/GameplayAbilities/Public/GameplayCueNotify_Actor.h
@@ -145,10 +145,19 @@ public:
 	UPROPERTY(EditDefaultsOnly, Category = GameplayCue)
 	int32 NumPreallocatedInstances;
 
-	FGCNotifyActorKey NotifyKey;
+PRAGMA_DISABLE_DEPRECATION_WARNINGS
+	UE_DEPRECATED(5.3, "NotifyKey is deprecated and unused.  See comments in FGCNotifyActorKey")
+	FGCNotifyActorKey NotifyKey = {};
+PRAGMA_ENABLE_DEPRECATION_WARNINGS
 
 	// Set when the GC actor is in the recycle queue (E.g., not active in world. This is to prevent rentrancy in the recyle code since multiple paths can lead the GC actor there)
 	bool bInRecycleQueue;
+
+	// Keep track of the Instigator so we can decide if this GameplayCueNotify is one we're searching for.
+	TWeakObjectPtr<AActor> CueInstigator;
+
+	// Keep track of the SourceObject so we can decide if this GameplayCueNotify is one we're searching for.
+	TWeakObjectPtr<const UObject> CueSourceObject;
 	
 protected:
 	FTimerHandle FinishTimerHandle;
diff --git a/Engine/Plugins/Runtime/GameplayAbilities/Source/GameplayAbilities/Public/GameplayCue_Types.h b/Engine/Plugins/Runtime/GameplayAbilities/Source/GameplayAbilities/Public/GameplayCue_Types.h
--- a/Engine/Plugins/Runtime/GameplayAbilities/Source/GameplayAbilities/Public/GameplayCue_Types.h
+++ b/Engine/Plugins/Runtime/GameplayAbilities/Source/GameplayAbilities/Public/GameplayCue_Types.h
@@ -88,13 +88,19 @@ struct FPreallocationInfo
 
 
 /** Struct that is used by the gameplaycue manager to tie an instanced gameplaycue to the calling gamecode. Usually this is just the target actor, but can also be unique per instigator/sourceobject */
-struct FGCNotifyActorKey
+struct UE_DEPRECATED(5.3, "Look at FGCNotifyActorKey constructor notes for upgrade steps.") FGCNotifyActorKey
 {
 	FGCNotifyActorKey()
 	{
 
 	}
 
+	// This class is deprecated.  It was previously used to index into a map to find an instance of AGameplayCueNotify_Actor which was in use.
+	// Instead, you can find the AGameplayCueNotify_Actor this way:
+	//	1. InTargetActor will be its parent.  So you search InTargetActor->Children for instances of AGameplayCueNotify_Actor.
+	//	2. InCueClass can be inferred from AGameplayCueNotify_Actor->GetClass().
+	//	3. AGameplayCueNotify_Actor has two new fields: CueInstigator (corresponds to InInstigatorActor) and CueSourceObject; (corresponds to InSourceObj).
+	// Refer to UGameplayCueManager::GetInstancedCueActor.
 	FGCNotifyActorKey(AActor* InTargetActor, UClass* InCueClass, AActor* InInstigatorActor=nullptr, const UObject* InSourceObj=nullptr)
 	{
 		TargetActor = FObjectKey(InTargetActor);
@@ -117,6 +123,7 @@ struct FGCNotifyActorKey
 	}
 };
 
+PRAGMA_DISABLE_DEPRECATION_WARNINGS
 FORCEINLINE uint32 GetTypeHash(const FGCNotifyActorKey& Key)
 {
 	return GetTypeHash(Key.TargetActor)	^
@@ -124,6 +131,7 @@ FORCEINLINE uint32 GetTypeHash(const FGCNotifyActorKey& Key)
 			GetTypeHash(Key.OptionalSourceObject) ^
 			GetTypeHash(Key.CueClass);
 }
+PRAGMA_ENABLE_DEPRECATION_WARNINGS
 
 /**
  *	FScopedGameplayCueSendContext
