Fix crash when trying to apply gameplay cues after a seamless travel

diff --git a/Engine/Plugins/Runtime/GameplayAbilities/Source/GameplayAbilities/Private/GameplayCueManager.cpp b/Engine/Plugins/Runtime/GameplayAbilities/Source/GameplayAbilities/Private/GameplayCueManager.cpp
--- a/Engine/Plugins/Runtime/GameplayAbilities/Source/GameplayAbilities/Private/GameplayCueManager.cpp
+++ b/Engine/Plugins/Runtime/GameplayAbilities/Source/GameplayAbilities/Private/GameplayCueManager.cpp
@@ -374,7 +374,7 @@ AGameplayCueNotify_Actor* UGameplayCueManager::FindExistingCueOnActor(const AAct
 {
 	for (AActor* Child : TargetActor.Children)
 	{
-		if (Child->IsA(CueClass))
+		if (IsValid(Child) && Child->IsA(CueClass))
 		{
 			AGameplayCueNotify_Actor* ChildNotify = CastChecked<AGameplayCueNotify_Actor>(Child);
 
