Remove stale replication flag from gameplay tasks & ability system components

diff --git a/Engine/Plugins/Runtime/GameplayAbilities/Source/GameplayAbilities/Private/AbilitySystemComponent.cpp b/Engine/Plugins/Runtime/GameplayAbilities/Source/GameplayAbilities/Private/AbilitySystemComponent.cpp
--- a/Engine/Plugins/Runtime/GameplayAbilities/Source/GameplayAbilities/Private/AbilitySystemComponent.cpp
+++ b/Engine/Plugins/Runtime/GameplayAbilities/Source/GameplayAbilities/Private/AbilitySystemComponent.cpp
@@ -59,7 +59,6 @@ UAbilitySystemComponent::UAbilitySystemComponent(const FObjectInitializer& Objec
 	bSuppressGrantAbility = false;
 	bSuppressGameplayCues = false;
 	bPendingMontageRep = false;
-	bIsNetDirty = true;
 	AffectedAnimInstanceTag = NAME_None; 
 	AbilityLastActivatedTime = 0.f;
 
@@ -104,7 +103,6 @@ const UAttributeSet* UAbilitySystemComponent::GetOrCreateAttributeSubobject(TSub
 			UAttributeSet* Attributes = NewObject<UAttributeSet>(OwningActor, AttributeClass);
 			AddSpawnedAttribute(Attributes);
 			MyAttributes = Attributes;
-			bIsNetDirty = true;
 		}
 	}
 
@@ -331,7 +329,6 @@ void UAbilitySystemComponent::SetNumericAttributeBase(const FGameplayAttribute &
 {
 	// Go through our active gameplay effects container so that aggregation/mods are handled properly.
 	ActiveGameplayEffects.SetAttributeBaseValue(Attribute, NewFloatValue);
-	bIsNetDirty = true;
 }
 
 float UAbilitySystemComponent::GetNumericAttributeBase(const FGameplayAttribute &Attribute) const
@@ -349,7 +346,6 @@ void UAbilitySystemComponent::SetNumericAttribute_Internal(const FGameplayAttrib
 	// Set the attribute directly: update the FProperty on the attribute set.
 	const UAttributeSet* AttributeSet = GetAttributeSubobjectChecked(Attribute.GetAttributeSetClass());
 	Attribute.SetNumericValueChecked(NewFloatValue, const_cast<UAttributeSet*>(AttributeSet));
-	bIsNetDirty = true;
 }
 
 float UAbilitySystemComponent::GetNumericAttribute(const FGameplayAttribute &Attribute) const
@@ -386,14 +382,12 @@ void UAbilitySystemComponent::ApplyModToAttribute(const FGameplayAttribute &Attr
 	if (IsOwnerActorAuthoritative())
 	{
 		ActiveGameplayEffects.ApplyModToAttribute(Attribute, ModifierOp, ModifierMagnitude);
-		bIsNetDirty = true;
 	}
 }
 
 void UAbilitySystemComponent::ApplyModToAttributeUnsafe(const FGameplayAttribute &Attribute, TEnumAsByte<EGameplayModOp::Type> ModifierOp, float ModifierMagnitude)
 {
 	ActiveGameplayEffects.ApplyModToAttribute(Attribute, ModifierOp, ModifierMagnitude);
-	bIsNetDirty = true;
 }
 
 FGameplayEffectSpecHandle UAbilitySystemComponent::MakeOutgoingSpec(TSubclassOf<UGameplayEffect> GameplayEffectClass, float Level, FGameplayEffectContextHandle Context) const
@@ -663,7 +657,6 @@ FOnGameplayEffectTagCountChanged& UAbilitySystemComponent::RegisterGenericGamepl
 
 FOnGameplayAttributeChange& UAbilitySystemComponent::RegisterGameplayAttributeEvent(FGameplayAttribute Attribute)
 {
-	bIsNetDirty = true;
 PRAGMA_DISABLE_DEPRECATION_WARNINGS
 	return ActiveGameplayEffects.RegisterGameplayAttributeEvent(Attribute);
 PRAGMA_ENABLE_DEPRECATION_WARNINGS
@@ -671,7 +664,6 @@ PRAGMA_ENABLE_DEPRECATION_WARNINGS
 
 FOnGameplayAttributeValueChange& UAbilitySystemComponent::GetGameplayAttributeValueChangeDelegate(FGameplayAttribute Attribute)
 {
-	bIsNetDirty = true;
 	return ActiveGameplayEffects.GetGameplayAttributeValueChangeDelegate(Attribute);
 }
 
@@ -829,7 +821,6 @@ FActiveGameplayEffectHandle UAbilitySystemComponent::ApplyGameplayEffectSpecToSe
 			return FActiveGameplayEffectHandle();
 		}
 	}
-	bIsNetDirty = true;
 
 	// Clients should treat predicted instant effects as if they have infinite duration. The effects will be cleaned up later.
 	bool bTreatAsInfiniteDuration = GetOwnerRole() != ROLE_Authority && PredictionKey.IsLocalClientKey() && Spec.Def->DurationPolicy == EGameplayEffectDurationType::Instant;
@@ -1030,7 +1021,6 @@ void UAbilitySystemComponent::ExecuteGameplayEffect(FGameplayEffectSpec &Spec, F
 			ABILITY_VLOG(GetOwnerActor(), Log, TEXT("         %s: %s %f"), *Modifier.Attribute.GetName(), *EGameplayModOpToString(Modifier.ModifierOp), Magnitude);
 		}
 	}
-	bIsNetDirty = true;
 
 	ActiveGameplayEffects.ExecuteActiveEffectsFrom(Spec, PredictionKey);
 }
@@ -1053,7 +1043,6 @@ const UGameplayEffect* UAbilitySystemComponent::GetGameplayEffectDefForHandle(FA
 
 bool UAbilitySystemComponent::RemoveActiveGameplayEffect(FActiveGameplayEffectHandle Handle, int32 StacksToRemove)
 {
-	bIsNetDirty = true;
 	return ActiveGameplayEffects.RemoveActiveGameplayEffect(Handle, StacksToRemove);
 }
 
@@ -1082,7 +1071,6 @@ void UAbilitySystemComponent::RemoveActiveGameplayEffectBySourceEffect(TSubclass
 
 			return bMatches;
 		});
-		bIsNetDirty = true;
 		ActiveGameplayEffects.RemoveActiveEffects(Query, StacksToRemove);
 	}
 }
@@ -1322,7 +1310,6 @@ void UAbilitySystemComponent::AddGameplayCue_Internal(const FGameplayTag Gamepla
 	{
 		bool bWasInList = HasMatchingGameplayTag(GameplayCueTag);
 
-		bIsNetDirty = true;
 		ForceReplication();
 		GameplayCueContainer.AddCue(GameplayCueTag, ScopedPredictionKey, GameplayCueParameters);
 		
@@ -1391,7 +1378,6 @@ void UAbilitySystemComponent::RemoveGameplayCue_Internal(const FGameplayTag Game
 		bool bWasInList = HasMatchingGameplayTag(GameplayCueTag);
 
 		// Force replication so GameplayCue removals are properly replicated to all clients during Mixed and Minimal replication modes
-		bIsNetDirty = true;
 		ForceReplication();
 		GameplayCueContainer.RemoveCue(GameplayCueTag);
 
@@ -1413,7 +1399,6 @@ void UAbilitySystemComponent::RemoveGameplayCue_Internal(const FGameplayTag Game
 
 void UAbilitySystemComponent::RemoveAllGameplayCues()
 {
-	bIsNetDirty = true;
 	for (int32 i = (ActiveGameplayCues.GameplayCues.Num() - 1); i >= 0; --i)
 	{
 		RemoveGameplayCue(ActiveGameplayCues.GameplayCues[i].GameplayCueTag);
@@ -1604,7 +1589,6 @@ int32 UAbilitySystemComponent::RemoveActiveEffects(const FGameplayEffectQuery& Q
 {
 	if (IsOwnerActorAuthoritative())
 	{
-		bIsNetDirty = true;
 		return ActiveGameplayEffects.RemoveActiveEffects(Query, StacksToRemove);
 	}
 
@@ -1671,7 +1655,6 @@ void UAbilitySystemComponent::UpdateMinimalReplicationGameplayCuesCondition()
 
 void UAbilitySystemComponent::ForceReplication()
 {
-	bIsNetDirty = true;
 	AActor *OwningActor = GetOwner();
 	if (OwningActor)
 	{
diff --git a/Engine/Plugins/Runtime/GameplayAbilities/Source/GameplayAbilities/Private/AbilitySystemComponent_Abilities.cpp b/Engine/Plugins/Runtime/GameplayAbilities/Source/GameplayAbilities/Private/AbilitySystemComponent_Abilities.cpp
--- a/Engine/Plugins/Runtime/GameplayAbilities/Source/GameplayAbilities/Private/AbilitySystemComponent_Abilities.cpp
+++ b/Engine/Plugins/Runtime/GameplayAbilities/Source/GameplayAbilities/Private/AbilitySystemComponent_Abilities.cpp
@@ -74,7 +74,6 @@ void UAbilitySystemComponent::InitializeComponent()
 		if (Set)  
 		{
 			SpawnedAttributes.AddUnique(Set);
-			bIsNetDirty = true;
 		}
 	}
 
@@ -409,7 +408,6 @@ void UAbilitySystemComponent::ClearAllAbilities()
 
 	ActivatableAbilities.Items.Empty(ActivatableAbilities.Items.Num());
 	ActivatableAbilities.MarkArrayDirty();
-	bIsNetDirty = true;
 
 	CheckForClearedAbilities();
 }
@@ -437,7 +435,6 @@ void UAbilitySystemComponent::ClearAbility(const FGameplayAbilitySpecHandle& Han
 		return;
 	}
 
-	bIsNetDirty = true;
 	for (int Idx = 0; Idx < AbilityPendingAdds.Num(); ++Idx)
 	{
 		if (AbilityPendingAdds[Idx].Handle == Handle)
@@ -798,7 +795,6 @@ void UAbilitySystemComponent::MarkAbilitySpecDirty(FGameplayAbilitySpec& Spec, b
 		// Don't mark dirty for specs that are server only unless it was an add/remove
 		if (!(Spec.Ability && Spec.Ability->NetExecutionPolicy == EGameplayAbilityNetExecutionPolicy::ServerOnly && !WasAddOrRemove))
 		{
-			bIsNetDirty = true;
 			ActivatableAbilities.MarkItemDirty(Spec);
 		}
 		AbilitySpecDirtiedCallbacks.Broadcast(Spec);
@@ -1232,7 +1228,6 @@ void UAbilitySystemComponent::UnBlockAbilitiesWithTags(const FGameplayTagContain
 
 void UAbilitySystemComponent::BlockAbilityByInputID(int32 InputID)
 {
-	bIsNetDirty = true;
 	const TArray<uint8>& ConstBlockedAbilityBindings = GetBlockedAbilityBindings();
 	if (InputID >= 0 && InputID < ConstBlockedAbilityBindings.Num())
 	{
@@ -1242,7 +1237,6 @@ void UAbilitySystemComponent::BlockAbilityByInputID(int32 InputID)
 
 void UAbilitySystemComponent::UnBlockAbilityByInputID(int32 InputID)
 {
-	bIsNetDirty = true;
 	const TArray<uint8>& ConstBlockedAbilityBindings = GetBlockedAbilityBindings();
 	if (InputID >= 0 && InputID < ConstBlockedAbilityBindings.Num() && ConstBlockedAbilityBindings[InputID] > 0)
 	{
@@ -2478,7 +2472,6 @@ void UAbilitySystemComponent::BindAbilityActivationToInputComponent(UInputCompon
 void UAbilitySystemComponent::SetBlockAbilityBindingsArray(FGameplayAbilityInputBinds BindInfo)
 {
 	UEnum* EnumBinds = BindInfo.GetBindEnum();
-	bIsNetDirty = true;
 	GetBlockedAbilityBindings_Mutable().SetNumZeroed(EnumBinds->NumEnums());
 }
 
diff --git a/Engine/Plugins/Runtime/GameplayAbilities/Source/GameplayAbilities/Private/GameplayPrediction.cpp b/Engine/Plugins/Runtime/GameplayAbilities/Source/GameplayAbilities/Private/GameplayPrediction.cpp
--- a/Engine/Plugins/Runtime/GameplayAbilities/Source/GameplayAbilities/Private/GameplayPrediction.cpp
+++ b/Engine/Plugins/Runtime/GameplayAbilities/Source/GameplayAbilities/Private/GameplayPrediction.cpp
@@ -344,7 +344,6 @@ FScopedPredictionWindow::~FScopedPredictionWindow()
 			if (OwnerPtr->ScopedPredictionKey.IsValidKey())
 			{
 				OwnerPtr->ReplicatedPredictionKeyMap.ReplicatePredictionKey(OwnerPtr->ScopedPredictionKey);
-				OwnerPtr->bIsNetDirty = true;
 			}
 		}
 		if (ClearScopedPredictionKey)
diff --git a/Engine/Plugins/Runtime/GameplayAbilities/Source/GameplayAbilities/Public/AbilitySystemComponent.h b/Engine/Plugins/Runtime/GameplayAbilities/Source/GameplayAbilities/Public/AbilitySystemComponent.h
--- a/Engine/Plugins/Runtime/GameplayAbilities/Source/GameplayAbilities/Public/AbilitySystemComponent.h
+++ b/Engine/Plugins/Runtime/GameplayAbilities/Source/GameplayAbilities/Public/AbilitySystemComponent.h
@@ -643,31 +643,26 @@ class GAMEPLAYABILITIES_API UAbilitySystemComponent : public UGameplayTasksCompo
 	FORCEINLINE void AddReplicatedLooseGameplayTag(const FGameplayTag& GameplayTag)
 	{
 		GetReplicatedLooseTags_Mutable().AddTag(GameplayTag);
-		bIsNetDirty = true;
 	}
 
 	FORCEINLINE void AddReplicatedLooseGameplayTags(const FGameplayTagContainer& GameplayTags)
 	{
 		GetReplicatedLooseTags_Mutable().AddTags(GameplayTags);
-		bIsNetDirty = true;
 	}
 
 	FORCEINLINE void RemoveReplicatedLooseGameplayTag(const FGameplayTag& GameplayTag)
 	{
 		GetReplicatedLooseTags_Mutable().RemoveTag(GameplayTag);
-		bIsNetDirty = true;
 	}
 
 	FORCEINLINE void RemoveReplicatedLooseGameplayTags(const FGameplayTagContainer& GameplayTags)
 	{
 		GetReplicatedLooseTags_Mutable().RemoveTags(GameplayTags);
-		bIsNetDirty = true;
 	}
 
 	FORCEINLINE void SetReplicatedLooseGameplayTagCount(const FGameplayTag& GameplayTag, int32 NewCount)
 	{
 		GetReplicatedLooseTags_Mutable().SetTagCount(GameplayTag, NewCount);
-		bIsNetDirty = true;
 	}
 
 	/** 	 
@@ -677,25 +672,21 @@ class GAMEPLAYABILITIES_API UAbilitySystemComponent : public UGameplayTasksCompo
 	FORCEINLINE void AddMinimalReplicationGameplayTag(const FGameplayTag& GameplayTag)
 	{
 		GetMinimalReplicationTags_Mutable().AddTag(GameplayTag);
-		bIsNetDirty = true;
 	}
 
 	FORCEINLINE void AddMinimalReplicationGameplayTags(const FGameplayTagContainer& GameplayTags)
 	{
 		GetMinimalReplicationTags_Mutable().AddTags(GameplayTags);
-		bIsNetDirty = true;
 	}
 
 	FORCEINLINE void RemoveMinimalReplicationGameplayTag(const FGameplayTag& GameplayTag)
 	{
 		GetMinimalReplicationTags_Mutable().RemoveTag(GameplayTag);
-		bIsNetDirty = true;
 	}
 
 	FORCEINLINE void RemoveMinimalReplicationGameplayTags(const FGameplayTagContainer& GameplayTags)
 	{
 		GetMinimalReplicationTags_Mutable().RemoveTags(GameplayTags);
-		bIsNetDirty = true;
 	}
 
 	/** Allow events to be registered for specific gameplay tags being added or removed */
diff --git a/Engine/Source/Runtime/GameplayTasks/Classes/GameplayTasksComponent.h b/Engine/Source/Runtime/GameplayTasks/Classes/GameplayTasksComponent.h
--- a/Engine/Source/Runtime/GameplayTasks/Classes/GameplayTasksComponent.h
+++ b/Engine/Source/Runtime/GameplayTasks/Classes/GameplayTasksComponent.h
@@ -69,11 +69,6 @@ private:
 
 	uint8 bInEventProcessingInProgress : 1;
 
-public:
-	/** Set to indicate that GameplayTasksComponent needs immediate replication. @TODO could just use ForceReplication(), but this allows initial implementation to be game specific. */
-	UPROPERTY()
-	uint8 bIsNetDirty:1;
-
 protected:
 	/** Indicates what's the highest priority among currently running tasks */
 	uint8 TopActivePriority;
diff --git a/Engine/Source/Runtime/GameplayTasks/Private/GameplayTasksComponent.cpp b/Engine/Source/Runtime/GameplayTasks/Private/GameplayTasksComponent.cpp
--- a/Engine/Source/Runtime/GameplayTasks/Private/GameplayTasksComponent.cpp
+++ b/Engine/Source/Runtime/GameplayTasks/Private/GameplayTasksComponent.cpp
@@ -99,7 +99,6 @@ void UGameplayTasksComponent::OnGameplayTaskActivated(UGameplayTask& Task)
 	{
 		const bool bWasAdded = AddSimulatedTask(&Task);
 		check(bWasAdded == true);
-		bIsNetDirty = true;
 	}
 
 	IGameplayTaskOwnerInterface* TaskOwner = Task.GetTaskOwner();
@@ -144,7 +143,6 @@ void UGameplayTasksComponent::OnGameplayTaskDeactivated(UGameplayTask& Task)
 	if (Task.IsSimulatedTask())
 	{
 		RemoveSimulatedTask(&Task);
-		bIsNetDirty = true;
 	}
 
 	// Resource-using task
