Fix Crash closing Lyra after LiveCoding recompile

diff --git a/Engine/Plugins/Runtime/GameplayAbilities/Source/GameplayAbilities/Private/Abilities/Tasks/AbilityTask.cpp b/Engine/Plugins/Runtime/GameplayAbilities/Source/GameplayAbilities/Private/Abilities/Tasks/AbilityTask.cpp
--- a/Engine/Plugins/Runtime/GameplayAbilities/Source/GameplayAbilities/Private/Abilities/Tasks/AbilityTask.cpp
+++ b/Engine/Plugins/Runtime/GameplayAbilities/Source/GameplayAbilities/Private/Abilities/Tasks/AbilityTask.cpp
@@ -72,54 +72,60 @@ UAbilityTask::UAbilityTask(const FObjectInitializer& ObjectInitializer)
 {
 	WaitStateBitMask = static_cast<uint8>(EAbilityTaskWaitState::WaitingOnGame);
 
-	if (AbilityTaskCVars::AbilityTaskRecordingType >= AbilityTaskConstants::DebugMinValueToEnableRecording)
-	{
-		DebugRecordAbilityTaskCreated(this);
-	}
-
-	bool bExceededAbilityTaskMaxCount = false;
-
-	++GlobalAbilityTaskCount;
-	SET_DWORD_STAT(STAT_AbilitySystem_TaskCount, GlobalAbilityTaskCount);
-	if (!(GlobalAbilityTaskCount < AbilityTaskCVars::AbilityTaskMaxCount))
+	if (!HasAnyFlags(RF_ClassDefaultObject))
 	{
-		bExceededAbilityTaskMaxCount = true;
+		if (AbilityTaskCVars::AbilityTaskRecordingType >= AbilityTaskConstants::DebugMinValueToEnableRecording)
+		{
+			DebugRecordAbilityTaskCreated(this);
+		}
 
-		ABILITY_LOG(Warning, TEXT("Way too many AbilityTasks are currently active! %d. %s"), GlobalAbilityTaskCount, *GetClass()->GetName());
+		bool bExceededAbilityTaskMaxCount = false;
 
-		// Auto dump the counts if we hit the limit
-		if (AbilityTaskCVars::AbilityTaskRecordingType >= AbilityTaskConstants::DebugMinValueToEnableRecording)
+		++GlobalAbilityTaskCount;
+		SET_DWORD_STAT(STAT_AbilitySystem_TaskCount, GlobalAbilityTaskCount);
+		if (!(GlobalAbilityTaskCount < AbilityTaskCVars::AbilityTaskMaxCount))
 		{
-			static bool bHasDumpedAbilityTasks = false;  // The dump is spammy, so we only want to auto-dump once
+			bExceededAbilityTaskMaxCount = true;
 
-			if (!bHasDumpedAbilityTasks)
+			ABILITY_LOG(Warning, TEXT("Way too many AbilityTasks are currently active! %d. %s"), GlobalAbilityTaskCount, *GetClass()->GetName());
+
+			// Auto dump the counts if we hit the limit
+			if (AbilityTaskCVars::AbilityTaskRecordingType >= AbilityTaskConstants::DebugMinValueToEnableRecording)
 			{
-				DebugPrintAbilityTasksByClass();
-				bHasDumpedAbilityTasks = true;
+				static bool bHasDumpedAbilityTasks = false;  // The dump is spammy, so we only want to auto-dump once
+
+				if (!bHasDumpedAbilityTasks)
+				{
+					DebugPrintAbilityTasksByClass();
+					bHasDumpedAbilityTasks = true;
 
-				// If we don't flush here the ensure is hit without debug ability task info printed in log
-				GLog->FlushThreadedLogs();
-				GLog->Flush();
+					// If we don't flush here the ensure is hit without debug ability task info printed in log
+					GLog->FlushThreadedLogs();
+					GLog->Flush();
+				}
 			}
 		}
-	}
 
-	ensureMsgf(!bExceededAbilityTaskMaxCount, TEXT("Exceeded AbilityTaskMaxCount. For more information in log set AbilitySystem.AbilityTask.Debug.SourceRecordingEnabled to 1 for non-shipping builds, or 2 for all builds (including shipping)."));
+		ensureMsgf(!bExceededAbilityTaskMaxCount, TEXT("Exceeded AbilityTaskMaxCount. For more information in log set AbilitySystem.AbilityTask.Debug.SourceRecordingEnabled to 1 for non-shipping builds, or 2 for all builds (including shipping)."));
+	}
 }
 
 void UAbilityTask::OnDestroy(bool bInOwnerFinished)
 {
-	checkf(GlobalAbilityTaskCount > 0, TEXT("Mismatched AbilityTask counting"));
-	--GlobalAbilityTaskCount;
-	SET_DWORD_STAT(STAT_AbilitySystem_TaskCount, GlobalAbilityTaskCount);
-
-	bWasSuccessfullyDestroyed = true;
-
-	if (AbilityTaskCVars::AbilityTaskRecordingType >= AbilityTaskConstants::DebugMinValueToEnableRecording)
+	if (!HasAnyFlags(RF_ClassDefaultObject))
 	{
-		DebugRecordAbilityTaskDestroyed(this);
-	}
+		ensureMsgf(GlobalAbilityTaskCount > 0, TEXT("Mismatched AbilityTask counting"));
+		--GlobalAbilityTaskCount;
+		SET_DWORD_STAT(STAT_AbilitySystem_TaskCount, GlobalAbilityTaskCount);
 
+		if (AbilityTaskCVars::AbilityTaskRecordingType >= AbilityTaskConstants::DebugMinValueToEnableRecording)
+		{
+			DebugRecordAbilityTaskDestroyed(this);
+		}
+	}
+	
+	bWasSuccessfullyDestroyed = true;
+	
 	// #KillPendingKill Clear ability reference so we don't hold onto it and GC can delete it.
 	Ability = nullptr;
 
@@ -132,15 +138,18 @@ void UAbilityTask::BeginDestroy()
 
 	if (!bWasSuccessfullyDestroyed)
 	{
-		// this shouldn't happen, it means that ability was destroyed while being active, but we need to keep GlobalAbilityTaskCount in sync anyway
-		checkf(GlobalAbilityTaskCount > 0, TEXT("Mismatched AbilityTask counting"));
-		--GlobalAbilityTaskCount;
-		SET_DWORD_STAT(STAT_AbilitySystem_TaskCount, GlobalAbilityTaskCount);
 		bWasSuccessfullyDestroyed = true;
-
-		if (AbilityTaskCVars::AbilityTaskRecordingType >= AbilityTaskConstants::DebugMinValueToEnableRecording)
+		if (!HasAnyFlags(RF_ClassDefaultObject))
 		{
-			DebugRecordAbilityTaskDestroyed(this);
+			// this shouldn't happen, it means that ability was destroyed while being active, but we need to keep GlobalAbilityTaskCount in sync anyway
+			checkf(GlobalAbilityTaskCount > 0, TEXT("Mismatched AbilityTask counting"));
+			--GlobalAbilityTaskCount;
+			SET_DWORD_STAT(STAT_AbilitySystem_TaskCount, GlobalAbilityTaskCount);
+
+			if (AbilityTaskCVars::AbilityTaskRecordingType >= AbilityTaskConstants::DebugMinValueToEnableRecording)
+			{
+				DebugRecordAbilityTaskDestroyed(this);
+			}
 		}
 	}
 }
