Fix typo in CVarREcordAbilityTaskCounts description

diff --git a/Engine/Plugins/Runtime/GameplayAbilities/Source/GameplayAbilities/Private/Abilities/Tasks/AbilityTask.cpp b/Engine/Plugins/Runtime/GameplayAbilities/Source/GameplayAbilities/Private/Abilities/Tasks/AbilityTask.cpp
--- a/Engine/Plugins/Runtime/GameplayAbilities/Source/GameplayAbilities/Private/Abilities/Tasks/AbilityTask.cpp
+++ b/Engine/Plugins/Runtime/GameplayAbilities/Source/GameplayAbilities/Private/Abilities/Tasks/AbilityTask.cpp
@@ -33,7 +33,7 @@ namespace AbilityTaskCVars
 	static FAutoConsoleVariableRef CVarMaxAbilityTaskCount(
 		TEXT("AbilitySystem.AbilityTask.MaxCount"),
 		AbilityTaskMaxCount,
-		TEXT("Global limit on the number of extant AbilityTasks. Use 'AbilitySystem.Debug.RecordingEnabled' and 'AbilitySystem.AbilityTask.Debug.PrintCounts' to debug why you are hitting this before raising the cap.")
+		TEXT("Global limit on the number of extant AbilityTasks. Use 'AbilitySystem.AbilityTask.Debug.RecordingEnabled' and 'AbilitySystem.AbilityTask.Debug.PrintCounts' to debug why you are hitting this before raising the cap.")
 	);
 
 	// 0 - disabled, 1 - enabled in non-shipping builds, 2 - enabled in all builds (including shipping)
