Modular Gameplay Effects

Move some gameplay effect functionality into optional components.

Gameplay effects have been streamlined to remove functionality that wasn't used by many effects. All of the previous functionality is still available in components that can optionally be added to gameplay effects. This is intended to make the content creator's intent more clear when looking at a gameplay effect and make them more approachable for new users.

Projects are encouraged to write their own components if desired. Newly written components can coexist with the plugin components.  Alternatively, users may replace functionality using inheritance at the project level.

All existing content will automatically update to use components during PostCDOCompiled if necessary.

diff --git a/Engine/Plugins/Runtime/GameplayAbilities/Source/GameplayAbilities/Private/Abilities/GameplayAbility.cpp b/Engine/Plugins/Runtime/GameplayAbilities/Source/GameplayAbilities/Private/Abilities/GameplayAbility.cpp
--- a/Engine/Plugins/Runtime/GameplayAbilities/Source/GameplayAbilities/Private/Abilities/GameplayAbility.cpp
+++ b/Engine/Plugins/Runtime/GameplayAbilities/Source/GameplayAbilities/Private/Abilities/GameplayAbility.cpp
@@ -1018,7 +1018,7 @@ void UGameplayAbility::GetCooldownTimeRemainingAndDuration(FGameplayAbilitySpecH
 const FGameplayTagContainer* UGameplayAbility::GetCooldownTags() const
 {
 	UGameplayEffect* CDGE = GetCooldownGameplayEffect();
-	return CDGE ? &CDGE->InheritableOwnedTagsContainer.CombinedTags : nullptr;
+	return CDGE ? &CDGE->GetGrantedTags() : nullptr;
 }
 
 FGameplayAbilityActorInfo UGameplayAbility::GetActorInfo() const
@@ -1775,7 +1775,7 @@ FActiveGameplayEffectHandle UGameplayAbility::ApplyGameplayEffectToOwner(const F
 		FGameplayEffectSpecHandle SpecHandle = MakeOutgoingGameplayEffectSpec(Handle, ActorInfo, ActivationInfo, GameplayEffect->GetClass(), GameplayEffectLevel);
 		if (SpecHandle.IsValid())
 		{
-			SpecHandle.Data->StackCount = Stacks;
+			SpecHandle.Data->SetStackCount(Stacks);
 			return ApplyGameplayEffectSpecToOwner(Handle, ActorInfo, ActivationInfo, SpecHandle);
 		}
 	}
@@ -1834,7 +1834,7 @@ TArray<FActiveGameplayEffectHandle> UGameplayAbility::ApplyGameplayEffectToTarge
 		FGameplayEffectSpecHandle SpecHandle = MakeOutgoingGameplayEffectSpec(Handle, ActorInfo, ActivationInfo, GameplayEffectClass, GameplayEffectLevel);
 		if (SpecHandle.Data.IsValid())
 		{
-			SpecHandle.Data->StackCount = Stacks;
+			SpecHandle.Data->SetStackCount(Stacks);
 
 			SCOPE_CYCLE_UOBJECT(Source, SpecHandle.Data->GetContext().GetSourceObject());
 			EffectHandles.Append(ApplyGameplayEffectSpecToTarget(Handle, ActorInfo, ActivationInfo, SpecHandle, Target));
diff --git a/Engine/Plugins/Runtime/GameplayAbilities/Source/GameplayAbilities/Private/AbilitySystemBlueprintLibrary.cpp b/Engine/Plugins/Runtime/GameplayAbilities/Source/GameplayAbilities/Private/AbilitySystemBlueprintLibrary.cpp
--- a/Engine/Plugins/Runtime/GameplayAbilities/Source/GameplayAbilities/Private/AbilitySystemBlueprintLibrary.cpp
+++ b/Engine/Plugins/Runtime/GameplayAbilities/Source/GameplayAbilities/Private/AbilitySystemBlueprintLibrary.cpp
@@ -7,6 +7,7 @@
 #include "AbilitySystemLog.h"
 #include "Engine/World.h"
 #include "GameplayEffect.h"
+#include "GameplayEffectComponents/AdditionalEffectsGameplayEffectComponent.h"
 #include "GameplayEffectUIData.h"
 #include "GameplayAbilitySpec.h"
 
@@ -857,6 +858,7 @@ FGameplayEffectSpecHandle UAbilitySystemBlueprintLibrary::AddAssetTags(FGameplay
 	
 FGameplayEffectSpecHandle UAbilitySystemBlueprintLibrary::AddLinkedGameplayEffectSpec(FGameplayEffectSpecHandle SpecHandle, FGameplayEffectSpecHandle LinkedGameplayEffectSpec)
 {
+PRAGMA_DISABLE_DEPRECATION_WARNINGS
 	FGameplayEffectSpec* Spec = SpecHandle.Data.Get();
 	if (Spec)
 	{
@@ -868,10 +870,12 @@ FGameplayEffectSpecHandle UAbilitySystemBlueprintLibrary::AddLinkedGameplayEffec
 	}
 
 	return SpecHandle;
+PRAGMA_ENABLE_DEPRECATION_WARNINGS
 }
 
 FGameplayEffectSpecHandle UAbilitySystemBlueprintLibrary::AddLinkedGameplayEffect(FGameplayEffectSpecHandle SpecHandle, TSubclassOf<UGameplayEffect> LinkedGameplayEffect)
 {
+PRAGMA_DISABLE_DEPRECATION_WARNINGS
 	FGameplayEffectSpecHandle LinkedSpecHandle;
 	FGameplayEffectSpec* Spec = SpecHandle.Data.Get();
 	if (Spec)
@@ -888,6 +892,7 @@ FGameplayEffectSpecHandle UAbilitySystemBlueprintLibrary::AddLinkedGameplayEffec
 	}
 
 	return LinkedSpecHandle;
+PRAGMA_ENABLE_DEPRECATION_WARNINGS
 }
 
 
@@ -896,11 +901,11 @@ FGameplayEffectSpecHandle UAbilitySystemBlueprintLibrary::SetStackCount(FGamepla
 	FGameplayEffectSpec* Spec = SpecHandle.Data.Get();
 	if (Spec)
 	{
-		Spec->StackCount = StackCount;
+		Spec->SetStackCount(StackCount);
 	}
 	else
 	{
-		ABILITY_LOG(Warning, TEXT("UAbilitySystemBlueprintLibrary::AddLinkedGameplayEffectSpec called with invalid SpecHandle"));
+		ABILITY_LOG(Warning, TEXT("%s called with invalid SpecHandle"), ANSI_TO_TCHAR(__func__));
 	}
 	return SpecHandle;
 }
@@ -910,11 +915,11 @@ FGameplayEffectSpecHandle UAbilitySystemBlueprintLibrary::SetStackCountToMax(FGa
 	FGameplayEffectSpec* Spec = SpecHandle.Data.Get();
 	if (Spec && Spec->Def)
 	{
-		Spec->StackCount = Spec->Def->StackLimitCount;
+		Spec->SetStackCount(Spec->Def->StackLimitCount);
 	}
 	else
 	{
-		ABILITY_LOG(Warning, TEXT("UAbilitySystemBlueprintLibrary::AddLinkedGameplayEffectSpec called with invalid SpecHandle"));
+		ABILITY_LOG(Warning, TEXT("%s called with invalid SpecHandle"), ANSI_TO_TCHAR(__func__));
 	}
 	return SpecHandle;
 }
@@ -927,7 +932,7 @@ FGameplayEffectContextHandle UAbilitySystemBlueprintLibrary::GetEffectContext(FG
 	}
 	else
 	{
-		ABILITY_LOG(Warning, TEXT("UAbilitySystemBlueprintLibrary::GetEffectContext called with invalid SpecHandle"));
+		ABILITY_LOG(Warning, TEXT("%s called with invalid SpecHandle"), ANSI_TO_TCHAR(__func__));
 	}
 
 	return FGameplayEffectContextHandle();
@@ -935,17 +940,20 @@ FGameplayEffectContextHandle UAbilitySystemBlueprintLibrary::GetEffectContext(FG
 
 TArray<FGameplayEffectSpecHandle> UAbilitySystemBlueprintLibrary::GetAllLinkedGameplayEffectSpecHandles(FGameplayEffectSpecHandle SpecHandle)
 {
+PRAGMA_DISABLE_DEPRECATION_WARNINGS
+
 	if (FGameplayEffectSpec* Spec = SpecHandle.Data.Get())
 	{
 		return Spec->TargetEffectSpecs;
 	}
 	else
 	{
-		ABILITY_LOG(Warning, TEXT("UAbilitySystemBlueprintLibrary::GetEffectContext called with invalid SpecHandle"));
+		ABILITY_LOG(Warning, TEXT("%s called with invalid SpecHandle"), ANSI_TO_TCHAR(__func__));
 	}
 
 	TArray<FGameplayEffectSpecHandle> Handles;
 	return Handles;
+PRAGMA_ENABLE_DEPRECATION_WARNINGS
 }
 
 int32 UAbilitySystemBlueprintLibrary::GetActiveGameplayEffectStackCount(FActiveGameplayEffectHandle ActiveHandle)
@@ -966,7 +974,7 @@ int32 UAbilitySystemBlueprintLibrary::GetActiveGameplayEffectStackLimitCount(FAc
 		const UGameplayEffect* ActiveGE = ASC->GetGameplayEffectDefForHandle(ActiveHandle);
 		if (ActiveGE)
 		{
-			return ActiveGE->StackLimitCount;
+			return ActiveGE->GetStackLimitCount();
 		}
 	}
 	return 0;
@@ -1100,14 +1108,22 @@ bool UAbilitySystemBlueprintLibrary::RemoveLooseGameplayTags(AActor* Actor, cons
 
 const UGameplayEffectUIData* UAbilitySystemBlueprintLibrary::GetGameplayEffectUIData(TSubclassOf<UGameplayEffect> EffectClass, TSubclassOf<UGameplayEffectUIData> DataType)
 {
-	if (UClass* ActualPtr = EffectClass.Get())
+	if (const UGameplayEffect* Effect = EffectClass.GetDefaultObject())
 	{
-		const UGameplayEffectUIData* UIData = GetDefault<UGameplayEffect>(ActualPtr)->UIData;
+		const UGameplayEffectUIData* UIData = Effect->FindComponent<UGameplayEffectUIData>();
+		if (!UIData)
+		{
+PRAGMA_DISABLE_DEPRECATION_WARNINGS
+			UIData = Effect->UIData;
+PRAGMA_ENABLE_DEPRECATION_WARNINGS
+		}
+
 		if ((UIData != nullptr) && (DataType != nullptr) && UIData->IsA(DataType))
 		{
 			return UIData;
 		}
 	}
+
 	return nullptr;
 }
 
diff --git a/Engine/Plugins/Runtime/GameplayAbilities/Source/GameplayAbilities/Private/AbilitySystemComponent.cpp b/Engine/Plugins/Runtime/GameplayAbilities/Source/GameplayAbilities/Private/AbilitySystemComponent.cpp
--- a/Engine/Plugins/Runtime/GameplayAbilities/Source/GameplayAbilities/Private/AbilitySystemComponent.cpp
+++ b/Engine/Plugins/Runtime/GameplayAbilities/Source/GameplayAbilities/Private/AbilitySystemComponent.cpp
@@ -275,6 +275,37 @@ void UAbilitySystemComponent::CacheIsNetSimulated()
 	UpdateActiveGameplayEffectsReplicationCondition();
 }
 
+void UAbilitySystemComponent::InhibitActiveGameplayEffect(FActiveGameplayEffectHandle ActiveGEHandle, bool bInhibit, bool bInvokeGameplayCueEvents)
+{
+	FActiveGameplayEffect* ActiveGE = ActiveGameplayEffects.GetActiveGameplayEffect(ActiveGEHandle);
+	if (!ActiveGE)
+	{
+		ABILITY_LOG(Error, TEXT("%s received bad Active GameplayEffect Handle: %s"), ANSI_TO_TCHAR(__func__), *ActiveGEHandle.ToString());
+		return;
+	}
+
+	if (ActiveGE->bIsInhibited != bInhibit)
+	{
+		ActiveGE->bIsInhibited = bInhibit;
+
+		// All OnDirty callbacks must be inhibited until we update this entire GameplayEffect.
+		FScopedAggregatorOnDirtyBatch	AggregatorOnDirtyBatcher;
+		if (bInhibit)
+		{
+			// Remove our ActiveGameplayEffects modifiers with our Attribute Aggregators
+			ActiveGameplayEffects.RemoveActiveGameplayEffectGrantedTagsAndModifiers(*ActiveGE, bInvokeGameplayCueEvents);
+		}
+		else
+		{
+			ActiveGameplayEffects.AddActiveGameplayEffectGrantedTagsAndModifiers(*ActiveGE, bInvokeGameplayCueEvents);
+		}
+
+		ActiveGE->EventSet.OnInhibitionChanged.Broadcast(ActiveGEHandle, ActiveGE->bIsInhibited);
+	}
+}
+
+
+
 const FActiveGameplayEffect* UAbilitySystemComponent::GetActiveGameplayEffect(const FActiveGameplayEffectHandle Handle) const
 {
 	return ActiveGameplayEffects.GetActiveGameplayEffect(Handle);
@@ -632,9 +663,9 @@ FOnGameplayEffectTagCountChanged& UAbilitySystemComponent::RegisterGameplayTagEv
 	return GameplayTagCountContainer.RegisterGameplayTagEvent(Tag, EventType);
 }
 
-void UAbilitySystemComponent::UnregisterGameplayTagEvent(FDelegateHandle DelegateHandle, FGameplayTag Tag, EGameplayTagEventType::Type EventType)
+bool UAbilitySystemComponent::UnregisterGameplayTagEvent(FDelegateHandle DelegateHandle, FGameplayTag Tag, EGameplayTagEventType::Type EventType)
 {
-	GameplayTagCountContainer.RegisterGameplayTagEvent(Tag, EventType).Remove(DelegateHandle);
+	return GameplayTagCountContainer.RegisterGameplayTagEvent(Tag, EventType).Remove(DelegateHandle);
 }
 
 FDelegateHandle UAbilitySystemComponent::RegisterAndCallGameplayTagEvent(FGameplayTag Tag, FOnGameplayEffectTagCountChanged::FDelegate Delegate, EGameplayTagEventType::Type EventType)
@@ -749,9 +780,9 @@ FActiveGameplayEffectHandle UAbilitySystemComponent::ApplyGameplayEffectSpecToSe
 	}
 
 	// Don't allow prediction of periodic effects
-	if(PredictionKey.IsValidKey() && Spec.GetPeriod() > 0.f)
+	if (PredictionKey.IsValidKey() && Spec.GetPeriod() > 0.f)
 	{
-		if(IsOwnerActorAuthoritative())
+		if (IsOwnerActorAuthoritative())
 		{
 			// Server continue with invalid prediction key
 			PredictionKey = FPredictionKey();
@@ -763,64 +794,34 @@ FActiveGameplayEffectHandle UAbilitySystemComponent::ApplyGameplayEffectSpecToSe
 		}
 	}
 
-	// Are we currently immune to this? (ApplicationImmunity)
-	const FActiveGameplayEffect* ImmunityGE=nullptr;
-	if (ActiveGameplayEffects.HasApplicationImmunityToSpec(Spec, ImmunityGE))
+	// Check if there is a registered "application" query that can block the application
+	for (const FGameplayEffectApplicationQuery& ApplicationQuery : GameplayEffectApplicationQueries)
 	{
-		OnImmunityBlockGameplayEffect(Spec, ImmunityGE);
-		return FActiveGameplayEffectHandle();
-	}
-
-	// Check AttributeSet requirements: make sure all attributes are valid
-	// We may want to cache this off in some way to make the runtime check quicker.
-	// We also need to handle things in the execution list
-	for (const FGameplayModifierInfo& Mod : Spec.Def->Modifiers)
-	{
-		if (!Mod.Attribute.IsValid())
+		const bool bAllowed = ApplicationQuery.Execute(ActiveGameplayEffects, Spec);
+		if (!bAllowed)
 		{
-			ABILITY_LOG(Warning, TEXT("%s has a null modifier attribute."), *Spec.Def->GetPathName());
 			return FActiveGameplayEffectHandle();
 		}
 	}
 
 	// check if the effect being applied actually succeeds
-	float ChanceToApply = Spec.GetChanceToApplyToTarget();
-	if ((ChanceToApply < 1.f - SMALL_NUMBER) && (FMath::FRand() > ChanceToApply))
+	if (!Spec.Def->CanApply(ActiveGameplayEffects, Spec))
 	{
 		return FActiveGameplayEffectHandle();
 	}
 
-	// Get MyTags.
-	//	We may want to cache off a GameplayTagContainer instead of rebuilding it every time.
-	//	But this will also be where we need to merge in context tags? (Headshot, executing ability, etc?)
-	//	Or do we push these tags into (our copy of the spec)?
-
+	// Check AttributeSet requirements: make sure all attributes are valid
+	// We may want to cache this off in some way to make the runtime check quicker.
+	// We also need to handle things in the execution list
+	for (const FGameplayModifierInfo& Mod : Spec.Def->Modifiers)
 	{
-		// Note: static is ok here since the scope is so limited, but wider usage of MyTags is not safe since this function can be recursively called
-		static FGameplayTagContainer MyTags;
-		MyTags.Reset();
-
-		GetOwnedGameplayTags(MyTags);
-
-		if (Spec.Def->ApplicationTagRequirements.RequirementsMet(MyTags) == false)
-		{
-			return FActiveGameplayEffectHandle();
-		}
-
-		if (!Spec.Def->RemovalTagRequirements.IsEmpty() && Spec.Def->RemovalTagRequirements.RequirementsMet(MyTags) == true)
+		if (!Mod.Attribute.IsValid())
 		{
+			ABILITY_LOG(Warning, TEXT("%s has a null modifier attribute."), *Spec.Def->GetPathName());
 			return FActiveGameplayEffectHandle();
 		}
 	}
 
-	// Custom application requirement check
-	for (const TSubclassOf<UGameplayEffectCustomApplicationRequirement>& AppReq : Spec.Def->ApplicationRequirements)
-	{
-		if (*AppReq && AppReq->GetDefaultObject<UGameplayEffectCustomApplicationRequirement>()->CanApplyGameplayEffect(Spec.Def, Spec, this) == false)
-		{
-			return FActiveGameplayEffectHandle();
-		}
-	}
 
 	// Clients should treat predicted instant effects as if they have infinite duration. The effects will be cleaned up later.
 	bool bTreatAsInfiniteDuration = GetOwnerRole() != ROLE_Authority && PredictionKey.IsLocalClientKey() && Spec.Def->DurationPolicy == EGameplayEffectDurationType::Instant;
@@ -833,9 +834,8 @@ FActiveGameplayEffectHandle UAbilitySystemComponent::ApplyGameplayEffectSpecToSe
 	bool bFoundExistingStackableGE = false;
 
 	FActiveGameplayEffect* AppliedEffect = nullptr;
-
 	FGameplayEffectSpec* OurCopyOfSpec = nullptr;
-	TSharedPtr<FGameplayEffectSpec> StackSpec;
+	TUniquePtr<FGameplayEffectSpec> StackSpec;
 	{
 		if (Spec.Def->DurationPolicy != EGameplayEffectDurationType::Instant || bTreatAsInfiniteDuration)
 		{
@@ -864,8 +864,9 @@ FActiveGameplayEffectHandle UAbilitySystemComponent::ApplyGameplayEffectSpecToSe
 
 		if (!OurCopyOfSpec)
 		{
-			StackSpec = TSharedPtr<FGameplayEffectSpec>(new FGameplayEffectSpec(Spec));
+			StackSpec = MakeUnique<FGameplayEffectSpec>(Spec);
 			OurCopyOfSpec = StackSpec.Get();
+
 			UAbilitySystemGlobals::Get().GlobalPreGameplayEffectSpecApply(*OurCopyOfSpec, this);
 			OurCopyOfSpec->CaptureAttributeDataFromTarget(this);
 		}
@@ -878,12 +879,8 @@ FActiveGameplayEffectHandle UAbilitySystemComponent::ApplyGameplayEffectSpecToSe
 		}
 	}
 
-	if (OurCopyOfSpec)
-	{
-		// Update (not push) the global spec being applied [we want to switch it to our copy, from the const input copy)
-		UAbilitySystemGlobals::Get().SetCurrentAppliedGE(OurCopyOfSpec);
-	}
-	
+	// Update (not push) the global spec being applied [we want to switch it to our copy, from the const input copy)
+	UAbilitySystemGlobals::Get().SetCurrentAppliedGE(OurCopyOfSpec);
 
 	// We still probably want to apply tags and stuff even if instant?
 	// If bSuppressStackingCues is set for this GameplayEffect, only add the GameplayCue if this is the first instance of the GameplayEffect
@@ -896,7 +893,7 @@ FActiveGameplayEffectHandle UAbilitySystemComponent::ApplyGameplayEffectSpecToSe
 
 		// Fixme: what if we wanted to scale Cue magnitude based on damage? E.g, scale an cue effect when the GE is buffed?
 
-		if (OurCopyOfSpec->StackCount > Spec.StackCount)
+		if (OurCopyOfSpec->GetStackCount() > Spec.GetStackCount())
 		{
 			// Because PostReplicatedChange will get called from modifying the stack count
 			// (and not PostReplicatedAdd) we won't know which GE was modified.
@@ -926,39 +923,10 @@ FActiveGameplayEffectHandle UAbilitySystemComponent::ApplyGameplayEffectSpecToSe
 	}
 	else if (Spec.Def->DurationPolicy == EGameplayEffectDurationType::Instant)
 	{
-		if (OurCopyOfSpec->Def->OngoingTagRequirements.IsEmpty())
-		{
-			ExecuteGameplayEffect(*OurCopyOfSpec, PredictionKey);
-		}
-		else
-		{
-			ABILITY_LOG(Warning, TEXT("%s is instant but has tag requirements. Tag requirements can only be used with gameplay effects that have a duration. This gameplay effect will be ignored."), *Spec.Def->GetPathName());
-		}
-	}
-
-	if (Spec.GetPeriod() != UGameplayEffect::NO_PERIOD && Spec.TargetEffectSpecs.Num() > 0)
-	{
-		ABILITY_LOG(Warning, TEXT("%s is periodic but also applies GameplayEffects to its target. GameplayEffects will only be applied once, not every period."), *Spec.Def->GetPathName());
+		// This is a non-predicted instant effect (it never gets added to ActiveGameplayEffects)
+		ExecuteGameplayEffect(*OurCopyOfSpec, PredictionKey);
 	}
 	
-	// evaluate if any active effects need to be removed by the application of this effect
-	if (bIsNetAuthority)
-	{
-		ActiveGameplayEffects.AttemptRemoveActiveEffectsOnEffectApplication(*OurCopyOfSpec, MyHandle);
-	}
-	
-	// ------------------------------------------------------
-	// Apply Linked effects
-	// todo: this is ignoring the returned handles, should we put them into a TArray and return all of the handles?
-	// ------------------------------------------------------
-	for (const FGameplayEffectSpecHandle& TargetSpec: Spec.TargetEffectSpecs)
-	{
-		if (TargetSpec.IsValid())
-		{
-			ApplyGameplayEffectSpecToSelf(*TargetSpec.Data.Get(), PredictionKey);
-		}
-	}
-
 	UAbilitySystemComponent* InstigatorASC = Spec.GetContext().GetInstigatorAbilitySystemComponent();
 
 	// Send ourselves a callback	
@@ -1127,7 +1095,7 @@ int32 UAbilitySystemComponent::GetCurrentStackCount(FActiveGameplayEffectHandle
 {
 	if (const FActiveGameplayEffect* ActiveGE = ActiveGameplayEffects.GetActiveGameplayEffect(Handle))
 	{
-		return ActiveGE->Spec.StackCount;
+		return ActiveGE->Spec.GetStackCount();
 	}
 	return 0;
 }
@@ -1148,24 +1116,21 @@ FString UAbilitySystemComponent::GetActiveGEDebugString(FActiveGameplayEffectHan
 
 	if (const FActiveGameplayEffect* ActiveGE = ActiveGameplayEffects.GetActiveGameplayEffect(Handle))
 	{
-		Str = FString::Printf(TEXT("%s - (Level: %.2f. Stacks: %d)"), *ActiveGE->Spec.Def->GetName(), ActiveGE->Spec.GetLevel(), ActiveGE->Spec.StackCount);
+		Str = FString::Printf(TEXT("%s - (Level: %.2f. Stacks: %d)"), *ActiveGE->Spec.Def->GetName(), ActiveGE->Spec.GetLevel(), ActiveGE->Spec.GetStackCount());
 	}
 
 	return Str;
 }
 
+/** Gets the GE Handle of the GE that granted the passed in Ability */
 FActiveGameplayEffectHandle UAbilitySystemComponent::FindActiveGameplayEffectHandle(FGameplayAbilitySpecHandle Handle) const
 {
-	for (const FActiveGameplayEffect& ActiveGE : &ActiveGameplayEffects)
+	const FGameplayAbilitySpec* Spec = FindAbilitySpecFromHandle(Handle);
+	if (Spec)
 	{
-		for (const FGameplayAbilitySpecDef& AbilitySpecDef : ActiveGE.Spec.GrantedAbilitySpecs)
-		{
-			if (AbilitySpecDef.AssignedHandle == Handle)
-			{
-				return ActiveGE.Handle;
-			}
-		}
+		return Spec->GameplayEffectHandle;
 	}
+
 	return FActiveGameplayEffectHandle();
 }
 
@@ -2546,16 +2511,16 @@ void UAbilitySystemComponent::Debug_Internal(FAbilitySystemComponentDebugInfo& I
 			}
 
 			FString StackString;
-			if (ActiveGE.Spec.StackCount > 1)
+			if (ActiveGE.Spec.GetStackCount() > 1)
 			{
 
 				if (ActiveGE.Spec.Def->StackingType == EGameplayEffectStackingType::AggregateBySource)
 				{
-					StackString = FString::Printf(TEXT("(Stacks: %d. From: %s) "), ActiveGE.Spec.StackCount, *GetNameSafe(ActiveGE.Spec.GetContext().GetInstigatorAbilitySystemComponent()->GetAvatarActor_Direct()));
+					StackString = FString::Printf(TEXT("(Stacks: %d. From: %s) "), ActiveGE.Spec.GetStackCount(), *GetNameSafe(ActiveGE.Spec.GetContext().GetInstigatorAbilitySystemComponent()->GetAvatarActor_Direct()));
 				}
 				else
 				{
-					StackString = FString::Printf(TEXT("(Stacks: %d) "), ActiveGE.Spec.StackCount);
+					StackString = FString::Printf(TEXT("(Stacks: %d) "), ActiveGE.Spec.GetStackCount());
 				}
 			}
 
diff --git a/Engine/Plugins/Runtime/GameplayAbilities/Source/GameplayAbilities/Private/AbilitySystemComponent_Abilities.cpp b/Engine/Plugins/Runtime/GameplayAbilities/Source/GameplayAbilities/Private/AbilitySystemComponent_Abilities.cpp
--- a/Engine/Plugins/Runtime/GameplayAbilities/Source/GameplayAbilities/Private/AbilitySystemComponent_Abilities.cpp
+++ b/Engine/Plugins/Runtime/GameplayAbilities/Source/GameplayAbilities/Private/AbilitySystemComponent_Abilities.cpp
@@ -741,38 +741,38 @@ void UAbilitySystemComponent::DecrementAbilityListLock()
 	}
 }
 
-FGameplayAbilitySpec* UAbilitySystemComponent::FindAbilitySpecFromHandle(FGameplayAbilitySpecHandle Handle)
+FGameplayAbilitySpec* UAbilitySystemComponent::FindAbilitySpecFromHandle(FGameplayAbilitySpecHandle Handle) const
 {
 	SCOPE_CYCLE_COUNTER(STAT_FindAbilitySpecFromHandle);
 
-	for (FGameplayAbilitySpec& Spec : ActivatableAbilities.Items)
+	for (const FGameplayAbilitySpec& Spec : ActivatableAbilities.Items)
 	{
 		if (Spec.Handle == Handle)
 		{
-			return &Spec;
+			return const_cast<FGameplayAbilitySpec*>(&Spec);
 		}
 	}
 
 	return nullptr;
 }
 
-FGameplayAbilitySpec* UAbilitySystemComponent::FindAbilitySpecFromGEHandle(FActiveGameplayEffectHandle Handle)
+FGameplayAbilitySpec* UAbilitySystemComponent::FindAbilitySpecFromGEHandle(FActiveGameplayEffectHandle Handle) const
 {
-	for (FGameplayAbilitySpec& Spec : ActivatableAbilities.Items)
+	for (const FGameplayAbilitySpec& Spec : ActivatableAbilities.Items)
 	{
 		if (Spec.GameplayEffectHandle == Handle)
 		{
-			return &Spec;
+			return const_cast<FGameplayAbilitySpec*>(&Spec);
 		}
 	}
 	return nullptr;
 }
 
-FGameplayAbilitySpec* UAbilitySystemComponent::FindAbilitySpecFromClass(TSubclassOf<UGameplayAbility> InAbilityClass)
+FGameplayAbilitySpec* UAbilitySystemComponent::FindAbilitySpecFromClass(TSubclassOf<UGameplayAbility> InAbilityClass) const
 {
 	SCOPE_CYCLE_COUNTER(STAT_FindAbilitySpecFromHandle);
 
-	for (FGameplayAbilitySpec& Spec : ActivatableAbilities.Items)
+	for (const FGameplayAbilitySpec& Spec : ActivatableAbilities.Items)
 	{
 		if (Spec.Ability == nullptr)
 		{
@@ -781,7 +781,7 @@ FGameplayAbilitySpec* UAbilitySystemComponent::FindAbilitySpecFromClass(TSubclas
 
 		if (Spec.Ability->GetClass() == InAbilityClass)
 		{
-			return &Spec;
+			return const_cast<FGameplayAbilitySpec*>(&Spec);
 		}
 	}
 
@@ -806,15 +806,15 @@ void UAbilitySystemComponent::MarkAbilitySpecDirty(FGameplayAbilitySpec& Spec, b
 	}
 }
 
-FGameplayAbilitySpec* UAbilitySystemComponent::FindAbilitySpecFromInputID(int32 InputID)
+FGameplayAbilitySpec* UAbilitySystemComponent::FindAbilitySpecFromInputID(int32 InputID) const
 {
 	if (InputID != INDEX_NONE)
 	{
-		for (FGameplayAbilitySpec& Spec : ActivatableAbilities.Items)
+		for (const FGameplayAbilitySpec& Spec : ActivatableAbilities.Items)
 		{
 			if (Spec.InputID == InputID)
 			{
-				return &Spec;
+				return const_cast<FGameplayAbilitySpec*>(&Spec);
 			}
 		}
 	}
diff --git a/Engine/Plugins/Runtime/GameplayAbilities/Source/GameplayAbilities/Private/GameplayAbilityTypes.cpp b/Engine/Plugins/Runtime/GameplayAbilities/Source/GameplayAbilities/Private/GameplayAbilityTypes.cpp
--- a/Engine/Plugins/Runtime/GameplayAbilities/Source/GameplayAbilities/Private/GameplayAbilityTypes.cpp
+++ b/Engine/Plugins/Runtime/GameplayAbilities/Source/GameplayAbilities/Private/GameplayAbilityTypes.cpp
@@ -183,6 +183,19 @@ void FGameplayAbilityActivationInfo::SetActivationRejected()
 	ActivationMode = EGameplayAbilityActivationMode::Rejected;
 }
 
+bool FGameplayAbilitySpecDef::operator==(const FGameplayAbilitySpecDef& Other) const
+{
+	return Ability == Other.Ability &&
+		LevelScalableFloat == Other.LevelScalableFloat &&
+		InputID == Other.InputID &&
+		RemovalPolicy == Other.RemovalPolicy;
+}
+
+bool FGameplayAbilitySpecDef::operator!=(const FGameplayAbilitySpecDef& Other) const
+{
+	return !(*this == Other);
+}
+
 bool FGameplayAbilitySpec::IsActive() const
 {
 	// If ability hasn't replicated yet we're not active
diff --git a/Engine/Plugins/Runtime/GameplayAbilities/Source/GameplayAbilities/Private/GameplayDebuggerCategory_Abilities.cpp b/Engine/Plugins/Runtime/GameplayAbilities/Source/GameplayAbilities/Private/GameplayDebuggerCategory_Abilities.cpp
--- a/Engine/Plugins/Runtime/GameplayAbilities/Source/GameplayAbilities/Private/GameplayDebuggerCategory_Abilities.cpp
+++ b/Engine/Plugins/Runtime/GameplayAbilities/Source/GameplayAbilities/Private/GameplayDebuggerCategory_Abilities.cpp
@@ -85,7 +85,7 @@ void FGameplayDebuggerCategory_Abilities::CollectData(APlayerController* OwnerPC
 			ItemData.Context = EffectSpec.GetContext().ToString();
 			ItemData.Duration = EffectSpec.GetDuration();
 			ItemData.Period = EffectSpec.GetPeriod();
-			ItemData.Stacks = EffectSpec.StackCount;
+			ItemData.Stacks = EffectSpec.GetStackCount();
 			ItemData.Level = EffectSpec.GetLevel();
 
 			DataPack.GameplayEffects.Add(ItemData);
diff --git a/Engine/Plugins/Runtime/GameplayAbilities/Source/GameplayAbilities/Private/GameplayEffect.cpp b/Engine/Plugins/Runtime/GameplayAbilities/Source/GameplayAbilities/Private/GameplayEffect.cpp
--- a/Engine/Plugins/Runtime/GameplayAbilities/Source/GameplayAbilities/Private/GameplayEffect.cpp
+++ b/Engine/Plugins/Runtime/GameplayAbilities/Source/GameplayAbilities/Private/GameplayEffect.cpp
@@ -17,9 +17,26 @@
 #include "GameplayEffectExecutionCalculation.h"
 #include "GameplayCueManager.h"
 #include "UObject/ObjectSaveContext.h"
+#include "UObject/UObjectThreadContext.h" // FUObjectSerializeContext
+#include "GameplayEffectTypes.h"
+#include "GameplayEffectCustomApplicationRequirement.h"
+#include "GameplayEffectUIData.h"
+#include "Misc/DataValidation.h"
+
+#include "GameplayEffectComponents/AssetTagsGameplayEffectComponent.h"
+#include "GameplayEffectComponents/BlockAbilityTagsGameplayEffectComponent.h"
+#include "GameplayEffectComponents/ChanceToApplyGameplayEffectComponent.h"
+#include "GameplayEffectComponents/AdditionalEffectsGameplayEffectComponent.h"
+#include "GameplayEffectComponents/CustomCanApplyGameplayEffectComponent.h"
+#include "GameplayEffectComponents/ImmunityGameplayEffectComponent.h"
+#include "GameplayEffectComponents/RemoveOtherGameplayEffectComponent.h"
+#include "GameplayEffectComponents/TargetTagRequirementsGameplayEffectComponent.h"
+#include "GameplayEffectComponents/TargetTagsGameplayEffectComponent.h"
 
 #include UE_INLINE_GENERATED_CPP_BY_NAME(GameplayEffect)
 
+#define LOCTEXT_NAMESPACE "GameplayEffect"
+
 #if ENABLE_VISUAL_LOG
 //#include "VisualLogger/VisualLogger.h"
 #endif // ENABLE_VISUAL_LOG
@@ -46,158 +63,625 @@ DECLARE_CYCLE_STAT(TEXT("MakeQuery"), STAT_MakeGameplayEffectQuery, STATGROUP_Ab
 	SCALABLEFLOAT_REPORTERROR_WITHPATHNAME(Scalable, PathNameString);
 #endif // WITH_EDITOR
 
+// Versioning code
+bool FGameplayEffectVersion::Serialize(FStructuredArchive::FSlot Slot)
+{
+	Slot << CurrentVersion;
+	return true;
+}
+
 // --------------------------------------------------------------------------------------------------------------------------------------------------------
 //
 //	UGameplayEffect
 //
 // --------------------------------------------------------------------------------------------------------------------------------------------------------
 
-
 UGameplayEffect::UGameplayEffect(const FObjectInitializer& ObjectInitializer)
 	: Super(ObjectInitializer)
 {
+	PRAGMA_DISABLE_DEPRECATION_WARNINGS
 	DurationPolicy = EGameplayEffectDurationType::Instant;
 	bExecutePeriodicEffectOnApplication = true;
 	PeriodicInhibitionPolicy = EGameplayEffectPeriodInhibitionRemovedPolicy::NeverReset;
-	ChanceToApplyToTarget.SetValue(1.f);
+	ChanceToApplyToTarget_DEPRECATED.SetValue(1.f);
 	StackingType = EGameplayEffectStackingType::None;
 	StackLimitCount = 0;
 	StackDurationRefreshPolicy = EGameplayEffectStackingDurationPolicy::RefreshOnSuccessfulApplication;
 	StackPeriodResetPolicy = EGameplayEffectStackingPeriodPolicy::ResetOnSuccessfulApplication;
 	bRequireModifierSuccessToTriggerCues = true;
+	PRAGMA_ENABLE_DEPRECATION_WARNINGS
+}
+
+// Needed for versioning code
+void UGameplayEffect::PostInitProperties()
+{
+	Super::PostInitProperties();
+
+#if WITH_EDITORONLY_DATA
+	SetVersion(EGameplayEffectVersion::Monolithic);
+#endif
 }
 
 void UGameplayEffect::GetOwnedGameplayTags(FGameplayTagContainer& TagContainer) const
 {
-	TagContainer.AppendTags(InheritableOwnedTagsContainer.CombinedTags);
+	UE_LOG(LogGameplayEffects, Warning, TEXT("GetOwnedGameplayTags: The implementation and method name did not match.  Use GetGrantedTags() to get the tags Granted to the Actor this GameplayEffect is applied to."));
+	TagContainer.AppendTags(GetGrantedTags());
+}
+
+bool UGameplayEffect::HasMatchingGameplayTag(FGameplayTag TagToCheck) const
+{
+	UE_LOG(LogGameplayEffects, Warning, TEXT("HasMatchingGameplayTag: The implementation and method name did not match.  Use GetGrantedTags().HasTag() to check against the tags this GameplayEffect will Grant to the Actor."));
+	return IGameplayTagAssetInterface::HasMatchingGameplayTag(TagToCheck);
+}
+
+bool UGameplayEffect::HasAllMatchingGameplayTags(const FGameplayTagContainer& TagContainer) const
+{
+	UE_LOG(LogGameplayEffects, Warning, TEXT("HasAllMatchingGameplayTags: The implementation and method name did not match.  Use GetGrantedTags().HasAll() to check against the tags this GameplayEffect will Grant to the Actor."));
+	return IGameplayTagAssetInterface::HasAllMatchingGameplayTags(TagContainer);
+}
+
+bool UGameplayEffect::HasAnyMatchingGameplayTags(const FGameplayTagContainer& TagContainer) const
+{
+	UE_LOG(LogGameplayEffects, Warning, TEXT("HasAnyMatchingGameplayTags: The implementation and method name did not match.  Use GetGrantedTags().HasAny() to check against the tags this GameplayEffect will Grant to the Actor."));
+	return IGameplayTagAssetInterface::HasAnyMatchingGameplayTags(TagContainer);
 }
 
 void UGameplayEffect::GetBlockedAbilityTags(FGameplayTagContainer& OutTagContainer) const
 {
-	OutTagContainer.AppendTags(InheritableBlockedAbilityTagsContainer.CombinedTags);
+	OutTagContainer.AppendTags(CachedBlockedAbilityTags);
+}
+
+#if WITH_EDITOR
+EDataValidationResult UGameplayEffect::IsDataValid(FDataValidationContext& Context)
+{
+	EDataValidationResult ValidationResult = Super::IsDataValid(Context);
+
+	if (ValidationResult != EDataValidationResult::Invalid)
+	{
+		for (UGameplayEffectComponent* GEComponent : GEComponents)
+		{
+			if (GEComponent)
+			{
+				ValidationResult = GEComponent->IsDataValid(Context);
+				if (ValidationResult == EDataValidationResult::Invalid)
+				{
+					break;
+				}
+			}
+			else
+			{
+				Context.AddWarning(LOCTEXT("GEIsNull", "Null entry in GEComponents"));
+			}
+		}
+	}
+
+	TArray<FText> Warnings, Errors;
+	Context.SplitIssues(Warnings, Errors);
+
+	if (Errors.Num() > 0)
+	{
+		EditorStatusText = FText::FormatOrdered(LOCTEXT("ErrorsFmt", "Error: {0}"), FText::Join(FText::FromStringView(TEXT(", ")), Errors));
+	}
+	else if (Warnings.Num() > 0)
+	{
+		EditorStatusText = FText::FormatOrdered(LOCTEXT("WarningsFmt", "Warning: {0}"), FText::Join(FText::FromStringView(TEXT(", ")), Warnings));
+	}
+	else
+	{
+		EditorStatusText = LOCTEXT("AllOk", "All Ok");
+	}
+
+	return ValidationResult;
 }
+#endif
 
 void UGameplayEffect::PostLoad()
 {
 	Super::PostLoad();
 
-	// Temporary post-load fix-up to preserve magnitude data
-	for (FGameplayModifierInfo& CurModInfo : Modifiers)
+	// Call PostLoad on components and cache tag info from components for more efficient lookups
+	for (UGameplayEffectComponent* GEComponent : GEComponents)
 	{
+		if (GEComponent)
+		{
+			GEComponent->ConditionalPostLoad(); // Ensure the SubObject is PostLoaded
+			GEComponent->OnOwnerPostLoad();
+		}
+	}
+
 #if WITH_EDITOR
+	SCALABLEFLOAT_REPORTERROR_WITHPOSTLOAD(Period);
+
+	PRAGMA_DISABLE_DEPRECATION_WARNINGS
+	HasRemoveGameplayEffectsQuery = !RemoveGameplayEffectQuery.IsEmpty();
+	PRAGMA_ENABLE_DEPRECATION_WARNINGS
+
+	DurationMagnitude.ReportErrors(GetPathName());
+
+	for (FGameplayModifierInfo& CurModInfo : Modifiers)
+	{
 		CurModInfo.ModifierMagnitude.ReportErrors(GetPathName());
+	}
+
+	// Update the EditorStatusText
+	FDataValidationContext DataValidationContext;
+	IsDataValid(DataValidationContext);
+
+	// We're done loading (and therefore upgrading), boost the version.
+	SetVersion(EGameplayEffectVersion::Current);
 #endif // WITH_EDITOR
+}
+
+#if WITH_EDITOR
+/**
+ * We now support GameplayEffectComponents which are Subobjects.
+ * 
+ * When we're loading a Blueprint (including Data Only Blueprints), we go through a bunch of steps:  Preload, Serialize (which does the instancing), Init, CDO Compile, PostLoad.
+ * However, we are not guaranteed that our Parent Class (Archetype) is fully loaded until PostLoad.  So during load (or cooking, where we see this most) is that
+ * a Child will request to load, then it's possible the Parent has not fully loaded, so they may also request its Parent load and they go through these steps *together*.
+ * When we get to PostCDOCompiled, the Parent may create some Subobjects (for us, this happens during the Monolithic -> Modular upgrade) and yet the Child is also at the same step,
+ * so it hasn't actually seen those Subobjects and instanced them, but it *would have* instanced them if Parent Class was loaded first through some other means.
+ * So our job here is to ensure the Subobjects that exist in the Parent also exist in the Child, so that the order of loading the classes doesn't matter.
+ *  
+ * There are other issues that will pop-up:
+ *	1. You cannot remove a Parent's Subobject in the Child (this code will just recreate it).
+ *	2. If you remove a Subobject from the Parent, the Child will continue to own it and it will be delta serialized *from the Grandparent*
+ */
+void UGameplayEffect::PostCDOCompiledFixupSubobjects()
+{
+	const UGameplayEffect* Archetype = Cast<UGameplayEffect>(GetArchetype());
+	if (!Archetype)
+	{
+		return;
 	}
 
-	// We need to update when we first load to override values coming in from the superclass
-	// We also copy the tags from the old tag containers into the inheritable tag containers
+	for (const UGameplayEffectComponent* ParentComponent : Archetype->GEComponents)
+	{
+		if (!ParentComponent)
+		{
+			continue;
+		}
 
-	UpdateInheritedTagProperties();
+		bool bFound = false;
+		const FName ParentComponentName = ParentComponent->GetFName();
+		for (const UGameplayEffectComponent* ChildComponent : GEComponents)
+		{
+			// When the SubObject code decides how to delta serialize from its Archetype
+			// The Archetype is determined by name. Let's match-up specifically by name rather than type alone.
+			if (ChildComponent && ChildComponent->GetFName() == ParentComponentName)
+			{
+				bFound = true;
+				break;
+			}
+		}
 
-	HasGrantedApplicationImmunityQuery = !GrantedApplicationImmunityQuery.IsEmpty();
-	HasRemoveGameplayEffectsQuery = !RemoveGameplayEffectQuery.IsEmpty();
+		// We already have the Component that matches the Archetype's Component
+		if (bFound)
+		{
+			continue;
+		}
 
-#if WITH_EDITOR
-	SCALABLEFLOAT_REPORTERROR_WITHPOSTLOAD(Period);
-	SCALABLEFLOAT_REPORTERROR_WITHPOSTLOAD(ChanceToApplyToTarget);
-	DurationMagnitude.ReportErrors(GetPathName());
-#endif // WITH_EDITOR
+		// We don't already have the Archetype's Component, so add it here using the exact same name so we link up.
+		UE_LOG(LogGameplayEffects, Verbose, TEXT("%s is manually duplicating Archetype %s because they were not inherited through automatic instancing"), *GetFullNameSafe(this), *GetFullNameSafe(ParentComponent));
+		UGameplayEffectComponent* ChildComponent = DuplicateObject(ParentComponent, this, ParentComponentName);
+		GEComponents.Add(ChildComponent);
+	}
+}
+
+void UGameplayEffect::PostCDOCompiled(const FPostCDOCompiledContext& Context)
+{
+	Super::PostCDOCompiled(Context);
+	UE_LOG(LogGameplayEffects, VeryVerbose, TEXT("%s: Running Upgrade code. Current Version = %d"), *GetName(), GetVersion());
+
+	// Make sure everything is fixed up properly before going through upgrade code
+	PostCDOCompiledFixupSubobjects();
+
+	// Move any data from deprecated properties into components, create components as needed
+	ConvertAdditionalEffectsComponent();
+	ConvertAssetTagsComponent();
+	ConvertBlockByTagsComponent();
+	ConvertChanceToApplyComponent();
+	ConvertCustomCanApplyComponent();
+	ConvertImmunityComponent();
+	ConvertRemoveOtherComponent();
+	ConvertTagRequirementsComponent();
+	ConvertTargetTagsComponent();
+	ConvertUIComponent();
+}
+
+bool UGameplayEffect::VersionAllowsCreateGEComponentsDuringUpgrade() const
+{
+	return (GetVersion() < EGameplayEffectVersion::Modular53);
+}
+
+// Upgrade Code Section.  Lots of uses of deprecated variables (because we're upgrading them)
+// So just disable it throughout this section.
+PRAGMA_DISABLE_DEPRECATION_WARNINGS
+void UGameplayEffect::ConvertCustomCanApplyComponent()
+{
+	// Get the archetype to compare to
+	const UGameplayEffect* Archetype = CastChecked<UGameplayEffect>(GetArchetype());
+
+	const bool bChanged = !(ApplicationRequirements_DEPRECATED == Archetype->ApplicationRequirements_DEPRECATED);
+	if (bChanged && VersionAllowsCreateGEComponentsDuringUpgrade())
+	{
+		UCustomCanApplyGameplayEffectComponent& ApplicationComponent = FindOrAddComponent<UCustomCanApplyGameplayEffectComponent>();
+		ApplicationComponent.ApplicationRequirements.Append(ApplicationRequirements_DEPRECATED);
+	}
+
+	// Keep backwards compatibility (at least in terms of reading from the data)
+	const UCustomCanApplyGameplayEffectComponent* ApplicationComponent = FindComponent<UCustomCanApplyGameplayEffectComponent>();
+	if (ApplicationComponent)
+	{
+		ApplicationRequirements_DEPRECATED = ApplicationComponent->ApplicationRequirements;
+	}
+}
+
+void UGameplayEffect::ConvertChanceToApplyComponent()
+{
+	// Get the archetype to compare to
+	const UGameplayEffect* Archetype = CastChecked<UGameplayEffect>(GetArchetype());
+
+	bool bChanged = (ChanceToApplyToTarget_DEPRECATED != Archetype->ChanceToApplyToTarget_DEPRECATED);
+	bChanged = bChanged || !(ApplicationRequirements_DEPRECATED == Archetype->ApplicationRequirements_DEPRECATED);
+
+	if (bChanged && VersionAllowsCreateGEComponentsDuringUpgrade())
+	{
+		UChanceToApplyGameplayEffectComponent& ChanceToApplyComponent = FindOrAddComponent<UChanceToApplyGameplayEffectComponent>();
+		ChanceToApplyComponent.SetChanceToApplyToTarget(ChanceToApplyToTarget_DEPRECATED);
+	}
+
+	// Keep backwards compatibility (at least in terms of reading from the data)
+	const UChanceToApplyGameplayEffectComponent* ChanceToApplyComponent = FindComponent<UChanceToApplyGameplayEffectComponent>();
+	if (ChanceToApplyComponent)
+	{
+		ChanceToApplyToTarget_DEPRECATED = ChanceToApplyComponent->GetChanceToApplyToTarget();
+	}
+}
 
+void UGameplayEffect::ConvertAssetTagsComponent()
+{
+	// Get the archetype to compare to
+	const UGameplayEffect* Archetype = Cast<UGameplayEffect>(GetArchetype());
 
-	for (const TSubclassOf<UGameplayEffect>& ConditionalEffectClass : TargetEffectClasses_DEPRECATED)
+	bool bChanged = InheritableGameplayEffectTags.CombinedTags != Archetype->InheritableGameplayEffectTags.CombinedTags;
+	if (bChanged && VersionAllowsCreateGEComponentsDuringUpgrade())
 	{
-		FConditionalGameplayEffect ConditionalGameplayEffect;
-		ConditionalGameplayEffect.EffectClass = ConditionalEffectClass;
-		ConditionalGameplayEffects.Add(ConditionalGameplayEffect);
+		UAssetTagsGameplayEffectComponent& AssetTagsComponent = FindOrAddComponent<UAssetTagsGameplayEffectComponent>();
+		AssetTagsComponent.SetAndApplyAssetTagChanges(InheritableGameplayEffectTags);
 	}
-	TargetEffectClasses_DEPRECATED.Empty();
 
-	for (FGameplayEffectExecutionDefinition& Execution : Executions)
+	// Keep backwards compatibility (at least in terms of reading from the data)
+	const UAssetTagsGameplayEffectComponent* AssetTagsComponent = FindComponent<UAssetTagsGameplayEffectComponent>();
+	if (AssetTagsComponent)
 	{
-		for (const TSubclassOf<UGameplayEffect>& ConditionalEffectClass : Execution.ConditionalGameplayEffectClasses_DEPRECATED)
+		InheritableGameplayEffectTags = AssetTagsComponent->GetConfiguredAssetTagChanges();
+		InheritableGameplayEffectTags.UpdateInheritedTagProperties(&Archetype->InheritableGameplayEffectTags);
+	}
+}
+
+void UGameplayEffect::ConvertAdditionalEffectsComponent()
+{
+	// Get the archetype to compare to
+	const UGameplayEffect* Archetype = Cast<UGameplayEffect>(GetArchetype());
+
+	bool bConditionalChanged = (ConditionalGameplayEffects.Num() != Archetype->ConditionalGameplayEffects.Num());
+	bool bPrematureChanged = (PrematureExpirationEffectClasses != Archetype->PrematureExpirationEffectClasses);
+	bool bRoutineChanged = (RoutineExpirationEffectClasses != Archetype->RoutineExpirationEffectClasses);
+
+	bool bChanged = bConditionalChanged || bPrematureChanged || bRoutineChanged;
+	if (bChanged && VersionAllowsCreateGEComponentsDuringUpgrade())
+	{
+		UAdditionalEffectsGameplayEffectComponent& ConditionalEffectsComponent = FindOrAddComponent<UAdditionalEffectsGameplayEffectComponent>();
+
+		if (bConditionalChanged)
+		{
+			ConditionalEffectsComponent.OnApplicationGameplayEffects = ConditionalGameplayEffects;
+		}
+
+		if (bPrematureChanged)
 		{
-			FConditionalGameplayEffect ConditionalGameplayEffect;
-			ConditionalGameplayEffect.EffectClass = ConditionalEffectClass;
-			Execution.ConditionalGameplayEffects.Add(ConditionalGameplayEffect);
+			ConditionalEffectsComponent.OnCompletePrematurely = PrematureExpirationEffectClasses;
 		}
 
-		Execution.ConditionalGameplayEffectClasses_DEPRECATED.Empty();
+		if (bRoutineChanged)
+		{
+			ConditionalEffectsComponent.OnCompleteNormal = RoutineExpirationEffectClasses;
+		}
+	}
+
+	// Keep backwards compatibility (at least in terms of reading from the data)
+	const UAdditionalEffectsGameplayEffectComponent* ConditionalEffectsComponent = FindComponent<UAdditionalEffectsGameplayEffectComponent>();
+	if (ConditionalEffectsComponent)
+	{
+		ConditionalGameplayEffects = ConditionalEffectsComponent->OnApplicationGameplayEffects;
+		PrematureExpirationEffectClasses = ConditionalEffectsComponent->OnCompletePrematurely;
+		RoutineExpirationEffectClasses = ConditionalEffectsComponent->OnCompleteNormal;
 	}
 }
 
-void UGameplayEffect::PostInitProperties()
+void UGameplayEffect::ConvertBlockByTagsComponent()
 {
-	Super::PostInitProperties();
+	// Get the archetype to compare to
+	const UGameplayEffect* Archetype = Cast<UGameplayEffect>(GetArchetype());
 
-	InheritableGameplayEffectTags.PostInitProperties();
-	InheritableOwnedTagsContainer.PostInitProperties();
-	InheritableBlockedAbilityTagsContainer.PostInitProperties();
-	RemoveGameplayEffectsWithTags.PostInitProperties();
+	bool bChanged = InheritableBlockedAbilityTagsContainer.CombinedTags != Archetype->InheritableBlockedAbilityTagsContainer.CombinedTags;
+	if (bChanged && VersionAllowsCreateGEComponentsDuringUpgrade())
+	{
+		UBlockAbilityTagsGameplayEffectComponent& BlockByTagsComponent = FindOrAddComponent<UBlockAbilityTagsGameplayEffectComponent>();
+		BlockByTagsComponent.SetAndApplyBlockedAbilityTagChanges(InheritableBlockedAbilityTagsContainer);
+	}
+
+	// Keep backwards compatibility (at least in terms of reading from the data)
+	const UBlockAbilityTagsGameplayEffectComponent* BlockByTagsComponent = FindComponent<UBlockAbilityTagsGameplayEffectComponent>();
+	if (BlockByTagsComponent)
+	{
+		InheritableBlockedAbilityTagsContainer = BlockByTagsComponent->GetConfiguredBlockedAbilityTagChanges();
+		InheritableBlockedAbilityTagsContainer.UpdateInheritedTagProperties(&Archetype->InheritableBlockedAbilityTagsContainer);
+	}
 }
 
-#if WITH_EDITOR
+void UGameplayEffect::ConvertTargetTagsComponent()
+{
+	// Get the archetype to compare to
+	const UGameplayEffect* Archetype = Cast<UGameplayEffect>(GetArchetype());
 
-void UGameplayEffect::PostEditChangeProperty(FPropertyChangedEvent& PropertyChangedEvent)
+	bool bChanged = InheritableOwnedTagsContainer.CombinedTags != Archetype->InheritableOwnedTagsContainer.CombinedTags;
+	if (bChanged && VersionAllowsCreateGEComponentsDuringUpgrade())
+	{
+		UTargetTagsGameplayEffectComponent& TargetTagsComponent = FindOrAddComponent<UTargetTagsGameplayEffectComponent>();
+		TargetTagsComponent.SetAndApplyTargetTagChanges(InheritableOwnedTagsContainer);
+	}
+
+	// Keep backwards compatibility (at least in terms of reading from the data)
+	const UTargetTagsGameplayEffectComponent* TargetTagsComponent = FindComponent<UTargetTagsGameplayEffectComponent>();
+	if (TargetTagsComponent)
+	{
+		InheritableOwnedTagsContainer = TargetTagsComponent->GetConfiguredTargetTagChanges();
+		InheritableOwnedTagsContainer.UpdateInheritedTagProperties(&Archetype->InheritableOwnedTagsContainer);
+	}
+}
+
+void UGameplayEffect::ConvertImmunityComponent()
 {
-	Super::PostEditChangeProperty(PropertyChangedEvent);
+	// Get the archetype to compare to
+	const UGameplayEffect* Archetype = Cast<UGameplayEffect>(GetArchetype());
 
-	const FProperty* PropertyThatChanged = PropertyChangedEvent.MemberProperty;
-	if (PropertyThatChanged)
+	const bool bTagsChanged = GrantedApplicationImmunityTags != Archetype->GrantedApplicationImmunityTags;
+	const bool bQueryChanged (GrantedApplicationImmunityQuery != Archetype->GrantedApplicationImmunityQuery);
+
+	const bool bChanged = bTagsChanged || bQueryChanged;
+	if (bChanged && VersionAllowsCreateGEComponentsDuringUpgrade())
 	{
-		UGameplayEffect* Parent = Cast<UGameplayEffect>(GetClass()->GetSuperClass()->GetDefaultObject());
-		FName PropName = PropertyThatChanged->GetFName();
-		if (PropName == GET_MEMBER_NAME_CHECKED(UGameplayEffect, InheritableGameplayEffectTags))
+		UImmunityGameplayEffectComponent& ImmunityComponent = FindOrAddComponent<UImmunityGameplayEffectComponent>();
+
+		// Our upgrade path is going to have to be conservative because of the Archetype inheritance rules.
+		// Let's always make sure there's 2 slots (originating Tags and originating Queries).
+		while (ImmunityComponent.ImmunityQueries.Num() < 2)
+		{
+			ImmunityComponent.ImmunityQueries.AddDefaulted();
+		}
+
+		if (bTagsChanged)
+		{
+			// Build up a query that matches the equivalent GrantedApplicationImmunityTags
+			FGameplayEffectQuery TagQuery;
+			TagQuery.SourceTagQuery.BuildQuery(
+				FGameplayTagQueryExpression().AnyTagsMatch().AddTags(GrantedApplicationImmunityTags.RequireTags).NoTagsMatch().AddTags(GrantedApplicationImmunityTags.IgnoreTags),
+				TEXT("GrantedApplicationImmunityTags")
+			);
+
+			ImmunityComponent.ImmunityQueries[0] = MoveTemp(TagQuery);
+		}
+
+		if (bQueryChanged)
 		{
-			InheritableGameplayEffectTags.UpdateInheritedTagProperties(Parent ? &Parent->InheritableGameplayEffectTags : NULL);
+			ImmunityComponent.ImmunityQueries[1] = GrantedApplicationImmunityQuery;
 		}
-		else if (PropName == GET_MEMBER_NAME_CHECKED(UGameplayEffect, InheritableOwnedTagsContainer))
+	}
+
+	// Keep backwards compatibility (at least in terms of reading from the data)
+	// Note: This isn't quite right because we've removed the Tags and converted it into a Query.
+	const UImmunityGameplayEffectComponent* ImmunityComponent = FindComponent<UImmunityGameplayEffectComponent>();
+	if (ImmunityComponent && !ImmunityComponent->ImmunityQueries.IsEmpty())
+	{
+		GrantedApplicationImmunityQuery = ImmunityComponent->ImmunityQueries.Last();
+	}
+}
+
+void UGameplayEffect::ConvertRemoveOtherComponent()
+{
+	// Get the archetype to compare to
+	const UGameplayEffect* Archetype = Cast<UGameplayEffect>(GetArchetype());
+
+	const bool bTagsChanged = RemoveGameplayEffectsWithTags != Archetype->RemoveGameplayEffectsWithTags;
+	const bool bQueryChanged = RemoveGameplayEffectQuery != Archetype->RemoveGameplayEffectQuery;
+
+	const bool bChanged = bTagsChanged || bQueryChanged;
+	if (bChanged && VersionAllowsCreateGEComponentsDuringUpgrade())
+	{
+		URemoveOtherGameplayEffectComponent& RemoveOtherComponent = FindOrAddComponent<URemoveOtherGameplayEffectComponent>();
+		
+		// Our upgrade path is going to have to be conservative because of the Archetype inheritance rules.
+		// Let's always make sure there's 2 slots (originating Tags and originating Queries).
+		while (RemoveOtherComponent.RemoveGameplayEffectQueries.Num() < 2)
 		{
-			InheritableOwnedTagsContainer.UpdateInheritedTagProperties(Parent ? &Parent->InheritableOwnedTagsContainer : NULL);
+			RemoveOtherComponent.RemoveGameplayEffectQueries.AddDefaulted();
 		}
-		else if (PropName == GET_MEMBER_NAME_CHECKED(UGameplayEffect, InheritableBlockedAbilityTagsContainer))
+
+		if (bTagsChanged)
 		{
-			InheritableBlockedAbilityTagsContainer.UpdateInheritedTagProperties(Parent ? &Parent->InheritableBlockedAbilityTagsContainer : nullptr);
+			RemoveOtherComponent.RemoveGameplayEffectQueries[0] = FGameplayEffectQuery::MakeQuery_MatchAnyOwningTags(RemoveGameplayEffectsWithTags.CombinedTags);
 		}
-		else if (PropName == GET_MEMBER_NAME_CHECKED(UGameplayEffect, RemoveGameplayEffectsWithTags))
+
+		if (bQueryChanged)
 		{
-			RemoveGameplayEffectsWithTags.UpdateInheritedTagProperties(Parent ? &Parent->RemoveGameplayEffectsWithTags : NULL);
+			RemoveOtherComponent.RemoveGameplayEffectQueries[1] = RemoveGameplayEffectQuery;
 		}
 	}
 
-	HasGrantedApplicationImmunityQuery = !GrantedApplicationImmunityQuery.IsEmpty();
-	HasRemoveGameplayEffectsQuery = !RemoveGameplayEffectQuery.IsEmpty();
+	// Keep backwards compatibility (at least in terms of reading from the data)
+	// Note: This isn't quite right because we've removed the Tags and converted it into a Query.
+	const URemoveOtherGameplayEffectComponent* RemoveOtherComponent = FindComponent<URemoveOtherGameplayEffectComponent>();
+	if (RemoveOtherComponent && !RemoveOtherComponent->RemoveGameplayEffectQueries.IsEmpty())
+	{
+		RemoveGameplayEffectQuery = RemoveOtherComponent->RemoveGameplayEffectQueries.Last();
+	}
+}
+
+void UGameplayEffect::ConvertTagRequirementsComponent()
+{
+	// Get the archetype to compare to
+	const UGameplayEffect* Archetype = Cast<UGameplayEffect>(GetArchetype());
+
+	bool bApplicationChanged = ApplicationTagRequirements != Archetype->ApplicationTagRequirements;
+	bool bRemovalChanged = RemovalTagRequirements != Archetype->RemovalTagRequirements;
+	bool bOngoingChanged = OngoingTagRequirements != Archetype->OngoingTagRequirements;
+
+	const bool bChanged = bApplicationChanged || bRemovalChanged || bOngoingChanged;
+	if (bChanged && VersionAllowsCreateGEComponentsDuringUpgrade())
+	{
+		UTargetTagRequirementsGameplayEffectComponent& TagRequirementsComponent = FindOrAddComponent<UTargetTagRequirementsGameplayEffectComponent>();
+
+		if (bApplicationChanged)
+		{
+			TagRequirementsComponent.ApplicationTagRequirements = ApplicationTagRequirements;
+		}
+
+		if (bRemovalChanged)
+		{
+			TagRequirementsComponent.RemovalTagRequirements = RemovalTagRequirements;
+		}
+
+		if (bOngoingChanged)
+		{
+			TagRequirementsComponent.OngoingTagRequirements = OngoingTagRequirements;
+		}
+	}
+
+	// Keep backwards compatibility (at least in terms of reading from the data)
+	const UTargetTagRequirementsGameplayEffectComponent* TagRequirementsComponent = FindComponent<UTargetTagRequirementsGameplayEffectComponent>();
+	if (TagRequirementsComponent)
+	{
+		ApplicationTagRequirements = TagRequirementsComponent->ApplicationTagRequirements;
+		RemovalTagRequirements = TagRequirementsComponent->RemovalTagRequirements;
+		OngoingTagRequirements = TagRequirementsComponent->OngoingTagRequirements;
+	}
+}
+
+void UGameplayEffect::ConvertUIComponent()
+{
+	// Get the archetype to compare to
+	const UGameplayEffect* Archetype = Cast<UGameplayEffect>(GetArchetype());
+
+	const bool bChanged = (UIData != Archetype->UIData);
+	if (bChanged && VersionAllowsCreateGEComponentsDuringUpgrade())
+	{
+		const UGameplayEffectUIData* UIComponent = FindComponent<UGameplayEffectUIData>();
+		if (!UIComponent && UIData)
+		{
+			GEComponents.Add(MoveTemp(UIData));
+			UIData = nullptr;
+		}
+	}
 
+	// Note: We cannot keep backwards compatibility.  Since UGameplayEffectUIData (and UIData by extension) are instanced,
+	// we would end up with multiple references to a single GEComponent.
+	if (UIData)
+	{
+		UIData = nullptr;
+	}
+}
+PRAGMA_ENABLE_DEPRECATION_WARNINGS
+
+/** Let's keep track of the version so that we can upgrade the Components properly.  This will be set properly in PostLoad after upgrades. */
+EGameplayEffectVersion UGameplayEffect::GetVersion() const
+{
+	return DataVersion.CurrentVersion;
+}
+
+/** Sets the Version of the class to denote it's been upgraded */
+void UGameplayEffect::SetVersion(EGameplayEffectVersion Version)
+{
+	DataVersion.CurrentVersion = Version;
+}
+
+void UGameplayEffect::PostEditChangeProperty(FPropertyChangedEvent& PropertyChangedEvent)
+{
+	Super::PostEditChangeProperty(PropertyChangedEvent);
 	UAbilitySystemGlobals::Get().GameplayEffectPostEditChangeProperty(this, PropertyChangedEvent);
+
+	// Update the Status Text
+	FDataValidationContext DataValidationContext;
+	IsDataValid(DataValidationContext);
 }
 
 #endif // #if WITH_EDITOR
 
-void UGameplayEffect::PreSave(const class ITargetPlatform* TargetPlatform)
+bool UGameplayEffect::CanApply(const FActiveGameplayEffectsContainer& ActiveGEContainer, const FGameplayEffectSpec& GESpec) const
 {
-	PRAGMA_DISABLE_DEPRECATION_WARNINGS;
-	Super::PreSave(TargetPlatform);
-	PRAGMA_ENABLE_DEPRECATION_WARNINGS;
+	for (const UGameplayEffectComponent* GEComponent : GEComponents)
+	{
+		if (GEComponent && !GEComponent->CanGameplayEffectApply(ActiveGEContainer, GESpec))
+		{
+			UE_VLOG(ActiveGEContainer.Owner, LogAbilitySystem, Verbose, TEXT("%s could not apply. Blocked by %s"), *GetNameSafe(GESpec.Def), *GetNameSafe(GEComponent));
+			return false;
+		}
+	}
+
+	return true;
 }
 
-void UGameplayEffect::PreSave(FObjectPreSaveContext ObjectSaveContext)
+bool UGameplayEffect::OnAddedToActiveContainer(FActiveGameplayEffectsContainer& ActiveGEContainer, FActiveGameplayEffect& ActiveGE) const
 {
-	Super::PreSave(ObjectSaveContext);
-	HasGrantedApplicationImmunityQuery = !GrantedApplicationImmunityQuery.IsEmpty();
-	HasRemoveGameplayEffectsQuery = !RemoveGameplayEffectQuery.IsEmpty();
+	UE_VLOG(ActiveGEContainer.Owner, LogAbilitySystem, Verbose, TEXT("ActiveGameplayEffect %s added"), *ActiveGE.GetDebugString());
+
+	bool bShouldBeActive = true;
+	for (const UGameplayEffectComponent* GEComponent : GEComponents)
+	{
+		if (GEComponent)
+		{
+			bShouldBeActive = GEComponent->OnActiveGameplayEffectAdded(ActiveGEContainer, ActiveGE) && bShouldBeActive;
+		}
+	}
+
+	return bShouldBeActive;
 }
 
-void UGameplayEffect::UpdateInheritedTagProperties()
+void UGameplayEffect::OnExecuted(FActiveGameplayEffectsContainer& ActiveGEContainer, FGameplayEffectSpec& GESpec, FPredictionKey& PredictionKey) const
 {
-	UGameplayEffect* Parent = Cast<UGameplayEffect>(GetClass()->GetSuperClass()->GetDefaultObject());
+	UE_VLOG(ActiveGEContainer.Owner, LogAbilitySystem, Verbose, TEXT("GameplayEffect %s executed"), *GetNameSafe(GESpec.Def));
 
-	InheritableGameplayEffectTags.UpdateInheritedTagProperties(Parent ? &Parent->InheritableGameplayEffectTags : NULL);
-	InheritableOwnedTagsContainer.UpdateInheritedTagProperties(Parent ? &Parent->InheritableOwnedTagsContainer : NULL);
-	InheritableBlockedAbilityTagsContainer.UpdateInheritedTagProperties(Parent ? &Parent->InheritableBlockedAbilityTagsContainer : nullptr);
-	RemoveGameplayEffectsWithTags.UpdateInheritedTagProperties(Parent ? &Parent->RemoveGameplayEffectsWithTags : NULL);
+	for (const UGameplayEffectComponent* GEComponent : GEComponents)
+	{
+		if (GEComponent)
+		{
+			GEComponent->OnGameplayEffectExecuted(ActiveGEContainer, GESpec, PredictionKey);
+		}
+	}
 }
 
-void UGameplayEffect::ValidateGameplayEffect()
+int32 UGameplayEffect::GetStackLimitCount() const
 {
+	return StackLimitCount;
+}
+
+EGameplayEffectStackingExpirationPolicy UGameplayEffect::GetStackExpirationPolicy() const
+{
+	return StackExpirationPolicy;
+}
+
+const UGameplayEffectComponent* UGameplayEffect::FindComponent(TSubclassOf<UGameplayEffectComponent> ClassToFind) const
+{
+	for (const TObjectPtr<UGameplayEffectComponent>& GEComponent : GEComponents)
+	{
+		if (GEComponent && GEComponent->IsA(ClassToFind))
+		{
+			return GEComponent;
+		}
+	}
+
+	return nullptr;
 }
 
 // --------------------------------------------------------------------------------------------------------------------------------------------------------
@@ -657,6 +1141,16 @@ FGameplayEffectSpecHandle FConditionalGameplayEffect::CreateSpec(FGameplayEffect
 	return EffectCDO ? FGameplayEffectSpecHandle(new FGameplayEffectSpec(EffectCDO, EffectContext, SourceLevel)) : FGameplayEffectSpecHandle();
 }
 
+bool FConditionalGameplayEffect::operator==(const FConditionalGameplayEffect& Other) const
+{
+	return EffectClass == Other.EffectClass && RequiredSourceTags == Other.RequiredSourceTags;
+}
+
+bool FConditionalGameplayEffect::operator!=(const FConditionalGameplayEffect& Other) const
+{
+	return !(*this == Other);
+}
+
 
 // --------------------------------------------------------------------------------------------------------------------------------------------------------
 //
@@ -668,7 +1162,6 @@ FGameplayEffectSpec::FGameplayEffectSpec()
 	: Def(nullptr)
 	, Duration(UGameplayEffect::INSTANT_APPLICATION)
 	, Period(UGameplayEffect::NO_PERIOD)
-	, ChanceToApplyToTarget(1.f)
 	, StackCount(1)
 	, bCompletedSourceAttributeCapture(false)
 	, bCompletedTargetAttributeCapture(false)
@@ -681,7 +1174,6 @@ FGameplayEffectSpec::FGameplayEffectSpec(const UGameplayEffect* InDef, const FGa
 	: Def(InDef)
 	, Duration(UGameplayEffect::INSTANT_APPLICATION)
 	, Period(UGameplayEffect::NO_PERIOD)
-	, ChanceToApplyToTarget(1.f)
 	, StackCount(1)
 	, bCompletedSourceAttributeCapture(false)
 	, bCompletedTargetAttributeCapture(false)
@@ -709,7 +1201,6 @@ FGameplayEffectSpec::FGameplayEffectSpec(FGameplayEffectSpec&& Other)
 	, TargetEffectSpecs(MoveTemp(Other.TargetEffectSpecs))
 	, Duration(Other.Duration)
 	, Period(Other.Period)
-	, ChanceToApplyToTarget(Other.ChanceToApplyToTarget)
 	, CapturedSourceTags(MoveTemp(Other.CapturedSourceTags))
 	, CapturedTargetTags(MoveTemp(Other.CapturedTargetTags))
 	, DynamicGrantedTags(MoveTemp(Other.DynamicGrantedTags))
@@ -738,7 +1229,6 @@ FGameplayEffectSpec& FGameplayEffectSpec::operator=(FGameplayEffectSpec&& Other)
 	TargetEffectSpecs = MoveTemp(Other.TargetEffectSpecs);
 	Duration = Other.Duration;
 	Period = Other.Period;
-	ChanceToApplyToTarget = Other.ChanceToApplyToTarget;
 	CapturedSourceTags = MoveTemp(Other.CapturedSourceTags);
 	CapturedTargetTags = MoveTemp(Other.CapturedTargetTags);
 	DynamicGrantedTags = MoveTemp(Other.DynamicGrantedTags);
@@ -753,6 +1243,7 @@ FGameplayEffectSpec& FGameplayEffectSpec::operator=(FGameplayEffectSpec&& Other)
 	SetByCallerTagMagnitudes = MoveTemp(Other.SetByCallerTagMagnitudes);
 	EffectContext = Other.EffectContext;
 	Level = Other.Level;
+
 	return *this;
 }
 PRAGMA_ENABLE_DEPRECATION_WARNINGS
@@ -766,7 +1257,6 @@ FGameplayEffectSpec& FGameplayEffectSpec::operator=(const FGameplayEffectSpec& O
 	TargetEffectSpecs = Other.TargetEffectSpecs;
 	Duration = Other.Duration;
 	Period = Other.Period;
-	ChanceToApplyToTarget = Other.ChanceToApplyToTarget;
 	CapturedSourceTags = Other.CapturedSourceTags;
 	CapturedTargetTags = Other.CapturedTargetTags;
 	DynamicGrantedTags = Other.DynamicGrantedTags;
@@ -781,6 +1271,7 @@ FGameplayEffectSpec& FGameplayEffectSpec::operator=(const FGameplayEffectSpec& O
 	SetByCallerTagMagnitudes = Other.SetByCallerTagMagnitudes;
 	EffectContext = Other.EffectContext;
 	Level = Other.Level;
+
 	return *this;
 }
 PRAGMA_ENABLE_DEPRECATION_WARNINGS
@@ -821,6 +1312,7 @@ void FGameplayEffectSpec::Initialize(const UGameplayEffect* InDef, const FGamepl
 {
 	Def = InDef;
 	check(Def);	
+
 	// SetContext requires the level to be set before it runs
 	// however, there are code paths in SetLevel that can potentially (depends on game data setup) require the context to be set
 	Level = InLevel;
@@ -834,28 +1326,11 @@ void FGameplayEffectSpec::Initialize(const UGameplayEffect* InDef, const FGamepl
 	SetupAttributeCaptureDefinitions();
 	
 	// Add the GameplayEffect asset tags to the source Spec tags
-	CapturedSourceTags.GetSpecTags().AppendTags(InDef->InheritableGameplayEffectTags.CombinedTags);
+	CapturedSourceTags.GetSpecTags().AppendTags(Def->GetAssetTags());
 
-	// Prepare source tags before accessing them in ConditionalGameplayEffects
+	// Prepare source tags before accessing them in the Components
 	CaptureDataFromSource();
 
-	// ------------------------------------------------
-	//	Linked/Dependant Specs
-	// ------------------------------------------------
-
-
-	for (const FConditionalGameplayEffect& ConditionalEffect : InDef->ConditionalGameplayEffects)
-	{
-		if (ConditionalEffect.CanApply(CapturedSourceTags.GetActorTags(), InLevel))
-		{
-			FGameplayEffectSpecHandle SpecHandle = ConditionalEffect.CreateSpec(EffectContext, InLevel);
-			if (SpecHandle.IsValid())
-			{
-				TargetEffectSpecs.Add(SpecHandle);
-			}
-		}
-	}
-
 	// ------------------------------------------------
 	//	Granted Abilities
 	// ------------------------------------------------
@@ -886,7 +1361,7 @@ void FGameplayEffectSpec::InitializeFromLinkedSpec(const UGameplayEffect* InDef,
 	CapturedSourceTags = OriginalSpec.CapturedSourceTags;
 
 	// But then remove the tags the originating GE added
-	CapturedSourceTags.GetSpecTags().RemoveTags( OriginalSpec.Def->InheritableGameplayEffectTags.CombinedTags );
+	CapturedSourceTags.GetSpecTags().RemoveTags(OriginalSpec.Def->GetAssetTags());
 
 	// Now initialize like the normal cstor would have. Note that this will add the new GE's asset tags (in case they were removed in the line above / e.g., shared asset tags with the originating GE)					
 	Initialize(InDef, NewContextHandle, OriginalSpec.GetLevel());
@@ -1033,7 +1508,6 @@ void FGameplayEffectSpec::SetLevel(float InLevel)
 
 		FString ContextString = Def->GetName();
 		Period = Def->Period.GetValueAtLevel(InLevel, &ContextString);
-		ChanceToApplyToTarget = Def->ChanceToApplyToTarget.GetValueAtLevel(InLevel, &ContextString);
 	}
 }
 
@@ -1113,11 +1587,19 @@ float FGameplayEffectSpec::GetPeriod() const
 	return Period;
 }
 
-float FGameplayEffectSpec::GetChanceToApplyToTarget() const
+PRAGMA_DISABLE_DEPRECATION_WARNINGS
+void FGameplayEffectSpec::SetStackCount(int32 NewStackCount)
 {
-	return ChanceToApplyToTarget;
+
+	StackCount = NewStackCount;
 }
 
+int32 FGameplayEffectSpec::GetStackCount() const
+{
+	return StackCount;
+}
+PRAGMA_ENABLE_DEPRECATION_WARNINGS
+
 float FGameplayEffectSpec::GetModifierMagnitude(int32 ModifierIdx, bool bFactorInStackCount) const
 {
 	check(Modifiers.IsValidIndex(ModifierIdx) && Def && Def->Modifiers.IsValidIndex(ModifierIdx));
@@ -1127,7 +1609,7 @@ float FGameplayEffectSpec::GetModifierMagnitude(int32 ModifierIdx, bool bFactorI
 	float ModMagnitude = SingleEvaluatedMagnitude;
 	if (bFactorInStackCount)
 	{
-		ModMagnitude = GameplayEffectUtilities::ComputeStackedModifierMagnitude(SingleEvaluatedMagnitude, StackCount, Def->Modifiers[ModifierIdx].ModifierOp);
+		ModMagnitude = GameplayEffectUtilities::ComputeStackedModifierMagnitude(SingleEvaluatedMagnitude, GetStackCount(), Def->Modifiers[ModifierIdx].ModifierOp);
 	}
 
 	return ModMagnitude;
@@ -1230,33 +1712,32 @@ void FGameplayEffectSpec::SetContext(FGameplayEffectContextHandle NewEffectConte
 	}
 }
 
-void FGameplayEffectSpec::GetAllGrantedTags(OUT FGameplayTagContainer& Container) const
+void FGameplayEffectSpec::GetAllGrantedTags(OUT FGameplayTagContainer& OutContainer) const
 {
-	Container.AppendTags(DynamicGrantedTags);
+	OutContainer.AppendTags(DynamicGrantedTags);
 	if (Def)
 	{
-		Container.AppendTags(Def->InheritableOwnedTagsContainer.CombinedTags);
+		OutContainer.AppendTags(Def->GetGrantedTags());
 	}
 }
 
-void FGameplayEffectSpec::GetAllBlockedAbilityTags(FGameplayTagContainer& OutContainer) const
+void FGameplayEffectSpec::GetAllBlockedAbilityTags(OUT FGameplayTagContainer& OutContainer) const
 {
 	if (Def)
 	{
-		OutContainer.AppendTags(Def->InheritableBlockedAbilityTagsContainer.CombinedTags);
+		OutContainer.AppendTags(Def->GetBlockedAbilityTags());
 	}
 }
 
-void FGameplayEffectSpec::GetAllAssetTags(OUT FGameplayTagContainer& Container) const
+void FGameplayEffectSpec::GetAllAssetTags(OUT FGameplayTagContainer& OutContainer) const
 {
-	Container.AppendTags(GetDynamicAssetTags());
+	OutContainer.AppendTags(GetDynamicAssetTags());
 	if (Def)
 	{
-		Container.AppendTags(Def->InheritableGameplayEffectTags.CombinedTags);
+		OutContainer.AppendTags(Def->GetAssetTags());
 	}
 }
 
-
 void FGameplayEffectSpec::SetSetByCallerMagnitude(FName DataName, float Magnitude)
 {
 	if (DataName != NAME_None)
@@ -1747,49 +2228,13 @@ FActiveGameplayEffect& FActiveGameplayEffect::operator=(const FActiveGameplayEff
 	return *this;
 }
 
-/** This is the core function that turns the ActiveGE 'on' or 'off */
 void FActiveGameplayEffect::CheckOngoingTagRequirements(const FGameplayTagContainer& OwnerTags, FActiveGameplayEffectsContainer& OwningContainer, bool bInvokeGameplayCueEvents)
 {
-	bool bShouldBeInhibited = !Spec.Def->OngoingTagRequirements.RequirementsMet(OwnerTags);
-
-	if (bIsInhibited != bShouldBeInhibited)
-	{
-		{
-			// All OnDirty callbacks must be inhibited until we update this entire GameplayEffect.
-			FScopedAggregatorOnDirtyBatch	AggregatorOnDirtyBatcher;
 
-			// Important to set this prior to adding or removing, so that any delegates that are triggered can query accurately against this GE
-			bIsInhibited = bShouldBeInhibited;
-
-			if (bShouldBeInhibited)
-			{
-				// Remove our ActiveGameplayEffects modifiers with our Attribute Aggregators
-				OwningContainer.RemoveActiveGameplayEffectGrantedTagsAndModifiers(*this, bInvokeGameplayCueEvents);
-			}
-			else
-			{
-				OwningContainer.AddActiveGameplayEffectGrantedTagsAndModifiers(*this, bInvokeGameplayCueEvents);
-			}
-		}
-		EventSet.OnInhibitionChanged.Broadcast(Handle, bIsInhibited);
-	}
 }
 
 bool FActiveGameplayEffect::CheckRemovalTagRequirements(const FGameplayTagContainer& OwnerTags, FActiveGameplayEffectsContainer& OwningContainer) const
 {
-	if (Spec.Def == nullptr)
-	{
-		ABILITY_LOG(Error, TEXT("FActiveGameplayEffect::CheckRemovalTagRequirements called with no UGameplayEffect def. (%s)"), *Spec.GetEffectContext().ToString());
-		return false;
-	}
-
-	if (!Spec.Def->RemovalTagRequirements.IsEmpty())
-	{
-		if (Spec.Def->RemovalTagRequirements.RequirementsMet(OwnerTags))
-		{
-			return true;
-		}
-	}
 	return false;
 }
 
@@ -1804,6 +2249,7 @@ void FActiveGameplayEffect::PreReplicatedRemove(const struct FActiveGameplayEffe
 	ABILITY_LOG(Verbose, TEXT("PreReplicatedRemove: %s %s Marked as Pending Remove: %s"), *Handle.ToString(), *Spec.Def->GetName(), IsPendingRemove ? TEXT("TRUE") : TEXT("FALSE"));
 
 	FGameplayEffectRemovalInfo GameplayEffectRemovalInfo;
+	GameplayEffectRemovalInfo.ActiveEffect = this;
 	GameplayEffectRemovalInfo.StackCount = ClientCachedStackCount;
 	//Check duration to set bPrematureRemoval as req.
 	if (DurationHandle.IsValid())
@@ -1879,7 +2325,7 @@ void FActiveGameplayEffect::PostReplicatedAdd(const struct FActiveGameplayEffect
 	}
 
 	// Cache off StackCount
-	ClientCachedStackCount = Spec.StackCount;
+	ClientCachedStackCount = Spec.GetStackCount();
 
 	// Handles are not replicated, so create a new one.
 	Handle = FActiveGameplayEffectHandle::GenerateNewHandle(InArray.Owner);
@@ -1913,12 +2359,13 @@ void FActiveGameplayEffect::PostReplicatedChange(const struct FActiveGameplayEff
 		const_cast<FActiveGameplayEffectsContainer&>(InArray).OnDurationChange(*this);
 	}
 	
-	if (ClientCachedStackCount != Spec.StackCount)
+	int32 StackCount = Spec.GetStackCount();
+	if (ClientCachedStackCount != StackCount)
 	{
 		// If its a stack count change, we just call OnStackCountChange and it will broadcast delegates and update attribute aggregators
 		// Const cast is ok. It is there to prevent mutation of the GameplayEffects array, which this wont do.
-		const_cast<FActiveGameplayEffectsContainer&>(InArray).OnStackCountChange(*this, ClientCachedStackCount, Spec.StackCount);
-		ClientCachedStackCount = Spec.StackCount;
+		const_cast<FActiveGameplayEffectsContainer&>(InArray).OnStackCountChange(*this, ClientCachedStackCount, StackCount);
+		ClientCachedStackCount = StackCount;
 	}
 	else
 	{
@@ -1961,6 +2408,7 @@ FActiveGameplayEffectsContainer::FActiveGameplayEffectsContainer()
 	SetDeltaSerializationEnabled(true);
 }
 
+PRAGMA_DISABLE_DEPRECATION_WARNINGS
 FActiveGameplayEffectsContainer::~FActiveGameplayEffectsContainer()
 {
 	FActiveGameplayEffect* PendingGameplayEffect = PendingGameplayEffectHead;
@@ -1971,6 +2419,7 @@ FActiveGameplayEffectsContainer::~FActiveGameplayEffectsContainer()
 		PendingGameplayEffectHead =	Next;
 	}
 }
+PRAGMA_ENABLE_DEPRECATION_WARNINGS
 
 void FActiveGameplayEffectsContainer::RegisterWithOwner(UAbilitySystemComponent* InOwner)
 {
@@ -1978,10 +2427,6 @@ void FActiveGameplayEffectsContainer::RegisterWithOwner(UAbilitySystemComponent*
 	{
 		Owner = InOwner;
 		OwnerIsNetAuthority = Owner->IsOwnerActorAuthoritative();
-
-		// Binding raw is ok here, since the owner is literally the UObject that owns us. If we are destroyed, its because that uobject is destroyed,
-		// and if that is destroyed, the delegate wont be able to fire.
-		Owner->RegisterGenericGameplayTagEvent().AddRaw(this, &FActiveGameplayEffectsContainer::OnOwnerTagChange);
 	}
 }
 
@@ -2056,7 +2501,7 @@ void FActiveGameplayEffectsContainer::PredictivelyExecuteEffectSpec(FGameplayEff
 			bExecutionsProducedOutModifiers |= OutModifiers.Num() > 0;
 
 			const bool bApplyStackCountToEmittedMods = !ExecutionOutput.IsStackCountHandledManually();
-			const int32 SpecStackCount = SpecToUse.StackCount;
+			const int32 SpecStackCount = SpecToUse.GetStackCount();
 
 			for (FGameplayModifierEvaluatedData& CurExecMod : OutModifiers)
 			{
@@ -2172,7 +2617,7 @@ void FActiveGameplayEffectsContainer::ExecuteActiveEffectsFrom(FGameplayEffectSp
 			bExecutionsProducedOutModifiers |= OutModifiers.Num() > 0;
 
 			const bool bApplyStackCountToEmittedMods = !ExecutionOutput.IsStackCountHandledManually();
-			const int32 SpecStackCount = SpecToUse.StackCount;
+			const int32 SpecStackCount = SpecToUse.GetStackCount();
 
 			for (FGameplayModifierEvaluatedData& CurExecMod : OutModifiers)
 			{
@@ -2241,6 +2686,9 @@ void FActiveGameplayEffectsContainer::ExecuteActiveEffectsFrom(FGameplayEffectSp
 			Owner->ApplyGameplayEffectSpecToSelf(*TargetSpec.Data.Get(), PredictionKey);
 		}
 	}
+
+	// Notify the Gameplay Effect that it's been executed.  This will allow the GameplayEffectComponents to augment behavior.
+	Spec.Def->OnExecuted(*this, Spec, PredictionKey);
 }
 
 void FActiveGameplayEffectsContainer::ExecutePeriodicGameplayEffect(FActiveGameplayEffectHandle Handle)
@@ -2437,12 +2885,12 @@ void FActiveGameplayEffectsContainer::OnStackCountChange(FActiveGameplayEffect&
 
 	if (ActiveEffect.Spec.Def != nullptr)
 	{
-		Owner->NotifyTagMap_StackCountChange(ActiveEffect.Spec.Def->InheritableOwnedTagsContainer.CombinedTags);
+		Owner->NotifyTagMap_StackCountChange(ActiveEffect.Spec.Def->GetGrantedTags());
 	}
 
 	Owner->NotifyTagMap_StackCountChange(ActiveEffect.Spec.DynamicGrantedTags);
 
-	ActiveEffect.EventSet.OnStackChanged.Broadcast(ActiveEffect.Handle, ActiveEffect.Spec.StackCount, OldStackCount);
+	ActiveEffect.EventSet.OnStackChanged.Broadcast(ActiveEffect.Handle, ActiveEffect.Spec.GetStackCount(), OldStackCount);
 }
 
 /** Called when the duration or starttime of an AGE has changed */
@@ -2535,14 +2983,13 @@ bool FActiveGameplayEffectsContainer::HandleActiveGameplayEffectStackOverflow(co
 	const UGameplayEffect* StackedGE = OldSpec.Def;
 	const bool bAllowOverflowApplication = !(StackedGE->bDenyOverflowApplication);
 
-	FPredictionKey PredictionKey;
 	for (TSubclassOf<UGameplayEffect> OverflowEffect : StackedGE->OverflowEffects)
 	{
 		if (const UGameplayEffect* CDO = OverflowEffect.GetDefaultObject())
 		{
 			FGameplayEffectSpec NewGESpec;
 			NewGESpec.InitializeFromLinkedSpec(CDO, OverflowingSpec);
-			Owner->ApplyGameplayEffectSpecToSelf(NewGESpec, PredictionKey);
+			Owner->ApplyGameplayEffectSpecToSelf(NewGESpec);
 		}
 	}
 
@@ -3018,7 +3465,7 @@ FActiveGameplayEffect* FActiveGameplayEffectsContainer::ApplyGameplayEffectSpec(
 		bFoundExistingStackableGE = true;
 
 		FGameplayEffectSpec& ExistingSpec = ExistingStackableGE->Spec;
-		StartingStackCount = ExistingSpec.StackCount;
+		StartingStackCount = ExistingSpec.GetStackCount();
 
 		// This is now the global "being applied GE"
 		UAbilitySystemGlobals::Get().SetCurrentAppliedGE(&ExistingSpec);
@@ -3027,7 +3474,7 @@ FActiveGameplayEffect* FActiveGameplayEffectsContainer::ApplyGameplayEffectSpec(
 		// We still want to apply the stacks that didn't push us over, but we also want to call HandleActiveGameplayEffectStackOverflow.
 		
 		// For now: call HandleActiveGameplayEffectStackOverflow only if we are ALREADY at the limit. Else we just clamp stack limit to max.
-		if (ExistingSpec.StackCount == ExistingSpec.Def->StackLimitCount)
+		if (ExistingSpec.GetStackCount() == ExistingSpec.Def->StackLimitCount)
 		{
 			if (!HandleActiveGameplayEffectStackOverflow(*ExistingStackableGE, ExistingSpec, Spec))
 			{
@@ -3035,7 +3482,7 @@ FActiveGameplayEffect* FActiveGameplayEffectsContainer::ApplyGameplayEffectSpec(
 			}
 		}
 		
-		NewStackCount = ExistingSpec.StackCount + Spec.StackCount;
+		NewStackCount = ExistingSpec.GetStackCount() + Spec.GetStackCount();
 		if (ExistingSpec.Def->StackLimitCount > 0)
 		{
 			NewStackCount = FMath::Min(NewStackCount, ExistingSpec.Def->StackLimitCount);
@@ -3059,7 +3506,7 @@ FActiveGameplayEffect* FActiveGameplayEffectsContainer::ApplyGameplayEffectSpec(
 		ensureMsgf(ExistingSpec.GetDynamicAssetTags() == Spec.GetDynamicAssetTags(), TEXT("While adding a stack of the gameplay effect: %s, the old stack and the new application had different dynamic asset tags, which is currently not resolved properly! Existing: %s. New: %s"), *Spec.Def->GetName(), *ExistingSpec.GetDynamicAssetTags().ToStringSimple(), *Spec.GetDynamicAssetTags().ToStringSimple() );
 
 		ExistingStackableGE->Spec = Spec;
-		ExistingStackableGE->Spec.StackCount = NewStackCount;
+		ExistingStackableGE->Spec.SetStackCount(NewStackCount);
 
 		// Swap in old granted ability spec
 		ExistingStackableGE->Spec.GrantedAbilitySpecs = MoveTemp(GrantedSpecTempArray);
@@ -3287,39 +3734,14 @@ void FActiveGameplayEffectsContainer::InternalOnActiveGameplayEffectAdded(FActiv
 	AActor* OwnerActor = Owner->GetOwnerActor();
 	UE_VLOG(OwnerActor ? OwnerActor : Owner->GetOuter(), LogGameplayEffects, Log, TEXT("Added: %s"), *GetNameSafe(EffectDef->GetClass()));
 
-	// Add our ongoing tag requirements to the dependency map. We will actually check for these tags below.
-	for (const FGameplayTag& Tag : EffectDef->OngoingTagRequirements.IgnoreTags)
-	{
-		ActiveEffectTagDependencies.FindOrAdd(Tag).Add(Effect.Handle);
-	}
-
-	for (const FGameplayTag& Tag : EffectDef->OngoingTagRequirements.RequireTags)
-	{
-		ActiveEffectTagDependencies.FindOrAdd(Tag).Add(Effect.Handle);
-	}
-
-	// Add our removal tag requirements to the dependency map. We will actually check for these tags below.
-	for (const FGameplayTag& Tag : EffectDef->RemovalTagRequirements.IgnoreTags)
-	{
-		ActiveEffectTagDependencies.FindOrAdd(Tag).Add(Effect.Handle);
-	}
-
-	for (const FGameplayTag& Tag : EffectDef->RemovalTagRequirements.RequireTags)
-	{
-		ActiveEffectTagDependencies.FindOrAdd(Tag).Add(Effect.Handle);
-	}
-
 	// Add any external dependencies that might dirty the effect, if necessary
 	AddCustomMagnitudeExternalDependencies(Effect);
 
-	// Check if we should actually be turned on or not (this will turn us on for the first time)
-	static FGameplayTagContainer OwnerTags;
-	OwnerTags.Reset();
+	const bool bActive = EffectDef->OnAddedToActiveContainer(*this, Effect);
 
-	Owner->GetOwnedGameplayTags(OwnerTags);
-	
-	Effect.bIsInhibited = true; // Effect has to start inhibited, if it should be uninhibited, CheckOnGoingTagRequirements will handle that state change
-	Effect.CheckOngoingTagRequirements(OwnerTags, *this);
+	constexpr bool bInvokeCuesIfEnabled = true;
+	Effect.bIsInhibited = true; // Effect has to start inhibited, so our call to Inhibit will trigger if we should be active
+	Owner->InhibitActiveGameplayEffect(Effect.Handle, !bActive, bInvokeCuesIfEnabled);
 }
 
 void FActiveGameplayEffectsContainer::AddActiveGameplayEffectGrantedTagsAndModifiers(FActiveGameplayEffect& Effect, bool bInvokeGameplayCueEvents)
@@ -3381,29 +3803,23 @@ void FActiveGameplayEffectsContainer::AddActiveGameplayEffectGrantedTagsAndModif
 		}
 	}
 
-	// Update minimal replication if needed.
-	if (ShouldUseMinimalReplication())
-	{
-		Owner->AddMinimalReplicationGameplayTags(Effect.Spec.Def->InheritableOwnedTagsContainer.CombinedTags);
-		Owner->AddMinimalReplicationGameplayTags(Effect.Spec.DynamicGrantedTags);
-	}
 
 	// Update our owner with the tags this GameplayEffect grants them
-	Owner->UpdateTagMap(Effect.Spec.Def->InheritableOwnedTagsContainer.CombinedTags, 1);
+	Owner->UpdateTagMap(Effect.Spec.Def->GetGrantedTags(), 1);
 	Owner->UpdateTagMap(Effect.Spec.DynamicGrantedTags, 1);
 
 	// Update our owner with the blocked ability tags this GameplayEffect adds to them
-	Owner->BlockAbilitiesWithTags(Effect.Spec.Def->InheritableBlockedAbilityTagsContainer.CombinedTags);
+	Owner->BlockAbilitiesWithTags(Effect.Spec.Def->GetBlockedAbilityTags());
 
-	// Immunity
-	ApplicationImmunityGameplayTagCountContainer.UpdateTagCount(Effect.Spec.Def->GrantedApplicationImmunityTags.RequireTags, 1);
-	ApplicationImmunityGameplayTagCountContainer.UpdateTagCount(Effect.Spec.Def->GrantedApplicationImmunityTags.IgnoreTags, 1);
 
-	if (Effect.Spec.Def->HasGrantedApplicationImmunityQuery)
-	{	
-		ApplicationImmunityQueryEffects.Add(Effect.Spec.Def);
+	// Update minimal replication if needed.
+	if (ShouldUseMinimalReplication())
+	{
+		Owner->AddMinimalReplicationGameplayTags(Effect.Spec.Def->GetGrantedTags());
+		Owner->AddMinimalReplicationGameplayTags(Effect.Spec.DynamicGrantedTags);
 	}
 
+
 	// Grant abilities
 	if (IsNetAuthority() && !Owner->bSuppressGrantAbility)
 	{
@@ -3536,16 +3952,17 @@ bool FActiveGameplayEffectsContainer::InternalRemoveActiveGameplayEffect(int32 I
 		ABILITY_LOG(Verbose, TEXT("InternalRemoveActiveGameplayEffect: Auth: %s Handle: %s Def: %s"), IsNetAuthority() ? TEXT("TRUE") : TEXT("FALSE"), *Effect.Handle.ToString(), Effect.Spec.Def ? *Effect.Spec.Def->GetName() : TEXT("NONE"));
 
 		FGameplayEffectRemovalInfo GameplayEffectRemovalInfo;
-		GameplayEffectRemovalInfo.StackCount = Effect.Spec.StackCount;
+		GameplayEffectRemovalInfo.ActiveEffect = &Effect;
+		GameplayEffectRemovalInfo.StackCount = Effect.Spec.GetStackCount();
 		GameplayEffectRemovalInfo.bPrematureRemoval = bPrematureRemoval;
 		GameplayEffectRemovalInfo.EffectContext = Effect.Spec.GetEffectContext();
 
-		if (StacksToRemove > 0 && Effect.Spec.StackCount > StacksToRemove)
+		if (StacksToRemove > 0 && Effect.Spec.GetStackCount() > StacksToRemove)
 		{
 			// This won't be a full remove, only a change in StackCount.
-			int32 StartingStackCount = Effect.Spec.StackCount;
-			Effect.Spec.StackCount -= StacksToRemove;
-			OnStackCountChange(Effect, StartingStackCount, Effect.Spec.StackCount);
+			int32 StartingStackCount = Effect.Spec.GetStackCount();
+			Effect.Spec.SetStackCount(StartingStackCount - StacksToRemove);
+			OnStackCountChange(Effect, StartingStackCount, Effect.Spec.GetStackCount());
 			return false;
 		}
 
@@ -3592,9 +4009,6 @@ bool FActiveGameplayEffectsContainer::InternalRemoveActiveGameplayEffect(int32 I
 		// Remove this handle from the global map
 		Effect.Handle.RemoveFromGlobalMap();
 
-		// Attempt to apply expiration effects, if necessary
-		InternalApplyExpirationEffects(Effect.Spec, bPrematureRemoval);
-
 		bool ModifiedArray = false;
 
 		// Finally remove the ActiveGameplayEffect
@@ -3643,12 +4057,6 @@ void FActiveGameplayEffectsContainer::InternalOnActiveGameplayEffectRemoved(FAct
 
 	if (Effect.Spec.Def)
 	{
-		// Remove our tag requirements from the dependency map
-		RemoveActiveEffectTagDependency(Effect.Spec.Def->OngoingTagRequirements.IgnoreTags, Effect.Handle);
-		RemoveActiveEffectTagDependency(Effect.Spec.Def->OngoingTagRequirements.RequireTags, Effect.Handle);
-		RemoveActiveEffectTagDependency(Effect.Spec.Def->RemovalTagRequirements.IgnoreTags, Effect.Handle);
-		RemoveActiveEffectTagDependency(Effect.Spec.Def->RemovalTagRequirements.RequireTags, Effect.Handle);
-
 		// Only Need to update tags and modifiers if the gameplay effect is active.
 		if (!Effect.bIsInhibited)
 		{
@@ -3686,28 +4094,19 @@ void FActiveGameplayEffectsContainer::RemoveActiveGameplayEffectGrantedTagsAndMo
 	}
 
 	// Update gameplaytag count and broadcast delegate if we are at 0
-	Owner->UpdateTagMap(Effect.Spec.Def->InheritableOwnedTagsContainer.CombinedTags, -1);
+	Owner->UpdateTagMap(Effect.Spec.Def->GetGrantedTags(), -1);
 	Owner->UpdateTagMap(Effect.Spec.DynamicGrantedTags, -1);
 
 	// Update our owner with the blocked ability tags this GameplayEffect adds to them
-	Owner->UnBlockAbilitiesWithTags(Effect.Spec.Def->InheritableBlockedAbilityTagsContainer.CombinedTags);
+	Owner->UnBlockAbilitiesWithTags(Effect.Spec.Def->GetBlockedAbilityTags());
 
 	// Update minimal replication if needed.
 	if (ShouldUseMinimalReplication())
 	{
-		Owner->RemoveMinimalReplicationGameplayTags(Effect.Spec.Def->InheritableOwnedTagsContainer.CombinedTags);
+		Owner->RemoveMinimalReplicationGameplayTags(Effect.Spec.Def->GetGrantedTags());
 		Owner->RemoveMinimalReplicationGameplayTags(Effect.Spec.DynamicGrantedTags);
 	}
 
-	// Immunity
-	ApplicationImmunityGameplayTagCountContainer.UpdateTagCount(Effect.Spec.Def->GrantedApplicationImmunityTags.RequireTags, -1);
-	ApplicationImmunityGameplayTagCountContainer.UpdateTagCount(Effect.Spec.Def->GrantedApplicationImmunityTags.IgnoreTags, -1);
-
-	if (Effect.Spec.Def->HasGrantedApplicationImmunityQuery)
-	{
-		ApplicationImmunityQueryEffects.Remove(Effect.Spec.Def);
-	}
-
 	// Cancel/remove granted abilities
 	if (IsNetAuthority())
 	{
@@ -3760,22 +4159,6 @@ void FActiveGameplayEffectsContainer::RemoveActiveGameplayEffectGrantedTagsAndMo
 	}
 }
 
-void FActiveGameplayEffectsContainer::RemoveActiveEffectTagDependency(const FGameplayTagContainer& Tags, FActiveGameplayEffectHandle Handle)
-{
-	for(const FGameplayTag& Tag : Tags)
-	{
-		auto Ptr = ActiveEffectTagDependencies.Find(Tag);
-		if (Ptr)
-		{
-			Ptr->Remove(Handle);
-			if (Ptr->Num() <= 0)
-			{
-				ActiveEffectTagDependencies.Remove(Tag);
-			}
-		}
-	}
-}
-
 void FActiveGameplayEffectsContainer::AddCustomMagnitudeExternalDependencies(FActiveGameplayEffect& Effect)
 {
 	const UGameplayEffect* GEDef = Effect.Spec.Def;
@@ -3900,37 +4283,6 @@ void FActiveGameplayEffectsContainer::OnCustomMagnitudeExternalDependencyFired(T
 	}
 }
 
-void FActiveGameplayEffectsContainer::InternalApplyExpirationEffects(const FGameplayEffectSpec& ExpiringSpec, bool bPrematureRemoval)
-{
-	GAMEPLAYEFFECT_SCOPE_LOCK();
-
-	check(Owner);
-
-	// Don't allow prediction of expiration effects
-	if (IsNetAuthority())
-	{
-		const UGameplayEffect* ExpiringGE = ExpiringSpec.Def;
-		if (ExpiringGE)
-		{
-			// Determine the appropriate type of effect to apply depending on whether the effect is being prematurely removed or not
-			const TArray<TSubclassOf<UGameplayEffect>>& ExpiryEffects = (bPrematureRemoval ? ExpiringGE->PrematureExpirationEffectClasses : ExpiringGE->RoutineExpirationEffectClasses);
-			for (const TSubclassOf<UGameplayEffect>& CurExpiryEffect : ExpiryEffects)
-			{
-				if (CurExpiryEffect)
-				{
-					const UGameplayEffect* CurExpiryCDO = CurExpiryEffect->GetDefaultObject<UGameplayEffect>();
-					check(CurExpiryCDO);
-
-					FGameplayEffectSpec NewSpec;
-					NewSpec.InitializeFromLinkedSpec(CurExpiryCDO, ExpiringSpec);
-					
-					Owner->ApplyGameplayEffectSpecToSelf(NewSpec);
-				}
-			}
-		}
-	}
-}
-
 void FActiveGameplayEffectsContainer::RestartActiveGameplayEffectDuration(FActiveGameplayEffect& ActiveGameplayEffect)
 {
 	ActiveGameplayEffect.StartServerWorldTime = GetServerWorldTime();
@@ -3941,104 +4293,6 @@ void FActiveGameplayEffectsContainer::RestartActiveGameplayEffectDuration(FActiv
 	OnDurationChange(ActiveGameplayEffect);
 }
 
-void FActiveGameplayEffectsContainer::OnOwnerTagChange(FGameplayTag TagChange, int32 NewCount)
-{
-	// It may be beneficial to do a scoped lock on attribute re-evaluation during this function
-	GAMEPLAYEFFECT_SCOPE_LOCK();
-
-	FGameplayTagContainer OwnerTags;
-	if (Owner)
-	{
-		Owner->GetOwnedGameplayTags(OwnerTags);
-	}
-	else
-	{
-		UE_LOG(LogGameplayEffects, Warning, TEXT("No Owner for FActiveGameplayEffectsContainer in OnOwnerTagChange"));
-	}
-
-	// if this tag has any dependencies, we need to evaluate any removals
-	if (ActiveEffectTagDependencies.Contains(TagChange))
-	{
-		for (int32 idx = GetNumGameplayEffects() - 1; idx >= 0; --idx)
-		{
-			const FActiveGameplayEffect* Effect = GetActiveGameplayEffect(idx);
-			if (ensureMsgf(Effect != nullptr && Effect->Spec.Def != nullptr, TEXT("GetActiveGameplayEffect(%i) returned %p. GetNumGameplayEffects is %i."), idx, Effect, GetNumGameplayEffects()))
-			{
-				if (Effect->IsPendingRemove == false && Effect->CheckRemovalTagRequirements(OwnerTags, *this))
-				{
-					InternalRemoveActiveGameplayEffect(idx, -1, true);
-				}
-			}
-
-			ensure(Effect == GetActiveGameplayEffect(idx));
-		}
-	}
-
-	// Removing effects could change and/or remove the dependency entries, so find the set after we
-	// perform any removals.
-	TSet<FActiveGameplayEffectHandle>* Ptr = ActiveEffectTagDependencies.Find(TagChange);
-	if (Ptr)
-	{
-		// Copy the set in case it's modified while iterating.
-		TSet<FActiveGameplayEffectHandle> Handles = *Ptr;
-		for (const FActiveGameplayEffectHandle& Handle : Handles)
-		{
-			FActiveGameplayEffect* ActiveEffect = GetActiveGameplayEffect(Handle);
-			if (ActiveEffect)
-			{
-				ActiveEffect->CheckOngoingTagRequirements(OwnerTags, *this, true);
-			}
-		}
-	}
-}
-
-bool FActiveGameplayEffectsContainer::HasApplicationImmunityToSpec(const FGameplayEffectSpec& SpecToApply, const FActiveGameplayEffect*& OutGEThatProvidedImmunity) const
-{
-	SCOPE_CYCLE_COUNTER(STAT_HasApplicationImmunityToSpec)
-
-	const FGameplayTagContainer* AggregatedSourceTags = SpecToApply.CapturedSourceTags.GetAggregatedTags();
-	if (!ensure(AggregatedSourceTags))
-	{
-		return false;
-	}
-
-	// Query
-	for (const UGameplayEffect* EffectDef : ApplicationImmunityQueryEffects)
-	{
-		if (EffectDef->GrantedApplicationImmunityQuery.Matches(SpecToApply))
-		{
-			// This is blocked, but who blocked? Search for that Active GE
-			for (const FActiveGameplayEffect& Effect : this)
-			{
-				if (Effect.Spec.Def == EffectDef)
-				{
-					OutGEThatProvidedImmunity = &Effect;
-					return true;
-				}
-			}
-			ABILITY_LOG(Error, TEXT("Application Immunity was triggered for Applied GE: %s by Granted GE: %s. But this GE was not found in the Active GameplayEffects list!"), *GetNameSafe(SpecToApply.Def), *GetNameSafe( EffectDef));
-			break;
-		}
-	}
-
-	// Quick map test
-	if (!AggregatedSourceTags->HasAny(ApplicationImmunityGameplayTagCountContainer.GetExplicitGameplayTags()))
-	{
-		return false;
-	}
-
-	for (const FActiveGameplayEffect& Effect : this)
-	{
-		if (Effect.Spec.Def->GrantedApplicationImmunityTags.IsEmpty() == false && Effect.Spec.Def->GrantedApplicationImmunityTags.RequirementsMet( *AggregatedSourceTags ))
-		{
-			OutGEThatProvidedImmunity = &Effect;
-			return true;
-		}
-	}
-
-	return false;
-}
-
 ELifetimeCondition FActiveGameplayEffectsContainer::GetReplicationCondition() const
 {
 	if (Owner)
@@ -4213,7 +4467,7 @@ void FActiveGameplayEffectsContainer::CheckDuration(FActiveGameplayEffectHandle
 			if (Duration > 0.f && (((Effect.StartWorldTime + Duration) < CurrentTime) || FMath::IsNearlyZero(CurrentTime - Duration - Effect.StartWorldTime, KINDA_SMALL_NUMBER)))
 			{
 				// Figure out what to do based on the expiration policy
-				switch(Effect.Spec.Def->StackExpirationPolicy)
+				switch(Effect.Spec.Def->GetStackExpirationPolicy())
 				{
 				case EGameplayEffectStackingExpirationPolicy::ClearEntireStack:
 					StacksToRemove = -1; // Remove all stacks
@@ -4222,7 +4476,7 @@ void FActiveGameplayEffectsContainer::CheckDuration(FActiveGameplayEffectHandle
 
 				case EGameplayEffectStackingExpirationPolicy::RemoveSingleStackAndRefreshDuration:
 					StacksToRemove = 1;
-					CheckForFinalPeriodicExec = (Effect.Spec.StackCount == 1);
+					CheckForFinalPeriodicExec = (Effect.Spec.GetStackCount() == 1);
 					RefreshStartTime = true;
 					RefreshDurationTimer = true;
 					break;
@@ -4513,49 +4767,6 @@ int32 FActiveGameplayEffectsContainer::RemoveActiveEffects(const FGameplayEffect
 	return NumRemoved;
 }
 
-void FActiveGameplayEffectsContainer::AttemptRemoveActiveEffectsOnEffectApplication(const FGameplayEffectSpec &InSpec, const FActiveGameplayEffectHandle& InHandle)
-{
-	GAMEPLAYEFFECT_SCOPE_LOCK();
-	if (InSpec.Def)
-	{
-		// Clear tags is always removing all stacks.
-		FGameplayEffectQuery ClearQuery;
-		if (InSpec.Def->RemoveGameplayEffectsWithTags.CombinedTags.Num() > 0)
-		{
-			ClearQuery = FGameplayEffectQuery::MakeQuery_MatchAnyOwningTags(InSpec.Def->RemoveGameplayEffectsWithTags.CombinedTags);
-			if (InHandle.IsValid())
-			{
-				ClearQuery.IgnoreHandles.Add(InHandle);
-			}
-		}
-
-		const FGameplayTagContainer* AggregatedTargetTags = InSpec.CapturedTargetTags.GetAggregatedTags();
-		for (int32 idx = GetNumGameplayEffects() - 1; idx >= 0; --idx)
-		{
-			const FActiveGameplayEffect& ActiveEffect = *GetActiveGameplayEffect(idx);
-			if (!ActiveEffect.IsPendingRemove)
-			{
-				const UGameplayEffect* Effect = ActiveEffect.Spec.Def;
-				// check if the effect has a tag query
-				if (Effect && Effect->HasRemoveGameplayEffectsQuery && Effect->RemoveGameplayEffectQuery.Matches(InSpec))
-				{
-					InternalRemoveActiveGameplayEffect(idx, -1, true);
-				}
-				// check if the effect has removal tag requirements
-				else if (Effect && AggregatedTargetTags && !Effect->RemovalTagRequirements.IsEmpty() && Effect->RemovalTagRequirements.RequirementsMet(*AggregatedTargetTags))
-				{
-					InternalRemoveActiveGameplayEffect(idx, -1, true);
-				}
-				// clear query check
-				else if (!ClearQuery.IsEmpty() && ClearQuery.Matches(ActiveEffect))
-				{
-					InternalRemoveActiveGameplayEffect(idx, -1, true);
-				}
-			}
-		}
-	}
-}
-
 int32 FActiveGameplayEffectsContainer::GetActiveEffectCount(const FGameplayEffectQuery& Query, bool bEnforceOnGoingCheck) const
 {
 	int32 Count = 0;
@@ -4566,7 +4777,7 @@ int32 FActiveGameplayEffectsContainer::GetActiveEffectCount(const FGameplayEffec
 		{
 			if (Query.Matches(Effect))
 			{
-				Count += Effect.Spec.StackCount;
+				Count += Effect.Spec.GetStackCount();
 			}
 		}
 	}
@@ -4629,11 +4840,11 @@ void FActiveGameplayEffectsContainer::GetActiveGameplayEffectDataByAttribute(TMu
 				Data.GameplayEffectName = Effect.Spec.Def->GetName();
 				Data.ModifierOp = Effect.Spec.Def->Modifiers[Idx].ModifierOp;
 				Data.Magnitude = Effect.Spec.Modifiers[Idx].GetEvaluatedMagnitude();
-				if (Effect.Spec.StackCount > 1)
+				if (Effect.Spec.GetStackCount() > 1)
 				{
-					Data.Magnitude = GameplayEffectUtilities::ComputeStackedModifierMagnitude(Data.Magnitude, Effect.Spec.StackCount, Data.ModifierOp);
+					Data.Magnitude = GameplayEffectUtilities::ComputeStackedModifierMagnitude(Data.Magnitude, Effect.Spec.GetStackCount(), Data.ModifierOp);
 				}
-				Data.StackCount = Effect.Spec.StackCount;
+				Data.StackCount = Effect.Spec.GetStackCount();
 
 				EffectMap.Add(Data.Attribute, Data);
 			}
@@ -4903,18 +5114,13 @@ bool FGameplayEffectQuery::Matches(const FGameplayEffectSpec& Spec) const
 		check(IsInGameThread());
 		static FGameplayTagContainer TargetTags;
 		TargetTags.Reset();
-		if (Spec.Def->InheritableGameplayEffectTags.CombinedTags.Num() > 0)
-		{
-			TargetTags.AppendTags(Spec.Def->InheritableGameplayEffectTags.CombinedTags);
-		}
-		if (Spec.Def->InheritableOwnedTagsContainer.CombinedTags.Num() > 0)
-		{
-			TargetTags.AppendTags(Spec.Def->InheritableOwnedTagsContainer.CombinedTags);
-		}
-		if (Spec.DynamicGrantedTags.Num() > 0)
-		{
-			TargetTags.AppendTags(Spec.DynamicGrantedTags);
-		}
+		
+		// This functionality is kept the same as prior to UE5.3:
+		// I believe this is a side-effect of not understanding the previous Asset vs. Granted tags in UGameplayEffect.
+		// We're keeping the functionality and actually comparing TargetTags against both Asset & Granted.
+		// The one change in 5.3: We're now also comparing against DynamicAssetTags to keep it consistent.
+		Spec.GetAllAssetTags(TargetTags);
+		Spec.GetAllGrantedTags(TargetTags);
 		
 		if (OwningTagQuery.Matches(TargetTags) == false)
 		{
@@ -4929,15 +5135,8 @@ bool FGameplayEffectQuery::Matches(const FGameplayEffectSpec& Spec) const
 		check(IsInGameThread());
 		static FGameplayTagContainer GETags;
 		GETags.Reset();
-		if (Spec.Def->InheritableGameplayEffectTags.CombinedTags.Num() > 0)
-		{
-			GETags.AppendTags(Spec.Def->InheritableGameplayEffectTags.CombinedTags);
-		}
-		const FGameplayTagContainer& SpecDynamicAssetTags = Spec.GetDynamicAssetTags();
-		if (SpecDynamicAssetTags.Num() > 0)
-		{
-			GETags.AppendTags(SpecDynamicAssetTags);
-		}
+
+		Spec.GetAllAssetTags(GETags);
 
 		if (EffectTagQuery.Matches(GETags) == false)
 		{
@@ -4947,7 +5146,8 @@ bool FGameplayEffectQuery::Matches(const FGameplayEffectSpec& Spec) const
 
 	if (SourceTagQuery.IsEmpty() == false)
 	{
-		FGameplayTagContainer const& SourceTags = Spec.CapturedSourceTags.GetSpecTags();
+		// Prior to UE5.3 this checked against the SpecTags, which seems like an obvious bug.
+		FGameplayTagContainer const& SourceTags = *Spec.CapturedSourceTags.GetAggregatedTags();
 		if (SourceTagQuery.Matches(SourceTags) == false)
 		{
 			return false;
@@ -5092,6 +5292,23 @@ FGameplayEffectQuery FGameplayEffectQuery::MakeQuery_MatchNoSourceTags(const FGa
 	return OutQuery;
 }
 
+bool FGameplayEffectQuery::operator==(const FGameplayEffectQuery& Other) const
+{
+	return CustomMatchDelegate_BP == Other.CustomMatchDelegate_BP &&
+	OwningTagQuery == Other.OwningTagQuery &&
+	EffectTagQuery == Other.EffectTagQuery &&
+	SourceTagQuery == Other.SourceTagQuery &&
+	ModifyingAttribute == Other.ModifyingAttribute &&
+	EffectSource == Other.EffectSource &&
+	EffectDefinition == Other.EffectDefinition &&
+	IgnoreHandles == Other.IgnoreHandles;
+}
+
+bool FGameplayEffectQuery::operator!=(const FGameplayEffectQuery& Other) const
+{
+	return !(*this == Other);
+}
+
 bool FGameplayModifierInfo::operator==(const FGameplayModifierInfo& Other) const
 {
 	if (Attribute != Other.Attribute)
@@ -5157,6 +5374,7 @@ void FInheritedTagContainer::UpdateInheritedTagProperties(const FInheritedTagCon
 	{
 		// Remove trumps add for explicit matches but not for parent tags.
 		// This lets us remove all inherited tags starting with Foo but still add Foo.Bar
+		// Keep in mind: When adding Foo.Bar, you are implicitly adding Foo (its parent tags).
 		if (!Removed.HasTagExact(*Itr))
 		{
 			CombinedTags.AddTag(*Itr);
@@ -5169,12 +5387,28 @@ void FInheritedTagContainer::UpdateInheritedTagProperties(const FInheritedTagCon
 #endif
 }
 
-void FInheritedTagContainer::PostInitProperties()
+void FInheritedTagContainer::ApplyTo(FGameplayTagContainer& ApplyToContainer) const
 {
-	// we shouldn't inherit the added and removed tags from our parents
-	// make sure that these fields are clear
-	Added.Reset();
-	Removed.Reset();
+	if (ApplyToContainer.IsEmpty() && !CombinedTags.IsEmpty())
+	{
+		// This is a fast path where our CombinedTags were already computed and
+		// the container we're apply to is empty, then we simply assign the CombinedTags.
+		ApplyToContainer = CombinedTags;
+	}
+	else
+	{
+		// Layer in the actual changes.  It's best understood with an example:
+		// Current: Foo.Car  Remove: Foo  Add: Foo.Bar
+		// Remove all tags that match the filter (e.g. Foo removes Foo.Car)
+		// Then, add in all of the Adds which don't match exactly (e.g. Foo.Bar would be added, but attempting to add exactly Foo would not be allowed since Removes take precedence).
+		// Keep in mind: When Foo.Bar is added, you are implicitly adding Foo (its parent tags).  Thus you can remove Foo.* and still add Foo.Bar.
+		FGameplayTagContainer RemovesThatApply = Removed.Filter(ApplyToContainer);
+		FGameplayTagContainer RemoveOverridesAdd = Added.FilterExact(Removed);
+		RemovesThatApply.AppendTags(RemoveOverridesAdd);
+
+		ApplyToContainer.AppendTags(Added);
+		ApplyToContainer.RemoveTags(RemovesThatApply);
+	}
 }
 
 void FInheritedTagContainer::AddTag(const FGameplayTag& TagToAdd)
@@ -5182,11 +5416,23 @@ void FInheritedTagContainer::AddTag(const FGameplayTag& TagToAdd)
 	CombinedTags.AddTag(TagToAdd);
 }
 
-void FInheritedTagContainer::RemoveTag(FGameplayTag TagToRemove)
+void FInheritedTagContainer::RemoveTag(const FGameplayTag& TagToRemove)
 {
 	CombinedTags.RemoveTag(TagToRemove);
 }
 
+bool FInheritedTagContainer::operator==(const FInheritedTagContainer& Other)  const
+{
+	return CombinedTags == Other.CombinedTags &&
+		Added == Other.Added &&
+		Removed == Other.Removed;
+}
+
+bool FInheritedTagContainer::operator!=(const FInheritedTagContainer& Other) const
+{
+	return !(*this == Other);
+}
+
 void FActiveGameplayEffectsContainer::IncrementLock()
 {
 	ScopedLockCount++;
@@ -5260,4 +5506,6 @@ FScopedActiveGameplayEffectLock::~FScopedActiveGameplayEffectLock()
 	Container.DecrementLock();
 }
 
+#undef LOCTEXT_NAMESPACE
+
 
diff --git a/Engine/Plugins/Runtime/GameplayAbilities/Source/GameplayAbilities/Private/GameplayEffectComponent.cpp b/Engine/Plugins/Runtime/GameplayAbilities/Source/GameplayAbilities/Private/GameplayEffectComponent.cpp
new file mode 100644
--- /dev/null
+++ b/Engine/Plugins/Runtime/GameplayAbilities/Source/GameplayAbilities/Private/GameplayEffectComponent.cpp
@@ -0,0 +1,66 @@
+// Copyright Epic Games, Inc. All Rights Reserved.
+
+#include "GameplayEffectComponent.h"
+
+#include "AbilitySystemLog.h"
+#include "GameplayEffect.h"
+#include "Misc/DataValidation.h"
+
+#define LOCTEXT_NAMESPACE "GameplayEffectComponent"
+
+UGameplayEffectComponent::UGameplayEffectComponent()
+{
+#if WITH_EDITORONLY_DATA
+	EditorFriendlyName = FName::NameToDisplayString(GetNameSafe(GetClass()), false);
+	EditorFriendlyName.ReplaceInline(TEXT("Gameplay Effect Component"), TEXT(""));
+#endif
+}
+
+UGameplayEffect* UGameplayEffectComponent::GetOwner() const
+{
+	return GetTypedOuter<UGameplayEffect>();
+}
+
+#if WITH_EDITOR
+EDataValidationResult UGameplayEffectComponent::IsDataValid(FDataValidationContext& Context)
+{
+	EDataValidationResult Result = Super::IsDataValid(Context);
+
+	TSubclassOf<UGameplayEffectComponent> ThisClass = GetClass();
+	
+	const UGameplayEffect* Owner = GetOwner();
+	if (Owner)
+	{
+		const UGameplayEffectComponent* FirstComponentOfThisType = Owner->FindComponent(ThisClass);
+		if (!FirstComponentOfThisType)
+		{
+			Context.AddError(LOCTEXT("ComponentNotOnGE", "Component does not exist in its Owner's GEComponents Array"));
+			Result = EDataValidationResult::Invalid;
+		}
+		else if (FirstComponentOfThisType != this)
+		{
+			Context.AddError(FText::FormatOrdered(LOCTEXT("MultipleComponentsOfSameTypeError", "Two or more types of {0} exist on GE"), FText::FromString(ThisClass->GetName())));
+			Result = EDataValidationResult::Invalid;
+		}
+	}
+	else
+	{
+		FText ErrorText = FText::FormatOrdered(LOCTEXT("NoGEOwner", "{0} has invalid Outer: {1}"), FText::FromString(GetNameSafe(this)), FText::FromString(GetNameSafe(GetOuter())));
+		Context.AddError(ErrorText);
+
+		UE_LOG(LogGameplayEffects, Error, TEXT("%s"), *ErrorText.ToString());
+
+		Result = EDataValidationResult::Invalid;
+	}
+
+	// Let's just say we've done enough to say it's valid and it's up to the child classes to invalidate
+	if (Result == EDataValidationResult::NotValidated)
+	{
+		Result = EDataValidationResult::Valid;
+	}
+
+	return Result;
+}
+#endif
+
+#undef LOCTEXT_NAMESPACE
diff --git a/Engine/Plugins/Runtime/GameplayAbilities/Source/GameplayAbilities/Private/GameplayEffectComponents/AdditionalEffectsGameplayEffectComponent.cpp b/Engine/Plugins/Runtime/GameplayAbilities/Source/GameplayAbilities/Private/GameplayEffectComponents/AdditionalEffectsGameplayEffectComponent.cpp
new file mode 100644
--- /dev/null
+++ b/Engine/Plugins/Runtime/GameplayAbilities/Source/GameplayAbilities/Private/GameplayEffectComponents/AdditionalEffectsGameplayEffectComponent.cpp
@@ -0,0 +1,151 @@
+// Copyright Epic Games, Inc. All Rights Reserved.
+
+#include "GameplayEffectComponents/AdditionalEffectsGameplayEffectComponent.h"
+#include "AbilitySystemComponent.h"
+#include "AbilitySystemLog.h"
+#include "Misc/DataValidation.h"
+
+#define LOCTEXT_NAMESPACE "AdditionalEffectsGameplayEffectComponent"
+
+bool UAdditionalEffectsGameplayEffectComponent::OnActiveGameplayEffectAdded(FActiveGameplayEffectsContainer& ActiveGEContainer, FActiveGameplayEffect& ActiveGE) const
+{
+	UAbilitySystemComponent* ASC = ActiveGEContainer.Owner;
+	if (ensure(ASC))
+	{
+		OnGameplayEffectApplied(*ASC, ActiveGE.Spec, ActiveGE.PredictionKey);
+	}
+
+	// We don't allow prediction of expiration (on removed) effects
+	if (ActiveGEContainer.IsNetAuthority())
+	{
+		// When this ActiveGE gets removed, so will our events so no need to unbind this.
+		ActiveGE.EventSet.OnEffectRemoved.AddUObject(this, &UAdditionalEffectsGameplayEffectComponent::OnActiveGameplayEffectRemoved, &ActiveGEContainer);
+	}
+
+	return true;
+}
+
+void UAdditionalEffectsGameplayEffectComponent::OnGameplayEffectExecuted(FActiveGameplayEffectsContainer& ActiveGEContainer, FGameplayEffectSpec& GESpec, FPredictionKey& PredictionKey) const
+{
+	if (GetOwner()->DurationPolicy == EGameplayEffectDurationType::Instant)
+	{
+		UAbilitySystemComponent* ASC = ActiveGEContainer.Owner;
+		if (ensure(ASC))
+		{
+			OnGameplayEffectApplied(*ASC, GESpec, PredictionKey);
+		}
+	}
+}
+
+void UAdditionalEffectsGameplayEffectComponent::OnGameplayEffectApplied(UAbilitySystemComponent& AppliedToASC, FGameplayEffectSpec& GESpec, FPredictionKey& PredictionKey) const
+{
+	const float GELevel = GESpec.GetLevel();
+	const FGameplayEffectContextHandle& GEContextHandle = GESpec.GetEffectContext();
+
+	/** other effects that need to be applied to the target if this effect is successful */
+	TArray<FGameplayEffectSpecHandle> TargetEffectSpecs;
+	for (const FConditionalGameplayEffect& ConditionalEffect : OnApplicationGameplayEffects)
+	{
+		const UGameplayEffect* GameplayEffectDef = ConditionalEffect.EffectClass.GetDefaultObject();
+		if (!GameplayEffectDef)
+		{
+			continue;
+		}
+
+		if (ConditionalEffect.CanApply(GESpec.CapturedSourceTags.GetActorTags(), GELevel))
+		{
+			FGameplayEffectSpecHandle SpecHandle;
+			if (bOnApplicationCopyDataFromOriginalSpec)
+			{
+				SpecHandle = FGameplayEffectSpecHandle(new FGameplayEffectSpec());
+				SpecHandle.Data->InitializeFromLinkedSpec(GameplayEffectDef, GESpec);
+			}
+			else
+			{
+				SpecHandle = ConditionalEffect.CreateSpec(GEContextHandle, GELevel);
+			}
+
+			if (ensure(SpecHandle.IsValid()))
+			{
+				TargetEffectSpecs.Add(SpecHandle);
+			}
+		}
+	}
+
+	// Add all of the dynamically linked ones that don't have any conditions applied to them
+	// Note: I believe this functionality should be removed as its usage is dynamic and thus not 
+	// synchronized between client/server.  I've deprecated the variable to help enforce this.
+PRAGMA_DISABLE_DEPRECATION_WARNINGS
+	TargetEffectSpecs.Append(GESpec.TargetEffectSpecs);
+PRAGMA_ENABLE_DEPRECATION_WARNINGS
+
+	for (const FGameplayEffectSpecHandle& TargetSpec : TargetEffectSpecs)
+	{
+		if (TargetSpec.IsValid())
+		{
+			AppliedToASC.ApplyGameplayEffectSpecToSelf(*TargetSpec.Data.Get(), PredictionKey);
+		}
+	}
+}
+
+void UAdditionalEffectsGameplayEffectComponent::OnActiveGameplayEffectRemoved(const FGameplayEffectRemovalInfo& RemovalInfo, FActiveGameplayEffectsContainer* ActiveGEContainer) const
+{
+	FScopedActiveGameplayEffectLock ActiveScopeLock(*ActiveGEContainer);
+
+	const FActiveGameplayEffect* ActiveGE = RemovalInfo.ActiveEffect;
+	if (!ensure(ActiveGE))
+	{
+		UE_LOG(LogGameplayEffects, Error, TEXT("FGameplayEffectRemovalInfo::ActiveEffect was not populated in OnActiveGameplayEffectRemoved"));
+		return;
+	}
+
+	UAbilitySystemComponent* ASC = ActiveGEContainer->Owner;
+	if (!ensure(IsValid(ASC)))
+	{
+		UE_LOG(LogGameplayEffects, Error, TEXT("ActiveGEContainer was invalid in OnActiveGameplayEffectRemoved"));
+		return;
+	}
+
+	// Determine the appropriate type of effect to apply depending on whether the effect is being prematurely removed or not
+	const TArray<TSubclassOf<UGameplayEffect>>& ExpiryEffects = (RemovalInfo.bPrematureRemoval ? OnCompletePrematurely : OnCompleteNormal);
+
+	// Mix-in the always-executing GameplayEffects
+	TArray<TSubclassOf<UGameplayEffect>> AllGameplayEffects{ ExpiryEffects };
+	AllGameplayEffects.Append(OnCompleteAlways);
+
+	for (const TSubclassOf<UGameplayEffect>& CurExpiryEffect : AllGameplayEffects)
+	{
+		if (const UGameplayEffect* CurExpiryCDO = CurExpiryEffect.GetDefaultObject())
+		{
+			FGameplayEffectSpec NewSpec;
+			NewSpec.InitializeFromLinkedSpec(CurExpiryCDO, ActiveGE->Spec);
+
+			ASC->ApplyGameplayEffectSpecToSelf(NewSpec);
+		}
+	}
+}
+
+#if WITH_EDITOR
+EDataValidationResult UAdditionalEffectsGameplayEffectComponent::IsDataValid(FDataValidationContext& Context)
+{
+	EDataValidationResult Result = Super::IsDataValid(Context);
+
+	if (GetOwner()->DurationPolicy == EGameplayEffectDurationType::Instant)
+	{
+		const bool bHasOnCompleteEffects = (OnCompleteAlways.Num() + OnCompleteNormal.Num() + OnCompletePrematurely.Num() > 0);
+		Context.AddError(FText::FormatOrdered(LOCTEXT("InstantDoesNotWorkWithOnComplete", "Instant GE will never received OnComplete for {0}."), FText::FromString(GetClass()->GetName())));
+		Result = EDataValidationResult::Invalid;
+	}
+	else if (GetOwner()->Period.Value > 0.0f)
+	{
+		if (OnApplicationGameplayEffects.Num() > 0)
+		{
+			Context.AddWarning(LOCTEXT("IsPeriodicAndHasOnApplication", "Periodic GE has OnApplicationGameplayEffects. Those GE's will only be applied once."));
+		}
+	}
+
+	return Result;
+}
+#endif // WITH_EDITOR
+
+#undef LOCTEXT_NAMESPACE
diff --git a/Engine/Plugins/Runtime/GameplayAbilities/Source/GameplayAbilities/Private/GameplayEffectComponents/AssetTagsGameplayEffectComponent.cpp b/Engine/Plugins/Runtime/GameplayAbilities/Source/GameplayAbilities/Private/GameplayEffectComponents/AssetTagsGameplayEffectComponent.cpp
new file mode 100644
--- /dev/null
+++ b/Engine/Plugins/Runtime/GameplayAbilities/Source/GameplayAbilities/Private/GameplayEffectComponents/AssetTagsGameplayEffectComponent.cpp
@@ -0,0 +1,54 @@
+// Copyright Epic Games, Inc. All Rights Reserved.
+
+#include "GameplayEffectComponents/AssetTagsGameplayEffectComponent.h"
+
+#define LOCTEXT_NAMESPACE "AssetTagsGameplayEffectComponent"
+
+void UAssetTagsGameplayEffectComponent::PostInitProperties()
+{
+	Super::PostInitProperties();
+
+#if WITH_EDITORONLY_DATA
+	if (GetOwner())
+	{
+		EditorFriendlyName = FText::FormatOrdered(LOCTEXT("AssetTagsOnObj", "Asset Tags (on {0})"), GetOwner()->GetClass()->GetDisplayNameText()).ToString();
+	}
+	else
+	{
+		EditorFriendlyName = LOCTEXT("AssetTagsOnGE", "Asset Tags (on Gameplay Effect)").ToString();
+	}
+#endif // WITH_EDITORONLY_DATA
+}
+
+void UAssetTagsGameplayEffectComponent::OnOwnerPostLoad()
+{
+	Super::OnOwnerPostLoad();
+	SetAndApplyAssetTagChanges(InheritableAssetTags);
+}
+
+#if WITH_EDITOR
+void UAssetTagsGameplayEffectComponent::PostEditChangeProperty(FPropertyChangedEvent& PropertyChangedEvent)
+{
+	Super::PostEditChangeProperty(PropertyChangedEvent);
+
+	if (PropertyChangedEvent.GetMemberPropertyName()  == GetInheritableAssetTagsName())
+	{
+		SetAndApplyAssetTagChanges(InheritableAssetTags);
+	}
+}
+#endif // WITH_EDITOR
+
+void UAssetTagsGameplayEffectComponent::SetAndApplyAssetTagChanges(const FInheritedTagContainer& TagContainerMods)
+{
+	InheritableAssetTags = TagContainerMods;
+
+	// Try to find the parent and update the inherited tags
+	const UAssetTagsGameplayEffectComponent* Parent = FindParentComponent(*this);
+	InheritableAssetTags.UpdateInheritedTagProperties(Parent ? &Parent->InheritableAssetTags : nullptr);
+
+	// Apply to the owning Gameplay Effect Component
+	UGameplayEffect* Owner = GetOwner();
+	InheritableAssetTags.ApplyTo(Owner->CachedAssetTags);
+}
+
+#undef LOCTEXT_NAMESPACE
\ No newline at end of file
diff --git a/Engine/Plugins/Runtime/GameplayAbilities/Source/GameplayAbilities/Private/GameplayEffectComponents/BlockAbilityTagsGameplayEffectComponent.cpp b/Engine/Plugins/Runtime/GameplayAbilities/Source/GameplayAbilities/Private/GameplayEffectComponents/BlockAbilityTagsGameplayEffectComponent.cpp
new file mode 100644
--- /dev/null
+++ b/Engine/Plugins/Runtime/GameplayAbilities/Source/GameplayAbilities/Private/GameplayEffectComponents/BlockAbilityTagsGameplayEffectComponent.cpp
@@ -0,0 +1,41 @@
+// Copyright Epic Games, Inc. All Rights Reserved.
+
+#include "GameplayEffectComponents/BlockAbilityTagsGameplayEffectComponent.h"
+
+UBlockAbilityTagsGameplayEffectComponent::UBlockAbilityTagsGameplayEffectComponent()
+{
+#if WITH_EDITORONLY_DATA
+	EditorFriendlyName = TEXT("Block Abilities w/ Tags");
+#endif
+}
+
+void UBlockAbilityTagsGameplayEffectComponent::OnOwnerPostLoad()
+{
+	Super::OnOwnerPostLoad();
+	SetAndApplyBlockedAbilityTagChanges(InheritableBlockedAbilityTagsContainer);
+}
+
+#if WITH_EDITOR
+void UBlockAbilityTagsGameplayEffectComponent::PostEditChangeProperty(FPropertyChangedEvent& PropertyChangedEvent)
+{
+	Super::PostEditChangeProperty(PropertyChangedEvent);
+
+	if (PropertyChangedEvent.GetMemberPropertyName() == GetInheritableBlockedAbilityTagsContainerPropertyName())
+	{
+		SetAndApplyBlockedAbilityTagChanges(InheritableBlockedAbilityTagsContainer);
+	}
+}
+#endif // WITH_EDITOR
+
+void UBlockAbilityTagsGameplayEffectComponent::SetAndApplyBlockedAbilityTagChanges(const FInheritedTagContainer& TagContainerMods)
+{
+	InheritableBlockedAbilityTagsContainer = TagContainerMods;
+
+	// Try to find the parent and update the inherited tags
+	const UBlockAbilityTagsGameplayEffectComponent* Parent = FindParentComponent(*this);
+	InheritableBlockedAbilityTagsContainer.UpdateInheritedTagProperties(Parent ? &Parent->InheritableBlockedAbilityTagsContainer : nullptr);
+
+	// Apply to the owning Gameplay Effect Component
+	UGameplayEffect* Owner = GetOwner();
+	InheritableBlockedAbilityTagsContainer.ApplyTo(Owner->CachedBlockedAbilityTags);
+}
diff --git a/Engine/Plugins/Runtime/GameplayAbilities/Source/GameplayAbilities/Private/GameplayEffectComponents/ChanceToApplyGameplayEffectComponent.cpp b/Engine/Plugins/Runtime/GameplayAbilities/Source/GameplayAbilities/Private/GameplayEffectComponents/ChanceToApplyGameplayEffectComponent.cpp
new file mode 100644
--- /dev/null
+++ b/Engine/Plugins/Runtime/GameplayAbilities/Source/GameplayAbilities/Private/GameplayEffectComponents/ChanceToApplyGameplayEffectComponent.cpp
@@ -0,0 +1,38 @@
+// Copyright Epic Games, Inc. All Rights Reserved.
+
+#include "GameplayEffectComponents/ChanceToApplyGameplayEffectComponent.h"
+#include "AbilitySystemComponent.h"
+
+#if WITH_EDITOR
+#define GETCURVE_REPORTERROR_WITHPOSTLOAD(Handle) \
+	if (Handle.CurveTable) const_cast<UCurveTable*>(ToRawPtr(Handle.CurveTable))->ConditionalPostLoad(); \
+	GETCURVE_REPORTERROR(Handle);
+
+#endif // WITH_EDITOR
+
+void UChanceToApplyGameplayEffectComponent::PostLoad()
+{
+	Super::PostLoad();
+#if WITH_EDITOR
+	GETCURVE_REPORTERROR_WITHPOSTLOAD(ChanceToApplyToTarget.Curve);
+#endif // WITH_EDITOR
+}
+
+bool UChanceToApplyGameplayEffectComponent::CanGameplayEffectApply(const FActiveGameplayEffectsContainer& ActiveGEContainer, const FGameplayEffectSpec& GESpec) const
+{
+	const FString ContextString = GESpec.Def->GetName();
+	const float CalculatedChanceToApplyToTarget = ChanceToApplyToTarget.GetValueAtLevel(GESpec.GetLevel(), &ContextString);
+
+	// check probability to apply
+	if ((CalculatedChanceToApplyToTarget < 1.f - SMALL_NUMBER) && (FMath::FRand() > CalculatedChanceToApplyToTarget))
+	{
+		return false;
+	}
+
+	return true;
+}
+
+void UChanceToApplyGameplayEffectComponent::SetChanceToApplyToTarget(const FScalableFloat& InChanceToApplyToTarget)
+{
+	ChanceToApplyToTarget = InChanceToApplyToTarget;
+}
diff --git a/Engine/Plugins/Runtime/GameplayAbilities/Source/GameplayAbilities/Private/GameplayEffectComponents/CustomCanApplyGameplayEffectComponent.cpp b/Engine/Plugins/Runtime/GameplayAbilities/Source/GameplayAbilities/Private/GameplayEffectComponents/CustomCanApplyGameplayEffectComponent.cpp
new file mode 100644
--- /dev/null
+++ b/Engine/Plugins/Runtime/GameplayAbilities/Source/GameplayAbilities/Private/GameplayEffectComponents/CustomCanApplyGameplayEffectComponent.cpp
@@ -0,0 +1,25 @@
+// Copyright Epic Games, Inc. All Rights Reserved.
+
+#include "GameplayEffectComponents/CustomCanApplyGameplayEffectComponent.h"
+#include "AbilitySystemComponent.h"
+
+UCustomCanApplyGameplayEffectComponent::UCustomCanApplyGameplayEffectComponent()
+{
+#if WITH_EDITORONLY_DATA
+	EditorFriendlyName = TEXT("Custom Can Apply Function (in Blueprint)");
+#endif
+}
+
+bool UCustomCanApplyGameplayEffectComponent::CanGameplayEffectApply(const FActiveGameplayEffectsContainer& ActiveGEContainer, const FGameplayEffectSpec& GESpec) const
+{
+	// do custom application checks
+	for (const TSubclassOf<UGameplayEffectCustomApplicationRequirement>& AppReq : ApplicationRequirements)
+	{
+		if (*AppReq && AppReq->GetDefaultObject<UGameplayEffectCustomApplicationRequirement>()->CanApplyGameplayEffect(GESpec.Def, GESpec, ActiveGEContainer.Owner) == false)
+		{
+			return false;
+		}
+	}
+
+	return true;
+}
diff --git a/Engine/Plugins/Runtime/GameplayAbilities/Source/GameplayAbilities/Private/GameplayEffectComponents/ImmunityGameplayEffectComponent.cpp b/Engine/Plugins/Runtime/GameplayAbilities/Source/GameplayAbilities/Private/GameplayEffectComponents/ImmunityGameplayEffectComponent.cpp
new file mode 100644
--- /dev/null
+++ b/Engine/Plugins/Runtime/GameplayAbilities/Source/GameplayAbilities/Private/GameplayEffectComponents/ImmunityGameplayEffectComponent.cpp
@@ -0,0 +1,90 @@
+// Copyright Epic Games, Inc. All Rights Reserved.
+
+#include "GameplayEffectComponents/ImmunityGameplayEffectComponent.h"
+#include "AbilitySystemComponent.h"
+#include "AbilitySystemLog.h"
+#include "AbilitySystemStats.h"
+#include "Misc/DataValidation.h"
+
+#define LOCTEXT_NAMESPACE "ImmunityGameplayEffectComponent"
+
+UImmunityGameplayEffectComponent::UImmunityGameplayEffectComponent()
+{
+#if WITH_EDITORONLY_DATA
+	EditorFriendlyName = TEXT("Immunity (Prevent Other GEs)");
+#endif // WITH_EDITORONLY_DATA
+}
+
+bool UImmunityGameplayEffectComponent::OnActiveGameplayEffectAdded(FActiveGameplayEffectsContainer& ActiveGEContainer, FActiveGameplayEffect& ActiveGE) const
+{
+	FActiveGameplayEffectHandle& ActiveGEHandle = ActiveGE.Handle;
+	UAbilitySystemComponent* OwnerASC = ActiveGEContainer.Owner;
+
+	// Register our immunity query to potentially block applications of any Gameplay Effects
+	FGameplayEffectApplicationQuery& BoundQuery = OwnerASC->GameplayEffectApplicationQueries.AddDefaulted_GetRef();
+	BoundQuery.BindUObject(this, &UImmunityGameplayEffectComponent::AllowGameplayEffectApplication, ActiveGEHandle);
+
+	// Now that we've bound that function, let's unbind it when we're removed.  This is safe because once we're removed, EventSet is gone.
+	ActiveGE.EventSet.OnEffectRemoved.AddLambda([OwnerASC, QueryToRemove = BoundQuery.GetHandle()](const FGameplayEffectRemovalInfo& RemovalInfo)
+		{
+			if (ensure(IsValid(OwnerASC)))
+			{
+				TArray<FGameplayEffectApplicationQuery>& GEAppQueries = OwnerASC->GameplayEffectApplicationQueries;
+				for (auto It = GEAppQueries.CreateIterator(); It; ++It)
+				{
+					if (It->GetHandle() == QueryToRemove)
+					{
+						It.RemoveCurrentSwap();
+						break;
+					}
+				}
+			}
+		});
+	return true;
+}
+
+bool UImmunityGameplayEffectComponent::AllowGameplayEffectApplication(const FActiveGameplayEffectsContainer& ActiveGEContainer, const FGameplayEffectSpec& GESpecToConsider, FActiveGameplayEffectHandle ImmunityActiveGEHandle) const
+{
+	SCOPE_CYCLE_COUNTER(STAT_HasApplicationImmunityToSpec)
+
+	const UAbilitySystemComponent* ASC = ActiveGEContainer.Owner;
+	if (ASC != ImmunityActiveGEHandle.GetOwningAbilitySystemComponent())
+	{
+		ensureMsgf(false, TEXT("Something went wrong where an ActiveGameplayEffect jumped AbilitySystemComponents"));
+		return false;
+	}
+
+	const FActiveGameplayEffect* ActiveGE = ASC->GetActiveGameplayEffect(ImmunityActiveGEHandle);
+	if (!ensure(ActiveGE) || ActiveGE->bIsInhibited)
+	{
+		return true;
+	}
+
+	for (const FGameplayEffectQuery& ImmunityQuery : ImmunityQueries)
+	{
+		if (!ImmunityQuery.IsEmpty() && ImmunityQuery.Matches(GESpecToConsider))
+		{
+			ASC->OnImmunityBlockGameplayEffectDelegate.Broadcast(GESpecToConsider, ActiveGE);
+			return false;
+		}
+	}
+
+	return true;
+}
+
+#if WITH_EDITOR
+EDataValidationResult UImmunityGameplayEffectComponent::IsDataValid(FDataValidationContext& Context)
+{
+	EDataValidationResult Result = Super::IsDataValid(Context);
+
+	if (GetOwner()->DurationPolicy == EGameplayEffectDurationType::Instant)
+	{
+		Context.AddError(FText::FormatOrdered(LOCTEXT("IsInstantEffectError", "GE is an Instant Effect and incompatible with {0}"), FText::FromString(GetClass()->GetName())));
+		Result = EDataValidationResult::Invalid;
+	}
+
+	return Result;
+}
+#endif // WITH_EDITOR
+
+#undef LOCTEXT_NAMESPACE
diff --git a/Engine/Plugins/Runtime/GameplayAbilities/Source/GameplayAbilities/Private/GameplayEffectComponents/RemoveOtherGameplayEffectComponent.cpp b/Engine/Plugins/Runtime/GameplayAbilities/Source/GameplayAbilities/Private/GameplayEffectComponents/RemoveOtherGameplayEffectComponent.cpp
new file mode 100644
--- /dev/null
+++ b/Engine/Plugins/Runtime/GameplayAbilities/Source/GameplayAbilities/Private/GameplayEffectComponents/RemoveOtherGameplayEffectComponent.cpp
@@ -0,0 +1,64 @@
+// Copyright Epic Games, Inc. All Rights Reserved.
+
+#include "GameplayEffectComponents/RemoveOtherGameplayEffectComponent.h"
+#include "AbilitySystemComponent.h"
+#include "Misc/DataValidation.h"
+
+#define LOCTEXT_NAMESPACE "RemoveOtherGameplayEffectComponent"
+
+URemoveOtherGameplayEffectComponent::URemoveOtherGameplayEffectComponent()
+{
+#if WITH_EDITORONLY_DATA
+	EditorFriendlyName = TEXT("Remove Other Gameplay Effects");
+#endif
+}
+
+bool URemoveOtherGameplayEffectComponent::OnActiveGameplayEffectAdded(FActiveGameplayEffectsContainer& ActiveGEContainer, FActiveGameplayEffect& ActiveGE) const
+{
+	if (ActiveGEContainer.OwnerIsNetAuthority)
+	{
+		OnGameplayEffectApplied(ActiveGEContainer);
+	}
+	return true;
+}
+
+void URemoveOtherGameplayEffectComponent::OnGameplayEffectExecuted(FActiveGameplayEffectsContainer& ActiveGEContainer, FGameplayEffectSpec& GESpec, FPredictionKey& PredictionKey) const
+{
+	const bool bInstantEffect = (GetOwner()->DurationPolicy == EGameplayEffectDurationType::Instant);
+	if (bInstantEffect && ActiveGEContainer.OwnerIsNetAuthority)
+	{
+		OnGameplayEffectApplied(ActiveGEContainer);
+	}
+}
+
+void URemoveOtherGameplayEffectComponent::OnGameplayEffectApplied(FActiveGameplayEffectsContainer& ActiveGEContainer) const
+{
+	constexpr int32 RemoveAllStacks = -1;
+	for (const FGameplayEffectQuery& RemoveQuery : RemoveGameplayEffectQueries)
+	{
+		if (!RemoveQuery.IsEmpty())
+		{
+			ActiveGEContainer.RemoveActiveEffects(RemoveQuery, RemoveAllStacks);
+		}
+	}
+}
+
+#if WITH_EDITOR
+EDataValidationResult URemoveOtherGameplayEffectComponent::IsDataValid(FDataValidationContext& Context)
+{
+	EDataValidationResult Result = Super::IsDataValid(Context);
+
+	if (GetOwner()->DurationPolicy != EGameplayEffectDurationType::Instant)
+	{
+		if (GetOwner()->Period.Value > 0.0f)
+		{
+			Context.AddError(FText::FormatOrdered(LOCTEXT("PeriodicEffectError", "GE is Periodic. Remove {0} and use TagRequirements (Ongoing) instead."), FText::FromString(GetClass()->GetName())));
+			Result = EDataValidationResult::Invalid;
+		}
+	}
+
+	return Result;
+}
+#endif // WITH_EDITOR
+
+#undef LOCTEXT_NAMESPACE
diff --git a/Engine/Plugins/Runtime/GameplayAbilities/Source/GameplayAbilities/Private/GameplayEffectComponents/TargetTagRequirementsGameplayEffectComponent.cpp b/Engine/Plugins/Runtime/GameplayAbilities/Source/GameplayAbilities/Private/GameplayEffectComponents/TargetTagRequirementsGameplayEffectComponent.cpp
new file mode 100644
--- /dev/null
+++ b/Engine/Plugins/Runtime/GameplayAbilities/Source/GameplayAbilities/Private/GameplayEffectComponents/TargetTagRequirementsGameplayEffectComponent.cpp
@@ -0,0 +1,147 @@
+// Copyright Epic Games, Inc. All Rights Reserved.
+
+#include "GameplayEffectComponents/TargetTagRequirementsGameplayEffectComponent.h"
+#include "AbilitySystemComponent.h"
+#include "AbilitySystemLog.h"
+#include "Misc/DataValidation.h"
+
+#define LOCTEXT_NAMESPACE "TargetTagRequirementsGameplayEffectComponent"
+
+UTargetTagRequirementsGameplayEffectComponent::UTargetTagRequirementsGameplayEffectComponent()
+{
+#if WITH_EDITORONLY_DATA
+	EditorFriendlyName = TEXT("Target Tag Reqs (While GE is Active)");
+#endif // WITH_EDITORONLY_DATA
+}
+
+bool UTargetTagRequirementsGameplayEffectComponent::CanGameplayEffectApply(const FActiveGameplayEffectsContainer& ActiveGEContainer, const FGameplayEffectSpec& GESpec) const
+{
+	FGameplayTagContainer Tags;
+	ActiveGEContainer.Owner->GetOwnedGameplayTags(Tags);
+	
+	if (ApplicationTagRequirements.RequirementsMet(Tags) == false)
+	{
+		return false;
+	}
+
+	if (!RemovalTagRequirements.IsEmpty() && RemovalTagRequirements.RequirementsMet(Tags) == true)
+	{
+		return false;
+	}
+
+	return true;
+}
+
+void UTargetTagRequirementsGameplayEffectComponent::OnGameplayEffectExecuted(FActiveGameplayEffectsContainer& ActiveGEContainer, FGameplayEffectSpec& GESpec, FPredictionKey& PredictionKey) const
+{
+	const bool bInstantEffect = (GetOwner()->DurationPolicy == EGameplayEffectDurationType::Instant);
+	UE_CLOG(bInstantEffect && !OngoingTagRequirements.IsEmpty(), LogGameplayEffects, Warning, TEXT("%s is instant but has tag requirements. Tag requirements can only be used with gameplay effects that have a duration. Tag requirements will be ignored."), *GetOwner()->GetPathName());
+}
+
+// UTargetTagRequirementsGameplayEffectComponent lives on an asset.  This doesn't get instanced at runtime, so this is NOT A UNIQUE INSTANCE (it is a shared instance for any GEContainer/ActiveGE that wants to use it).
+bool UTargetTagRequirementsGameplayEffectComponent::OnActiveGameplayEffectAdded(FActiveGameplayEffectsContainer& GEContainer, FActiveGameplayEffect& ActiveGE) const
+{
+	UAbilitySystemComponent* ASC = GEContainer.Owner;
+	if (!ensure(ASC))
+	{
+		return false;
+	}
+
+	FActiveGameplayEffectHandle ActiveGEHandle = ActiveGE.Handle;
+	if (FActiveGameplayEffectEvents* EventSet = ASC->GetActiveEffectEventSet(ActiveGEHandle))
+	{
+		// Keep track of everything I'm binding so that I can unbind it OnEffectRemoved
+		TArray<TTuple<FGameplayTag,FDelegateHandle>> AllBoundEvents;
+
+		// Add our tag requirements to the ASC's Callbacks map. This helps filter down the amount of callbacks we'll get due to tag changes
+		// (rather than registering for the one callback whenever any tag changes).  We also need to keep track to remove those registered delegates.
+		auto BindAllTags = [this, ASC, ActiveGEHandle, &AllBoundEvents](const FGameplayTagContainer& TagContainer)
+		{
+			for (const FGameplayTag& Tag : TagContainer)
+			{
+				FOnGameplayEffectTagCountChanged& OnTagEvent = ASC->RegisterGameplayTagEvent(Tag, EGameplayTagEventType::NewOrRemoved);
+				FDelegateHandle Handle = OnTagEvent.AddUObject(this, &UTargetTagRequirementsGameplayEffectComponent::OnTagChanged, ActiveGEHandle);
+				AllBoundEvents.Emplace(Tag, Handle);
+			}
+		};
+		BindAllTags(OngoingTagRequirements.IgnoreTags);
+		BindAllTags(OngoingTagRequirements.RequireTags);
+		BindAllTags(RemovalTagRequirements.IgnoreTags);
+		BindAllTags(RemovalTagRequirements.RequireTags);
+
+		// Now when this Effect is removed, we should remove all of our registered callbacks.
+		EventSet->OnEffectRemoved.AddUObject(this, &UTargetTagRequirementsGameplayEffectComponent::OnGameplayEffectRemoved, ASC, MoveTemp(AllBoundEvents));
+	}
+	else
+	{
+		UE_LOG(LogGameplayEffects, Error, TEXT("UTargetTagRequirementsGameplayEffectComponent::OnGameplayEffectAdded called with ActiveGE: %s which had an invalid FActiveGameplayEffectHandle."), *ActiveGE.GetDebugString());
+	}
+
+	FGameplayTagContainer TagContainer;
+	ASC->GetOwnedGameplayTags(TagContainer);
+
+	return OngoingTagRequirements.RequirementsMet(TagContainer);
+}
+
+void UTargetTagRequirementsGameplayEffectComponent::OnGameplayEffectRemoved(const FGameplayEffectRemovalInfo& GERemovalInfo, UAbilitySystemComponent* ASC, TArray<TTuple<FGameplayTag, FDelegateHandle>> AllBoundEvents) const
+{
+	for (TTuple<FGameplayTag, FDelegateHandle>& Pair : AllBoundEvents)
+	{
+		const bool bSuccess = ASC->UnregisterGameplayTagEvent(Pair.Value, Pair.Key, EGameplayTagEventType::NewOrRemoved);
+		UE_CLOG(!bSuccess, LogGameplayEffects, Error, TEXT("%s tried to unregister GameplayTagEvent '%s' on GameplayEffect '%s' but failed."), *GetName(), *Pair.Key.ToString(), *GetNameSafe(GetOwner()));
+	}
+}
+
+void UTargetTagRequirementsGameplayEffectComponent::OnTagChanged(const FGameplayTag GameplayTag, int32 NewCount, FActiveGameplayEffectHandle ActiveGEHandle) const
+{
+	// Note: This function can remove us (RemoveActiveGameplayEffect eventually calling OnGameplayEffectRemoved),
+	// but we could be in the middle of a stack of OnTagChanged callbacks, wo we could get a stale OnTagChanged.
+	UAbilitySystemComponent* Owner = ActiveGEHandle.GetOwningAbilitySystemComponent();
+	if (!Owner)
+	{
+		return;
+	}
+
+	const FActiveGameplayEffect* ActiveGE = Owner->GetActiveGameplayEffect(ActiveGEHandle);
+	if (ensure(ActiveGE) && !ActiveGE->IsPendingRemove)
+	{
+		FGameplayTagContainer OwnedTags;
+		Owner->GetOwnedGameplayTags(OwnedTags);
+
+		const bool bRemovalRequirementsMet = !RemovalTagRequirements.IsEmpty() && RemovalTagRequirements.RequirementsMet(OwnedTags);
+		if (bRemovalRequirementsMet)
+		{
+			// This is slightly different functionality from pre-UE5.3, we're calling RemoveActiveGameplayEffect rather than InternalRemoveActiveGameplayEffect.
+			// The result is we set the calculated magnitudes back to zero.  This also used to only run on the Server.
+			Owner->RemoveActiveGameplayEffect(ActiveGEHandle);
+		}
+		else
+		{
+			// See if we should be inhibiting the execution
+			constexpr bool bInvokeCuesIfStateChanged = true;
+			const bool bOngoingRequirementsMet = OngoingTagRequirements.IsEmpty() || OngoingTagRequirements.RequirementsMet(OwnedTags);
+			Owner->InhibitActiveGameplayEffect(ActiveGEHandle, !bOngoingRequirementsMet, bInvokeCuesIfStateChanged);
+		}
+	}
+}
+
+#if WITH_EDITOR
+/**
+ * Validate incompatable configurations
+ */
+EDataValidationResult UTargetTagRequirementsGameplayEffectComponent::IsDataValid(FDataValidationContext& Context)
+{
+	EDataValidationResult Result = Super::IsDataValid(Context);
+
+	const bool bInstantEffect = (GetOwner()->DurationPolicy == EGameplayEffectDurationType::Instant);
+	if (bInstantEffect && !OngoingTagRequirements.IsEmpty())
+	{
+		Context.AddError(LOCTEXT("GEInstantAndOngoing", "GE is instant but has OngoingTagRequirements."));
+		Result = EDataValidationResult::Invalid;
+	}
+
+	return Result;
+}
+#endif // WITH_EDITOR
+
+#undef LOCTEXT_NAMESPACE
diff --git a/Engine/Plugins/Runtime/GameplayAbilities/Source/GameplayAbilities/Private/GameplayEffectComponents/TargetTagsGameplayEffectComponent.cpp b/Engine/Plugins/Runtime/GameplayAbilities/Source/GameplayAbilities/Private/GameplayEffectComponents/TargetTagsGameplayEffectComponent.cpp
new file mode 100644
--- /dev/null
+++ b/Engine/Plugins/Runtime/GameplayAbilities/Source/GameplayAbilities/Private/GameplayEffectComponents/TargetTagsGameplayEffectComponent.cpp
@@ -0,0 +1,41 @@
+// Copyright Epic Games, Inc. All Rights Reserved.
+
+#include "GameplayEffectComponents/TargetTagsGameplayEffectComponent.h"
+
+UTargetTagsGameplayEffectComponent::UTargetTagsGameplayEffectComponent()
+{
+#if WITH_EDITORONLY_DATA
+	EditorFriendlyName = TEXT("Target Tags (Granted to Actor)");
+#endif // WITH_EDITORONLY_DATA
+}
+
+void UTargetTagsGameplayEffectComponent::OnOwnerPostLoad()
+{
+	Super::OnOwnerPostLoad();
+	SetAndApplyTargetTagChanges(InheritableGrantedTagsContainer);
+}
+
+#if WITH_EDITOR
+void UTargetTagsGameplayEffectComponent::PostEditChangeProperty(FPropertyChangedEvent& PropertyChangedEvent)
+{
+	Super::PostEditChangeProperty(PropertyChangedEvent);
+
+	if (PropertyChangedEvent.GetMemberPropertyName() == GetInheritableGrantedTagsContainerName())
+	{
+		SetAndApplyTargetTagChanges(InheritableGrantedTagsContainer);
+	}
+}
+#endif // WITH_EDITOR
+
+void UTargetTagsGameplayEffectComponent::SetAndApplyTargetTagChanges(const FInheritedTagContainer& TagContainerMods)
+{
+	InheritableGrantedTagsContainer = TagContainerMods;
+
+	// Try to find the parent and update the inherited tags
+	const UTargetTagsGameplayEffectComponent* Parent = FindParentComponent(*this);
+	InheritableGrantedTagsContainer.UpdateInheritedTagProperties(Parent ? &Parent->InheritableGrantedTagsContainer : nullptr);
+
+	// Apply to the owning Gameplay Effect Component
+	UGameplayEffect* Owner = GetOwner();
+	InheritableGrantedTagsContainer.ApplyTo(Owner->CachedGrantedTags);
+}
diff --git a/Engine/Plugins/Runtime/GameplayAbilities/Source/GameplayAbilities/Private/GameplayEffectTypes.cpp b/Engine/Plugins/Runtime/GameplayAbilities/Source/GameplayAbilities/Private/GameplayEffectTypes.cpp
--- a/Engine/Plugins/Runtime/GameplayAbilities/Source/GameplayAbilities/Private/GameplayEffectTypes.cpp
+++ b/Engine/Plugins/Runtime/GameplayAbilities/Source/GameplayAbilities/Private/GameplayEffectTypes.cpp
@@ -864,6 +864,16 @@ FString FGameplayTagRequirements::ToString() const
 	return Str;
 }
 
+bool FGameplayTagRequirements::operator==(const FGameplayTagRequirements& Other) const
+{
+	return RequireTags == Other.RequireTags && IgnoreTags == Other.IgnoreTags;
+}
+
+bool FGameplayTagRequirements::operator!=(const FGameplayTagRequirements& Other) const
+{
+	return !(*this == Other);
+}
+
 void FActiveGameplayEffectsContainer::PrintAllGameplayEffects() const
 {
 	for (const FActiveGameplayEffect& Effect : this)
@@ -897,10 +907,9 @@ const FGameplayTagContainer* FTagContainerAggregator::GetAggregatedTags() const
 	if (CacheIsValid == false)
 	{
 		CacheIsValid = true;
-		CachedAggregator.Reset(CapturedActorTags.Num() + CapturedSpecTags.Num() + ScopedTags.Num());
+		CachedAggregator.Reset(CapturedActorTags.Num() + CapturedSpecTags.Num());
 		CachedAggregator.AppendTags(CapturedActorTags);
 		CachedAggregator.AppendTags(CapturedSpecTags);
-		CachedAggregator.AppendTags(ScopedTags);
 	}
 
 	return &CachedAggregator;
diff --git a/Engine/Plugins/Runtime/GameplayAbilities/Source/GameplayAbilities/Private/GameplayEffectUIData.cpp b/Engine/Plugins/Runtime/GameplayAbilities/Source/GameplayAbilities/Private/GameplayEffectUIData.cpp
deleted file mode 100644
diff --git a/Engine/Plugins/Runtime/GameplayAbilities/Source/GameplayAbilities/Public/Abilities/GameplayAbility.h b/Engine/Plugins/Runtime/GameplayAbilities/Source/GameplayAbilities/Public/Abilities/GameplayAbility.h
--- a/Engine/Plugins/Runtime/GameplayAbilities/Source/GameplayAbilities/Public/Abilities/GameplayAbility.h
+++ b/Engine/Plugins/Runtime/GameplayAbilities/Source/GameplayAbilities/Public/Abilities/GameplayAbility.h
@@ -291,7 +291,7 @@ public:
 	/** Returns the time in seconds remaining on the currently active cooldown and the original duration for this cooldown. */
 	virtual void GetCooldownTimeRemainingAndDuration(FGameplayAbilitySpecHandle Handle, const FGameplayAbilityActorInfo* ActorInfo, float& TimeRemaining, float& CooldownDuration) const;
 
-	/** Returns all tags that are currently on cooldown */
+	/** Returns all tags that can put this ability into cooldown */
 	virtual const FGameplayTagContainer* GetCooldownTags() const;
 	
 	/** Returns true if none of the ability's tags are blocked and if it doesn't have a "Blocking" tag and has all "Required" tags. */
diff --git a/Engine/Plugins/Runtime/GameplayAbilities/Source/GameplayAbilities/Public/AbilitySystemBlueprintLibrary.h b/Engine/Plugins/Runtime/GameplayAbilities/Source/GameplayAbilities/Public/AbilitySystemBlueprintLibrary.h
--- a/Engine/Plugins/Runtime/GameplayAbilities/Source/GameplayAbilities/Public/AbilitySystemBlueprintLibrary.h
+++ b/Engine/Plugins/Runtime/GameplayAbilities/Source/GameplayAbilities/Public/AbilitySystemBlueprintLibrary.h
@@ -315,11 +315,13 @@ class GAMEPLAYABILITIES_API UAbilitySystemBlueprintLibrary : public UBlueprintFu
 	static FGameplayEffectSpecHandle AddAssetTags(FGameplayEffectSpecHandle SpecHandle, FGameplayTagContainer NewGameplayTags);
 
 	/** Adds LinkedGameplayEffectSpec to SpecHandles. LinkedGameplayEffectSpec will be applied when/if SpecHandle is applied successfully. LinkedGameplayEffectSpec will not be modified here. Returns the ORIGINAL SpecHandle (legacy decision) */
-	UFUNCTION(BlueprintCallable, Category = "Ability|GameplayEffect")
+	UE_DEPRECATED(5.3, "Linked GameplayEffects aren't replicated.  The new UAdditionalGameplayEffectsComponent renders this functionality obsolete (and since it's configured in the Asset, it is properly synced).")
+	UFUNCTION(BlueprintCallable, Category = "Ability|GameplayEffect", meta = (DeprecatedFunction, DeprecatedMessage="Linked GameplayEffects aren't replicated.  Configure the GameplayEffect asset with a suitable GameplayEffectComponent."))
 	static FGameplayEffectSpecHandle AddLinkedGameplayEffectSpec(FGameplayEffectSpecHandle SpecHandle, FGameplayEffectSpecHandle LinkedGameplayEffectSpec);
 
 	/** Adds LinkedGameplayEffect to SpecHandles. LinkedGameplayEffectSpec will be applied when/if SpecHandle is applied successfully. This will initialize the LinkedGameplayEffect's Spec for you. Returns to NEW linked spec in case you want to add more to it. */
-	UFUNCTION(BlueprintCallable, Category = "Ability|GameplayEffect")
+	UE_DEPRECATED(5.3, "Linked GameplayEffects aren't replicated.  The new UAdditionalGameplayEffectsComponent renders this functionality obsolete (and since it's configured in the Asset, it is properly synced).")
+	UFUNCTION(BlueprintCallable, Category = "Ability|GameplayEffect", meta = (DeprecatedFunction, DeprecatedMessage = "Linked GameplayEffects aren't replicated.  Configure the GameplayEffect asset with a suitable GameplayEffectComponent."))
 	static FGameplayEffectSpecHandle AddLinkedGameplayEffect(FGameplayEffectSpecHandle SpecHandle, TSubclassOf<UGameplayEffect> LinkedGameplayEffect);
 
 	/** Sets the GameplayEffectSpec's StackCount to the specified amount (prior to applying) */
@@ -335,7 +337,8 @@ class GAMEPLAYABILITIES_API UAbilitySystemBlueprintLibrary : public UBlueprintFu
 	static FGameplayEffectContextHandle GetEffectContext(FGameplayEffectSpecHandle SpecHandle);
 
 	/** Returns handles for all Linked GE Specs that SpecHandle may apply. Useful if you want to append additional information to them. */
-	UFUNCTION(BlueprintPure, Category = "Ability|GameplayEffect")
+	UE_DEPRECATED(5.3, "Linked GameplayEffects aren't replicated.  The new UAdditionalGameplayEffectsComponent renders this functionality obsolete (and since it's configured in the Asset, it is properly synced).")
+	UFUNCTION(BlueprintPure, Category = "Ability|GameplayEffect", meta = (DeprecatedFunction, DeprecatedMessage = "Linked GameplayEffects aren't replicated.  Configure the GameplayEffect asset with a suitable GameplayEffectComponent."))
 	static TArray<FGameplayEffectSpecHandle> GetAllLinkedGameplayEffectSpecHandles(FGameplayEffectSpecHandle SpecHandle);
 
 	/** Manually adds a set of tags to a given actor, and optionally replicates them. */
diff --git a/Engine/Plugins/Runtime/GameplayAbilities/Source/GameplayAbilities/Public/AbilitySystemComponent.h b/Engine/Plugins/Runtime/GameplayAbilities/Source/GameplayAbilities/Public/AbilitySystemComponent.h
--- a/Engine/Plugins/Runtime/GameplayAbilities/Source/GameplayAbilities/Public/AbilitySystemComponent.h
+++ b/Engine/Plugins/Runtime/GameplayAbilities/Source/GameplayAbilities/Public/AbilitySystemComponent.h
@@ -76,6 +76,9 @@ DECLARE_MULTICAST_DELEGATE_OneParam(FAbilitySpecDirtied, const FGameplayAbilityS
 /** Notifies when GameplayEffectSpec is blocked by an ActiveGameplayEffect due to immunity  */
 DECLARE_MULTICAST_DELEGATE_TwoParams(FImmunityBlockGE, const FGameplayEffectSpec& /*BlockedSpec*/, const FActiveGameplayEffect* /*ImmunityGameplayEffect*/);
 
+/** We allow a list of delegates to decide if the application of a Gameplay Effect can be blocked. If it's blocked, it will call the ImmunityBlockGE above */
+DECLARE_DELEGATE_RetVal_TwoParams(bool, FGameplayEffectApplicationQuery, const FActiveGameplayEffectsContainer& /*ActiveGEContainer*/, const FGameplayEffectSpec& /*GESpecToConsider*/);
+
 /** How gameplay effects will be replicated to clients */
 UENUM()
 enum class EGameplayEffectReplicationMode : uint8
@@ -405,6 +408,9 @@ class GAMEPLAYABILITIES_API UAbilitySystemComponent : public UGameplayTasksCompo
 	UFUNCTION(BlueprintCallable, BlueprintAuthorityOnly, Category = GameplayEffects)
 	virtual void SetActiveGameplayEffectLevelUsingQuery(FGameplayEffectQuery Query, int32 NewLevel);
 
+	/** Inhibit an active gameplay effect so that it is disabled, but not removed */
+	virtual void InhibitActiveGameplayEffect(FActiveGameplayEffectHandle ActiveGEHandle, bool bInhibit, bool bInvokeGameplayCueEvents);
+
 	/**
 	 * Raw accessor to ask the magnitude of a gameplay effect, not necessarily always correct. How should outside code (UI, etc) ask things like 'how much is this gameplay effect modifying my damage by'
 	 * (most likely we want to catch this on the backend - when damage is applied we can get a full dump/history of how the number got to where it is. But still we may need polling methods like below (how much would my damage be)
@@ -520,6 +526,9 @@ class GAMEPLAYABILITIES_API UAbilitySystemComponent : public UGameplayTasksCompo
 	/** Called when an ability spec's internals have changed */
 	FAbilitySpecDirtied AbilitySpecDirtiedCallbacks;
 
+	/** We allow users to setup a series of functions that must be true in order to allow a GameplayEffect to be applied */
+	TArray<FGameplayEffectApplicationQuery> GameplayEffectApplicationQueries;
+
 	/** Call notify callbacks above */
 	virtual void NotifyAbilityCommit(UGameplayAbility* Ability);
 	virtual void NotifyAbilityActivated(const FGameplayAbilitySpecHandle Handle, UGameplayAbility* Ability);
@@ -693,7 +702,7 @@ class GAMEPLAYABILITIES_API UAbilitySystemComponent : public UGameplayTasksCompo
 	FOnGameplayEffectTagCountChanged& RegisterGameplayTagEvent(FGameplayTag Tag, EGameplayTagEventType::Type EventType=EGameplayTagEventType::NewOrRemoved);
 
 	/** Unregister previously added events */
-	void UnregisterGameplayTagEvent(FDelegateHandle DelegateHandle, FGameplayTag Tag, EGameplayTagEventType::Type EventType=EGameplayTagEventType::NewOrRemoved);
+	bool UnregisterGameplayTagEvent(FDelegateHandle DelegateHandle, FGameplayTag Tag, EGameplayTagEventType::Type EventType=EGameplayTagEventType::NewOrRemoved);
 
 	/** Register a tag event and immediately call it */
 	FDelegateHandle RegisterAndCallGameplayTagEvent(FGameplayTag Tag, FOnGameplayEffectTagCountChanged::FDelegate Delegate, EGameplayTagEventType::Type EventType=EGameplayTagEventType::NewOrRemoved);
@@ -1088,16 +1097,16 @@ class GAMEPLAYABILITIES_API UAbilitySystemComponent : public UGameplayTasksCompo
 	float GetAbilityLastActivatedTime() const { return AbilityLastActivatedTime; }
 
 	/** Returns an ability spec from a handle. If modifying call MarkAbilitySpecDirty */
-	FGameplayAbilitySpec* FindAbilitySpecFromHandle(FGameplayAbilitySpecHandle Handle);
+	FGameplayAbilitySpec* FindAbilitySpecFromHandle(FGameplayAbilitySpecHandle Handle) const;
 	
 	/** Returns an ability spec from a GE handle. If modifying call MarkAbilitySpecDirty */
-	FGameplayAbilitySpec* FindAbilitySpecFromGEHandle(FActiveGameplayEffectHandle Handle);
+	FGameplayAbilitySpec* FindAbilitySpecFromGEHandle(FActiveGameplayEffectHandle Handle) const;
 
 	/** Returns an ability spec corresponding to given ability class. If modifying call MarkAbilitySpecDirty */
-	FGameplayAbilitySpec* FindAbilitySpecFromClass(TSubclassOf<UGameplayAbility> InAbilityClass);
+	FGameplayAbilitySpec* FindAbilitySpecFromClass(TSubclassOf<UGameplayAbility> InAbilityClass) const;
 
 	/** Returns an ability spec from a handle. If modifying call MarkAbilitySpecDirty */
-	FGameplayAbilitySpec* FindAbilitySpecFromInputID(int32 InputID);
+	FGameplayAbilitySpec* FindAbilitySpecFromInputID(int32 InputID) const;
 
 	/**
 	 * Returns all abilities with the given InputID
@@ -1776,6 +1785,7 @@ protected:
 	/** Returns true if the specified ability should be activated from an event in this network mode */
 	bool HasNetworkAuthorityToActivateTriggeredAbility(const FGameplayAbilitySpec &Spec) const;
 
+	UE_DEPRECATED(5.3, "Use OnImmunityBlockGameplayEffectDelegate directly.  It is trigger from a UImmunityGameplayEffectComponent.  You can create your own GameplayEffectComponent if you need different functionality.")
 	virtual void OnImmunityBlockGameplayEffect(const FGameplayEffectSpec& Spec, const FActiveGameplayEffect* ImmunityGE);
 
 	// Internal gameplay cue functions
diff --git a/Engine/Plugins/Runtime/GameplayAbilities/Source/GameplayAbilities/Public/GameplayAbilitySpec.h b/Engine/Plugins/Runtime/GameplayAbilities/Source/GameplayAbilities/Public/GameplayAbilitySpec.h
--- a/Engine/Plugins/Runtime/GameplayAbilities/Source/GameplayAbilities/Public/GameplayAbilitySpec.h
+++ b/Engine/Plugins/Runtime/GameplayAbilities/Source/GameplayAbilities/Public/GameplayAbilitySpec.h
@@ -93,6 +93,9 @@ struct FGameplayAbilitySpecDef
 	/** This handle can be set if the SpecDef is used to create a real FGameplaybilitySpec */
 	UPROPERTY()
 	FGameplayAbilitySpecHandle	AssignedHandle;
+
+	bool operator==(const FGameplayAbilitySpecDef& Other) const;
+	bool operator!=(const FGameplayAbilitySpecDef& Other) const;
 };
 
 /**
diff --git a/Engine/Plugins/Runtime/GameplayAbilities/Source/GameplayAbilities/Public/GameplayEffect.h b/Engine/Plugins/Runtime/GameplayAbilities/Source/GameplayAbilities/Public/GameplayEffect.h
--- a/Engine/Plugins/Runtime/GameplayAbilities/Source/GameplayAbilities/Public/GameplayEffect.h
+++ b/Engine/Plugins/Runtime/GameplayAbilities/Source/GameplayAbilities/Public/GameplayEffect.h
@@ -22,8 +22,31 @@
 #include "UObject/ObjectKey.h"
 #include "GameplayEffect.generated.h"
 
+/**
+ * Gameplay Effects are bundles of functionality that are _applied_ to Actors.  Think of Gameplay Effects like something at _affects_ an Actor.
+ * Gameplay Effects are assets, and thus immutable at runtime.  There are small exceptions to this where hacks are done to _create_ a GE at runtime (but once created and configured, the data is not modified).
+ * 
+ * Gameplay Effects and Lifetime
+ *		- A GE can be executed instantly, or not.  If not, it has a duration (which can be infinite).  GE's that have durations are _added_ to the Active Gameplay Effects Container.
+ *		- A GE that is instant is said to be _executed_, where it never makes its way into the Active Gameplay Effects Container. (exceptions mentioned below).
+ *		- In both cases (instant and duration) the lingo we use is "Applied" and all forms thereof (e.g. Application).  So "CanApplyGameplayEffect" does not take into account if it's Instant or not.
+ *		- Periodic effects are executed at every period (so it is both Added and Executed).
+ *		- One exception to the above is when we're _predicting_ a Gameplay Effect to happen on the Client (when we're ahead of the Server).  Then we pretend it's a Duration effect and wait for server confirmation.
+ * 
+ * Gameplay Effect Components
+ *	Since Unreal 5.3, we have deprecated the Monolithic UGameplayEffect and instead rely on UGameplayEffectComponents.  Those Components are designed to allow users of the Engine to customize the behavior
+ *  of UGameplayEffects based on their project without having to derive their own specialized classes.  It should also be more clear what a UGameplayEffect does when looking at the class, as there are less properties.
+ *  UGameplayEffectComponents are implemented as instanced SubObjects.  This comes with some baggage, as the current version of Unreal Engine does not seamlessly support SubObjects out of the box.  There is a fix-up step
+ *  required in PostCDOCompiled, with an explanation of what is required to achieve this and the limitations.
+ * 
+ * Gameplay Effect Specs
+ *	Gameplay Effect Specs are the runtime versions of Gameplay Effects.  You can think of them as the instanced data wrapper around the Gameplay Effect (the GE is an asset).  The Blueprint functionality is thus more concerned with
+ *	Gameplay Effect Specs rather than Gameplay Effects and that is reflected in the AbilitySystemBlueprintLibrary.
+ */
+
 class UAbilitySystemComponent;
 class UGameplayEffect;
+class UGameplayEffectComponent;
 class UGameplayEffectCustomApplicationRequirement;
 class UGameplayEffectExecutionCalculation;
 class UGameplayEffectTemplate;
@@ -63,6 +86,16 @@ enum class EAttributeBasedFloatCalculationType : uint8
 	AttributeMagnitudeEvaluatedUpToChannel
 };
 
+/** The version of the UGameplayEffect. Used for upgrade paths. */
+UENUM()
+enum class EGameplayEffectVersion : uint8
+{
+	Monolithic,	// Legacy version from Pre-UE5.3 (before we were versioning)
+	Modular53,	// New modular version available in UE5.3
+
+	Current = Modular53
+};
+
 struct GAMEPLAYABILITIES_API FGameplayEffectConstants
 {
 	/** Infinite duration */
@@ -441,9 +474,12 @@ struct GAMEPLAYABILITIES_API FConditionalGameplayEffect
 	UPROPERTY(EditDefaultsOnly, BlueprintReadOnly, Category = GameplayEffect)
 	TSubclassOf<UGameplayEffect> EffectClass;
 
-	/** Tags that the source must have for this GE to apply */
+	/** Tags that the source must have for this GE to apply.  If this is blank, then there are no requirements to apply the EffectClass. */
 	UPROPERTY(EditDefaultsOnly, BlueprintReadOnly, Category = GameplayEffect)
 	FGameplayTagContainer RequiredSourceTags;
+
+	bool operator==(const FConditionalGameplayEffect& Other) const;
+	bool operator!=(const FConditionalGameplayEffect& Other) const;
 };
 
 /** 
@@ -476,10 +512,6 @@ struct GAMEPLAYABILITIES_API FGameplayEffectExecutionDefinition
 	UPROPERTY(EditDefaultsOnly, Category = Execution)
 	TArray<FGameplayEffectExecutionScopedModifierInfo> CalculationModifiers;
 
-	/** Deprecated. */
-	UPROPERTY()
-	TArray<TSubclassOf<UGameplayEffect>> ConditionalGameplayEffectClasses_DEPRECATED;
-
 	/** Other Gameplay Effects that will be applied to the target of this execution if the execution is successful. Note if no execution class is selected, these will always apply. */
 	UPROPERTY(EditDefaultsOnly, Category = Execution)
 	TArray<FConditionalGameplayEffect> ConditionalGameplayEffects;
@@ -580,23 +612,27 @@ struct GAMEPLAYABILITIES_API FInheritedTagContainer
 	GENERATED_USTRUCT_BODY()
 
 	/** Tags that I inherited and tags that I added minus tags that I removed*/
-	UPROPERTY(VisibleAnywhere, Transient, BlueprintReadOnly, Category = Application)
+	UPROPERTY(VisibleAnywhere, BlueprintReadOnly, Category = Application)
 	FGameplayTagContainer CombinedTags;
 
 	/** Tags that I have in addition to my parent's tags */
-	UPROPERTY(EditDefaultsOnly, Transient, BlueprintReadOnly, Category = Application)
+	UPROPERTY(EditDefaultsOnly, BlueprintReadOnly, Category = Application)
 	FGameplayTagContainer Added;
 
 	/** Tags that should be removed if my parent had them */
-	UPROPERTY(EditDefaultsOnly, Transient, BlueprintReadOnly, Category = Application)
+	UPROPERTY(EditDefaultsOnly, BlueprintReadOnly, Category = Application)
 	FGameplayTagContainer Removed;
 
 	void UpdateInheritedTagProperties(const FInheritedTagContainer* Parent);
 
-	void PostInitProperties();
+	/** Apply the Added and Removed tags to the passed-in container (does not have to be the previously configured Parent!) */
+	void ApplyTo(FGameplayTagContainer& ApplyToContainer) const;
 
 	void AddTag(const FGameplayTag& TagToAdd);
-	void RemoveTag(FGameplayTag TagToRemove);
+	void RemoveTag(const FGameplayTag& TagToRemove);
+
+	bool operator==(const FInheritedTagContainer& Other) const;
+	bool operator!=(const FInheritedTagContainer& Other) const;
 };
 
 /** Gameplay effect duration policies */
@@ -990,7 +1026,15 @@ struct GAMEPLAYABILITIES_API FGameplayEffectSpec
 
 	float GetDuration() const;
 	float GetPeriod() const;
-	float GetChanceToApplyToTarget() const;
+	
+	UE_DEPRECATED(5.3, "This no longer applies.  Use UChanceToApplyGameplayEffectComponent instead")
+	float GetChanceToApplyToTarget() const { return 1.0f; }
+
+	/** Sets the stack count for this GE to NewStackCount if stacking is supported. */
+	void SetStackCount(int32 NewStackCount);
+
+	/** Returns the stack count for this GE spec. */
+	int32 GetStackCount() const;
 
 	/** Set the context info: who and where this spec came from. */
 	void SetContext(FGameplayEffectContextHandle NewEffectContext, bool bSkipRecaptureSourceActorTags = false);
@@ -1001,13 +1045,13 @@ struct GAMEPLAYABILITIES_API FGameplayEffectSpec
 	}
 
 	/** Appends all tags granted by this gameplay effect spec */
-	void GetAllGrantedTags(OUT FGameplayTagContainer& Container) const;
+	void GetAllGrantedTags(OUT FGameplayTagContainer& OutContainer) const;
 
 	/** Appends all blocked ability tags granted by this gameplay effect spec */
-	void GetAllBlockedAbilityTags(FGameplayTagContainer& OutContainer) const;
+	void GetAllBlockedAbilityTags(OUT FGameplayTagContainer& OutContainer) const;
 
 	/** Appends all tags that apply to this gameplay effect spec */
-	void GetAllAssetTags(OUT FGameplayTagContainer& Container) const;
+	void GetAllAssetTags(OUT FGameplayTagContainer& OutContainer) const;
 
 	/** Sets the magnitude of a SetByCaller modifier */
 	void SetSetByCallerMagnitude(FName DataName, float Magnitude);
@@ -1094,6 +1138,7 @@ public:
 	FGameplayEffectAttributeCaptureSpecContainer CapturedRelevantAttributes;
 
 	/** other effects that need to be applied to the target if this effect is successful */
+	UE_DEPRECATED(5.3, "These TargetEffectSpecs are not replicated, thus can only apply to the server. Use UAdditionalGameplayEffectComponent instead (or roll your own solution)")
 	TArray< FGameplayEffectSpecHandle > TargetEffectSpecs;
 
 	/**
@@ -1108,9 +1153,9 @@ public:
 	UPROPERTY()
 	float Period;
 
-	/** The chance, in a 0.0-1.0 range, that this GameplayEffect will be applied to the target Attribute or GameplayEffect */
+	UE_DEPRECATED(5.3, "This variable no longer has any effect.  See UChanceToApplyGameplayEffectComponent")
 	UPROPERTY()
-	float ChanceToApplyToTarget;
+	float ChanceToApplyToTarget = 1.0f;
 
 	/** Captured Source Tags on GameplayEffectSpec creation */
 	UPROPERTY(NotReplicated)
@@ -1129,11 +1174,12 @@ public:
 	UPROPERTY()
 	FGameplayTagContainer DynamicAssetTags;
 
-	/** The calcuated modifiers for this effect */
+	/** The calculated modifiers for this effect */
 	UPROPERTY()
 	TArray<FModifierSpec> Modifiers;
 
 	/** Total number of stacks of this effect */
+	UE_DEPRECATED(5.3, "This member will be moved to private in the future.  Use GetStackCount and SetStackCount.")
 	UPROPERTY()
 	int32 StackCount;
 
@@ -1272,10 +1318,13 @@ struct GAMEPLAYABILITIES_API FActiveGameplayEffect : public FFastArraySerializer
 		return (Duration == FGameplayEffectConstants::INFINITE_DURATION ? -1.f : Duration + StartWorldTime);
 	}
 
+	/** This was the core function that turns the ActiveGE 'on' or 'off.  That function can be carried out by UGameplayEffectComponents, @see UTargetTagRequirementsGameplayEffectComponent */
+	UE_DEPRECATED(5.3, "CheckOngoingTagRequirements has been deprecated in favor of UTargetTagRequirementsGameplayEffectComponent")
 	void CheckOngoingTagRequirements(const FGameplayTagContainer& OwnerTags, struct FActiveGameplayEffectsContainer& OwningContainer, bool bInvokeGameplayCueEvents = false);
 	
 	/** Method to check if this effect should remove because the owner tags pass the RemovalTagRequirements requirement check */
-	bool CheckRemovalTagRequirements(const FGameplayTagContainer& OwnerTags, struct FActiveGameplayEffectsContainer& OwningContainer) const;
+	UE_DEPRECATED(5.3, "CheckRemovalTagRequirements has been deprecated in favor of UTargetTagRequirementsGameplayEffectComponent")
+	bool CheckRemovalTagRequirements(const FGameplayTagContainer & OwnerTags, struct FActiveGameplayEffectsContainer& OwningContainer) const;
 
 	void PrintAll() const;
 
@@ -1363,7 +1412,7 @@ public:
 	FActiveGameplayEffectQueryCustomMatch CustomMatchDelegate;
 
 	/** BP-exposed delegate for providing custom matching conditions. */
-	UPROPERTY(EditAnywhere, BlueprintReadWrite, Category = Query)
+	UPROPERTY(BlueprintReadWrite, Category = Query)
 	FActiveGameplayEffectQueryCustomMatch_Dynamic CustomMatchDelegate_BP;
 
 	/** Query that is matched against tags this GE gives */
@@ -1426,6 +1475,9 @@ public:
 	static FGameplayEffectQuery MakeQuery_MatchAllSourceTags(const FGameplayTagContainer& InTags);
 	/** Creates an effect query that will match if there are no common tags between the given tags and an ActiveGameplayEffect's source tags */
 	static FGameplayEffectQuery MakeQuery_MatchNoSourceTags(const FGameplayTagContainer& InTags);
+
+	bool operator==(const FGameplayEffectQuery& Other) const;
+	bool operator!=(const FGameplayEffectQuery& Other) const;
 };
 
 /**
@@ -1696,7 +1748,8 @@ struct GAMEPLAYABILITIES_API FActiveGameplayEffectsContainer : public FFastArray
 	int32 RemoveActiveEffects(const FGameplayEffectQuery& Query, int32 StacksToRemove);
 
 	/** Method called during effect application to process if any active effects should be removed from this effects application */
-	void AttemptRemoveActiveEffectsOnEffectApplication(const FGameplayEffectSpec &InSpec, const FActiveGameplayEffectHandle& InHandle);
+	UE_DEPRECATED(5.3, "AttemptRemoveActiveEffectsOnEffectApplication has been deprecated in favor of URemoveOtherGameplayEffectComponent")
+	void AttemptRemoveActiveEffectsOnEffectApplication(const FGameplayEffectSpec& InSpec, const FActiveGameplayEffectHandle& InHandle) {}
 
 	/**
 	 * Get the count of the effects matching the specified query (including stack count)
@@ -1731,9 +1784,8 @@ struct GAMEPLAYABILITIES_API FActiveGameplayEffectsContainer : public FFastArray
 
 	FOnGameplayAttributeValueChange& GetGameplayAttributeValueChangeDelegate(FGameplayAttribute Attribute);
 
-	void OnOwnerTagChange(FGameplayTag TagChange, int32 NewCount);
-
-	bool HasApplicationImmunityToSpec(const FGameplayEffectSpec& SpecToApply, const FActiveGameplayEffect*& OutGEThatProvidedImmunity) const;
+	UE_DEPRECATED(5.3, "Use UImmunityGameplayEffectComponent. This function will now always return false.")
+	bool HasApplicationImmunityToSpec(const FGameplayEffectSpec& SpecToApply, const FActiveGameplayEffect*& OutGEThatProvidedImmunity) const { return false; }
 
 	void IncrementLock();
 	void DecrementLock();
@@ -1799,12 +1851,12 @@ private:
 	 * @return True if the mod successfully executed, false if it did not
 	 */
 	bool InternalExecuteMod(FGameplayEffectSpec& Spec, FGameplayModifierEvaluatedData& ModEvalData);
-
+public:
 	bool IsNetAuthority() const
 	{
 		return OwnerIsNetAuthority;
 	}
-
+private:
 	void InternalExecutePeriodicGameplayEffect(FActiveGameplayEffect& ActiveEffect);
 
 	/** Called internally to actually remove a GameplayEffect or to reduce its StackCount. Returns true if we resized our internal GameplayEffect array. */
@@ -1829,9 +1881,6 @@ private:
 	/** Internal callback fired as a result of a custom modifier magnitude external dependency delegate firing; Updates affected active gameplay effects as necessary */
 	void OnCustomMagnitudeExternalDependencyFired(TSubclassOf<UGameplayModMagnitudeCalculation> MagnitudeCalculationClass);
 
-	/** Internal helper function to apply expiration effects from a removed/expired gameplay effect spec */
-	void InternalApplyExpirationEffects(const FGameplayEffectSpec& ExpiringSpec, bool bPrematureRemoval);
-
 	void RestartActiveGameplayEffectDuration(FActiveGameplayEffect& ActiveGameplayEffect);
 
 	// -------------------------------------------------------------------------------------------
@@ -1843,20 +1892,12 @@ private:
 	
 	TMap<FGameplayAttribute, FOnGameplayAttributeValueChange> AttributeValueChangeDelegates;
 
-	TMap<FGameplayTag, TSet<FActiveGameplayEffectHandle> >	ActiveEffectTagDependencies;
-
 	/** Mapping of custom gameplay modifier magnitude calculation class to dependency handles for triggering updates on external delegates firing */
 	TMap<FObjectKey, FCustomModifierDependencyHandle> CustomMagnitudeClassDependencies;
 
 	/** A map to manage stacking while we are the source */
 	TMap<TWeakObjectPtr<UGameplayEffect>, TArray<FActiveGameplayEffectHandle> >	SourceStackingMap;
 	
-	/** Acceleration struct for immunity tests */
-	FGameplayTagCountContainer ApplicationImmunityGameplayTagCountContainer;
-
-	/** Active GEs that have immunity queries. This is an acceleration list to avoid searching through the Active GameplayEffect list frequetly. (We only search for the active GE if immunity procs) */
-	UPROPERTY(NotReplicated)
-	TArray<TObjectPtr<const UGameplayEffect>> ApplicationImmunityQueryEffects;
 
 	FAggregatorRef& FindOrCreateAttributeAggregator(FGameplayAttribute Attribute);
 
@@ -1933,12 +1974,49 @@ private:
 
 // -------------------------------------------------------------------------------------
 
+/**
+ * Gameplay Effects Data needs to be versioned (e.g. going from Monolithic to Modular)
+ * Unfortunately, the Custom Package Versioning doesn't work for this case, because we're upgrading
+ * outside of the Serialize function.  For that reason we want to store a UProperty that says what version
+ * we're on.  Unfortunately that propagates from Parent to Child (inheritance rules) so to overcome that,
+ * we have a special UStruct that always serializes its data (so it will always be loaded, not inherited).
+ * Here is how to do that:
+ *	1. Return false in Identical (which effectively disables delta serialization).
+ *  2. Reset the value in PostInitProperties, thereby not using the inherited value on the initial preload.
+ *  3. Ensure the CurrentVersion field isn't marked up as a UPROPERTY to avoid automatic copying between parent/child.
+ *  4. Implement Serialize to ensure the CurrentVersion is still serialized.
+ */
+USTRUCT()
+struct FGameplayEffectVersion
+{
+	GENERATED_BODY()
+
+	/** The version the owning GameplayEffect is currently set to */
+	EGameplayEffectVersion CurrentVersion;
+
+	/** By always returning false here, we can disable delta serialization */
+	bool Identical(const FGameplayEffectVersion* Other, uint32 PortFlags) const { return false; }
+
+	/** We must use a custom serializer to make sure CurrentVersion serializes properly (without markup to avoid unwanted copying) */
+	bool Serialize(FStructuredArchive::FSlot Slot);
+};
+
+template<>
+struct TStructOpsTypeTraits<FGameplayEffectVersion> : public TStructOpsTypeTraitsBase2<FGameplayEffectVersion>
+{
+	enum
+	{
+		WithIdentical = true,
+		WithStructuredSerializer = true,
+	};
+};
+
 /**
  * UGameplayEffect
  *	The GameplayEffect definition. This is the data asset defined in the editor that drives everything.
  *  This is only blueprintable to allow for templating gameplay effects. Gameplay effects should NOT contain blueprint graphs.
  */
-UCLASS(Blueprintable, meta = (ShortTooltip="A GameplayEffect modifies attributes and tags."))
+UCLASS(Blueprintable, PrioritizeCategories="Status Duration GameplayEffect GameplayCues Stacking", meta = (ShortTooltip = "A GameplayEffect modifies attributes and tags."))
 class GAMEPLAYABILITIES_API UGameplayEffect : public UObject, public IGameplayTagAssetInterface
 {
 
@@ -1951,102 +2029,216 @@ public:
 	static const float NO_PERIOD;	
 	static const float INVALID_LEVEL;
 
+	UE_DEPRECATED(5.3, "The implementation and method name did not match.  Use GetGrantedTags() to get the tags Granted to the Actor this GameplayEffect is applied to.")
 	virtual void GetOwnedGameplayTags(FGameplayTagContainer& TagContainer) const override;
+
+	UE_DEPRECATED(5.3, "The implementation and method name did not match.  Use GetGrantedTags().HasTag() to check against the tags this GameplayEffect will Grant to the Actor.")
+	virtual bool HasMatchingGameplayTag(FGameplayTag TagToCheck) const override;
+
+	UE_DEPRECATED(5.3, "The implementation and method name did not match.  Use GetGrantedTags().HasAll() to check against the tags this GameplayEffect will Grant to the Actor.")
+	virtual bool HasAllMatchingGameplayTags(const FGameplayTagContainer& TagContainer) const override;
+
+	UE_DEPRECATED(5.3, "The implementation and method name did not match.  Use GetGrantedTags().HasAny() to check against the tags this GameplayEffect will Grant to the Actor.")
+	virtual bool HasAnyMatchingGameplayTags(const FGameplayTagContainer& TagContainer) const override;
+
 	virtual void GetBlockedAbilityTags(FGameplayTagContainer& OutTagContainer) const;
+
+	// Needed to properly disable inheriting the version value from its parent.
 	virtual void PostInitProperties() override;
+
 	virtual void PostLoad() override;
-	PRAGMA_DISABLE_DEPRECATION_WARNINGS // Suppress compiler warning on override of deprecated function
-	UE_DEPRECATED(5.0, "Use version that takes FObjectPreSaveContext instead.")
-	virtual void PreSave(const class ITargetPlatform* TargetPlatform) override;
-	PRAGMA_ENABLE_DEPRECATION_WARNINGS
-	virtual void PreSave(FObjectPreSaveContext ObjectSaveContext) override;
+
 #if WITH_EDITOR
-	void PostEditChangeProperty(FPropertyChangedEvent& PropertyChangedEvent) override;
-#endif
+	/** Do our upgrades in PostCDOCompiled */
+	virtual void PostCDOCompiled(const FPostCDOCompiledContext& Context) override;
 
-	/** Update all the InheritedTagContainers */
-	void UpdateInheritedTagProperties();
+	/** We need to fix-up all of the SubObjects manually since the Engine doesn't fully support this */
+	void PostCDOCompiledFixupSubobjects();
+
+	virtual void PostEditChangeProperty(FPropertyChangedEvent& PropertyChangedEvent) override;
+#endif
 
 	/** Check for any errors */
-	void ValidateGameplayEffect();
+	UE_DEPRECATED(5.3, "This was never implemented. The proper way to do this now is to use IsDataValid")
+	void ValidateGameplayEffect() {}
+
+	/**
+	 * Can we Apply this Gameplay Effect?  Note: Apply is the generic term for adding to the active container or executing it.
+	 */
+	bool CanApply(const FActiveGameplayEffectsContainer& ActiveGEContainer, const FGameplayEffectSpec& GESpec) const;
+
+	/**
+	 * Receive a notify that this GameplayEffect has been added to an Active Container.  ActiveGE will be the Active version of this GameplayEffect.
+	 * Returns true if this GameplayEffect should be active (or false to inhibit).
+	 */
+	bool OnAddedToActiveContainer(FActiveGameplayEffectsContainer& ActiveGEContainer, FActiveGameplayEffect& ActiveGE) const;
+
+	/**
+     * Receive a notify that this GameplayEffect has been executed (it must be instant, as it is not added to the Container).
+     * Since it is not added to the container, it will not have an associated FActiveGameplayEffect, just the FGameplayEffectSpec.
+     */
+	void OnExecuted(FActiveGameplayEffectsContainer& ActiveGEContainer, FGameplayEffectSpec& GESpec, FPredictionKey& PredictionKey) const;
+
+	/** Returns all tags that this GE *has* and *does not* grant to any Actor. */
+	const FGameplayTagContainer& GetAssetTags() const { return CachedAssetTags; }
+
+	/** Returns all tags granted to the Target Actor of this gameplay effect. That is to say: If this GE is applied to an Actor, they get these tags. */
+	const FGameplayTagContainer& GetGrantedTags() const { return CachedGrantedTags; }
+
+	/** Returns all blocking ability tags granted by this gameplay effect definition. Blocking tags means the target will not be able to execute an ability with this asset tag. */
+	const FGameplayTagContainer& GetBlockedAbilityTags() const { return CachedBlockedAbilityTags; }
+
+	/** Returns the maximum stack size for this gameplay effect */
+	int32 GetStackLimitCount() const;
+
+	/** Returns the stack expiration policy for this gameplay effect */
+	EGameplayEffectStackingExpirationPolicy GetStackExpirationPolicy() const;
 
 	// ---------------------------------------------------------------------------------------------------------------------------------
 
+	/**
+	 * Find and return the component of class GEComponentClass on this GE, if one exists.
+	 * Note: The const is a hint that these are assets that should not be modified at runtime. If you're building your own dynamic GameplayEffect: @see AddComponent, FindOrAddComponent.
+	 *
+	 * @return	The component matching the query (if it exists).
+	 */
+	template<typename GEComponentClass>
+	const GEComponentClass* FindComponent() const;
+
+	/** @return the first component that derives from the passed-in ClassToFind, if one exists. */
+	const UGameplayEffectComponent* FindComponent(TSubclassOf<UGameplayEffectComponent> ClassToFind) const;
+
+	/** 
+	 * Add a GameplayEffectComponent to the GameplayEffect.  This does not check for duplicates and is guaranteed to return a new instance.
+	 * 
+	 * @return	The newly added instanced GameplayEffectComponent.
+	 */
+	template<typename GEComponentClass>
+	GEComponentClass& AddComponent();
+
+	/**
+	 * Find an existing GameplayEffectComponent of the requested class, or add one if none are found.
+	 * 
+	 * @return	The existing instance of GEComponentClass, or a newly added one.
+	 */
+	template<typename GEComponentClass>
+	GEComponentClass& FindOrAddComponent();
+
+#if WITH_EDITOR
+	/** Allow each Gameplay Effect Component to validate its own data.  Call this version (Super::IsDataValid) _after_ your overridden version to update EditorStatusText. */
+	virtual EDataValidationResult IsDataValid(class FDataValidationContext& Context) override;
+
+protected:
+	// ----------------------------------------------------------------------
+	//	Upgrade Path
+	// ----------------------------------------------------------------------
+	
+	/** Let's keep track of the version so that we can upgrade the Components properly.  This will be set properly in PostLoad after upgrades. */
+	EGameplayEffectVersion GetVersion() const;
+
+	/** Sets the Version of the class to denote it's been upgraded */
+	void SetVersion(EGameplayEffectVersion Version);
+
+private:
+	/** Logic that determines if we can upgrade this Gameplay Effect to use Gameplay Effect Components (or not if it's already been done). Used in the upgrade paths below. */
+	bool VersionAllowsCreateGEComponentsDuringUpgrade() const;
+
+	// Helper functions for converting data to use components
+	void ConvertAdditionalEffectsComponent();
+	void ConvertAssetTagsComponent();
+	void ConvertBlockByTagsComponent();
+	void ConvertChanceToApplyComponent();
+	void ConvertCustomCanApplyComponent();
+	void ConvertImmunityComponent();
+	void ConvertRemoveOtherComponent();
+	void ConvertTagRequirementsComponent();
+	void ConvertTargetTagsComponent();
+	void ConvertUIComponent();
+#endif
+
+	// ----------------------------------------------------------------------
+	//	Properties
+	// ----------------------------------------------------------------------
+public:
 	/** Policy for the duration of this effect */
-	UPROPERTY(EditDefaultsOnly, Category=GameplayEffect)
+	UPROPERTY(EditDefaultsOnly, Category=Duration)
 	EGameplayEffectDurationType DurationPolicy;
 
 	/** Duration in seconds. 0.0 for instantaneous effects; -1.0 for infinite duration. */
-	UPROPERTY(EditDefaultsOnly, Category=GameplayEffect)
+	UPROPERTY(EditDefaultsOnly, Category=Duration, meta=(EditCondition="DurationPolicy == EGameplayEffectDurationType::HasDuration", EditConditionHides))
 	FGameplayEffectModifierMagnitude DurationMagnitude;
 
 	/** Period in seconds. 0.0 for non-periodic effects */
-	UPROPERTY(EditDefaultsOnly, BlueprintReadOnly, Category=Period)
+	UPROPERTY(EditDefaultsOnly, BlueprintReadOnly, Category="Duration|Period", meta=(EditCondition="DurationPolicy != EGameplayEffectDurationType::Instant", EditConditionHides))
 	FScalableFloat	Period;
 	
 	/** If true, the effect executes on application and then at every period interval. If false, no execution occurs until the first period elapses. */
-	UPROPERTY(EditDefaultsOnly, BlueprintReadOnly, Category=Period)
+	UPROPERTY(EditDefaultsOnly, BlueprintReadOnly, Category="Duration|Period", meta=(EditCondition="true", EditConditionHides)) // EditCondition in FGameplayEffectDetails
 	bool bExecutePeriodicEffectOnApplication;
 
-	UPROPERTY(EditDefaultsOnly, BlueprintReadOnly, Category=Period)
+	/** How we should respond when a periodic gameplay effect is no longer inhibited */
+	UPROPERTY(EditDefaultsOnly, BlueprintReadOnly, Category="Duration|Period", meta=(EditCondition="true", EditConditionHides)) // EditCondition in FGameplayEffectDetails
 	EGameplayEffectPeriodInhibitionRemovedPolicy PeriodicInhibitionPolicy;
 
 	/** Array of modifiers that will affect the target of this effect */
 	UPROPERTY(EditDefaultsOnly, BlueprintReadOnly, Category=GameplayEffect, meta=(TitleProperty=Attribute))
 	TArray<FGameplayModifierInfo> Modifiers;
 
+	/** Array of executions that will affect the target of this effect */
 	UPROPERTY(EditDefaultsOnly, BlueprintReadOnly, Category = GameplayEffect)
 	TArray<FGameplayEffectExecutionDefinition> Executions;
 
 	/** Probability that this gameplay effect will be applied to the target actor (0.0 for never, 1.0 for always) */
-	UPROPERTY(EditDefaultsOnly, Category=Application, meta=(GameplayAttribute="True"))
-	FScalableFloat ChanceToApplyToTarget;
+	UE_DEPRECATED(5.3, "Chance To Apply To Target is deprecated. Use the UChanceToApplyGameplayEffectComponent instead.")
+	UPROPERTY(meta = (GameplayAttribute = "True", DeprecatedProperty))
+	FScalableFloat ChanceToApplyToTarget_DEPRECATED;
 
-	UPROPERTY(EditDefaultsOnly, Category=Application, DisplayName="Application Requirement")
-	TArray<TSubclassOf<UGameplayEffectCustomApplicationRequirement> > ApplicationRequirements;
-
-	/** Deprecated. Use ConditionalGameplayEffects instead */
-	UPROPERTY()
-	TArray<TSubclassOf<UGameplayEffect>> TargetEffectClasses_DEPRECATED;
+	UE_DEPRECATED(5.3, "Application Requirements is deprecated. Use the UCustomCanApplyGameplayEffectComponent instead.")
+	UPROPERTY(meta = (DisplayName = "Application Requirement", DeprecatedProperty))
+	TArray<TSubclassOf<UGameplayEffectCustomApplicationRequirement> > ApplicationRequirements_DEPRECATED;
 
 	/** other gameplay effects that will be applied to the target of this effect if this effect applies */
-	UPROPERTY(EditDefaultsOnly, BlueprintReadOnly, Category = GameplayEffect)
+	UE_DEPRECATED(5.3, "Conditional Gameplay Effects is deprecated. Use the UAdditionalEffectsGameplayEffectComponent instead.")
+	UPROPERTY(BlueprintReadOnly, Category = "Deprecated|GameplayEffect", meta = (DeprecatedProperty))
 	TArray<FConditionalGameplayEffect> ConditionalGameplayEffects;
 
 	/** Effects to apply when a stacking effect "overflows" its stack count through another attempted application. Added whether the overflow application succeeds or not. */
-	UPROPERTY(EditDefaultsOnly, BlueprintReadOnly, Category=Overflow)
+	UPROPERTY(EditDefaultsOnly, BlueprintReadOnly, Category="Stacking|Overflow", meta = (EditConditionHides, EditCondition = "StackingType != EGameplayEffectStackingType::None"))
 	TArray<TSubclassOf<UGameplayEffect>> OverflowEffects;
 
 	/** If true, stacking attempts made while at the stack count will fail, resulting in the duration and context not being refreshed */
-	UPROPERTY(EditDefaultsOnly, BlueprintReadOnly, Category=Overflow)
+	UPROPERTY(EditDefaultsOnly, BlueprintReadOnly, Category="Stacking|Overflow", meta = (EditConditionHides, EditCondition = "StackingType != EGameplayEffectStackingType::None"))
 	bool bDenyOverflowApplication;
 
 	/** If true, the entire stack of the effect will be cleared once it overflows */
-	UPROPERTY(EditDefaultsOnly, BlueprintReadOnly, Category=Overflow, meta=(EditCondition="bDenyOverflowApplication"))
+	UPROPERTY(EditDefaultsOnly, BlueprintReadOnly, Category="Stacking|Overflow", meta=(EditConditionHides, EditCondition="(StackingType != EGameplayEffectStackingType::None) && bDenyOverflowApplication"))
 	bool bClearStackOnOverflow;
 
 	/** Effects to apply when this effect is made to expire prematurely (like via a forced removal, clear tags, etc.); Only works for effects with a duration */
-	UPROPERTY(EditDefaultsOnly, BlueprintReadOnly, Category=Expiration)
+	UE_DEPRECATED(5.3, "Premature Expiration Effect Classes is deprecated. Use the UAdditionalEffectsGameplayEffectComponent instead.")
+	UPROPERTY(BlueprintReadOnly, Category="Deprecated|Expiration", meta = (DeprecatedProperty))
 	TArray<TSubclassOf<UGameplayEffect>> PrematureExpirationEffectClasses;
 
 	/** Effects to apply when this effect expires naturally via its duration; Only works for effects with a duration */
-	UPROPERTY(EditDefaultsOnly, BlueprintReadOnly, Category=Expiration)
+	UE_DEPRECATED(5.3, "Routine Expiration Effect Classes is deprecated. Use the UAdditionalEffectsGameplayEffectComponent instead.")
+	UPROPERTY(BlueprintReadOnly, Category="Deprecated|Expiration", meta = (DeprecatedProperty))
 	TArray<TSubclassOf<UGameplayEffect>> RoutineExpirationEffectClasses;
 
 	/** If true, cues will only trigger when any GE modifiers succeed being applied (whether through modifiers or executions), or if there were no modifiers application attempts. */
-	UPROPERTY(EditDefaultsOnly, BlueprintReadOnly, Category = Display)
+	UPROPERTY(EditDefaultsOnly, BlueprintReadOnly, Category = "GameplayCues")
 	bool bRequireModifierSuccessToTriggerCues;
 
 	/** If true, GameplayCues will only be triggered for the first instance in a stacking GameplayEffect. */
-	UPROPERTY(EditDefaultsOnly, Category = Display)
+	UPROPERTY(EditDefaultsOnly, Category = "GameplayCues")
 	bool bSuppressStackingCues;
 
 	/** Cues to trigger non-simulated reactions in response to this GameplayEffect such as sounds, particle effects, etc */
-	UPROPERTY(EditDefaultsOnly, BlueprintReadOnly, Category = Display)
+	UPROPERTY(EditDefaultsOnly, BlueprintReadOnly, Category = "GameplayCues")
 	TArray<FGameplayEffectCue>	GameplayCues;
 
 	/** Data for the UI representation of this effect. This should include things like text, icons, etc. Not available in server-only builds. */
-	UPROPERTY(EditDefaultsOnly, Instanced, BlueprintReadOnly, Category = Display)
+	UE_DEPRECATED(5.3, "UI Data is deprecated. UGameplayEffectUIData now dervies from UGameplayEffectComponent, add it as a GameplayEffectComponent.  You can then access it with FindComponent<UGameplayEffectUIData>().")
+	UPROPERTY(BlueprintReadOnly, Transient, Instanced, Category = "Deprecated|Display", meta = (DeprecatedProperty))
 	TObjectPtr<class UGameplayEffectUIData> UIData;
 
 	// ----------------------------------------------------------------------
@@ -2054,79 +2246,167 @@ public:
 	// ----------------------------------------------------------------------
 	
 	/** The GameplayEffect's Tags: tags the the GE *has* and DOES NOT give to the actor. */
-	UPROPERTY(EditDefaultsOnly, BlueprintReadOnly, Category = Tags, meta = (DisplayName = "GameplayEffectAssetTag", Categories="GameplayEffectTagsCategory"))
+	UE_DEPRECATED(5.3, "Inheritable Gameplay Effect Tags is deprecated. To configure, add a UAssetTagsGameplayEffectComponent.  To access, use GetAssetTags.")
+	UPROPERTY(BlueprintReadOnly, Category = Deprecated, meta = (DisplayName = "GameplayEffectAssetTag", Categories="GameplayEffectTagsCategory", DeprecatedProperty))
 	FInheritedTagContainer InheritableGameplayEffectTags;
 	
 	/** These tags are applied to the actor I am applied to */
-	UPROPERTY(EditDefaultsOnly, BlueprintReadOnly, Category = Tags, meta=(DisplayName="GrantedTags", Categories="OwnedTagsCategory"))
+	UE_DEPRECATED(5.3, "Inheritable Owned Tags Container is deprecated. To configure, add a UTargetTagsGameplayEffectComponent. To access, use GetGrantedTags.")
+	UPROPERTY(BlueprintReadOnly, Category = Deprecated, meta = (DisplayName="GrantedTags", Categories="OwnedTagsCategory", DeprecatedProperty))
 	FInheritedTagContainer InheritableOwnedTagsContainer;
 	
 	/** These blocked ability tags are applied to the actor I am applied to */
-	UPROPERTY(EditDefaultsOnly, BlueprintReadOnly, Category = Tags, meta=(DisplayName="GrantedBlockedAbilityTags", Categories="BlockedAbilityTagsCategory"))
+	UE_DEPRECATED(5.3, "Inheritable Blocked Ability Tags Container is deprecated. Use the UTargetTagsGameplayEffectComponent instead.  To access, use GetBlockedAbilityTags.")
+	UPROPERTY(BlueprintReadOnly, Category = Deprecated, meta = (DisplayName="GrantedBlockedAbilityTags", Categories="BlockedAbilityTagsCategory", DeprecatedProperty))
 	FInheritedTagContainer InheritableBlockedAbilityTagsContainer;
 	
 	/** Once Applied, these tags requirements are used to determined if the GameplayEffect is "on" or "off". A GameplayEffect can be off and do nothing, but still applied. */
-	UPROPERTY(EditDefaultsOnly, BlueprintReadOnly, Category = Tags, meta=(Categories="OngoingTagRequirementsCategory"))
+	UE_DEPRECATED(5.3, "Ongoing Tag Requirements is deprecated. Use the UTargetTagRequirementsGameplayEffectComponent instead.")
+	UPROPERTY(BlueprintReadOnly, Category = Deprecated, meta = (Categories="OngoingTagRequirementsCategory", DeprecatedProperty))
 	FGameplayTagRequirements OngoingTagRequirements;
 
 	/** Tag requirements for this GameplayEffect to be applied to a target. This is pass/fail at the time of application. If fail, this GE fails to apply. */
-	UPROPERTY(EditDefaultsOnly, BlueprintReadOnly, Category = Tags, meta=(Categories="ApplicationTagRequirementsCategory"))
+	UE_DEPRECATED(5.3, "Application Tag Requirements is deprecated. Use the UTargetTagRequirementsGameplayEffectComponent instead.")
+	UPROPERTY(BlueprintReadOnly, Category = Deprecated, meta = (Categories="ApplicationTagRequirementsCategory", DeprecatedProperty))
 	FGameplayTagRequirements ApplicationTagRequirements;
 
 	/** Tag requirements that if met will remove this effect. Also prevents effect application. */
-	UPROPERTY(EditDefaultsOnly, BlueprintReadOnly, Category = Tags, meta=(Categories="ApplicationTagRequirementsCategory"))
+	UE_DEPRECATED(5.3, "Removal Tag Requirements is deprecated. Use the URemoveOtherGameplayEffectComponent instead.")
+	UPROPERTY(BlueprintReadOnly, Category = Deprecated, meta = (Categories="ApplicationTagRequirementsCategory", DeprecatedProperty))
 	FGameplayTagRequirements RemovalTagRequirements;
 
 	/** GameplayEffects that *have* tags in this container will be cleared upon effect application. */
-	UPROPERTY(EditDefaultsOnly, BlueprintReadOnly, Category = Tags, meta=(Categories="RemoveTagRequirementsCategory"))
+	UE_DEPRECATED(5.3, "Remove Gameplay Effects With Tags is deprecated. Use the UTargetTagRequirementsGameplayEffectComponent instead.")
+	UPROPERTY(BlueprintReadOnly, Category = Deprecated, meta = (Categories="RemoveTagRequirementsCategory", DeprecatedProperty))
 	FInheritedTagContainer RemoveGameplayEffectsWithTags;
 
 	/** Grants the owner immunity from these source tags.  */
-	UPROPERTY(EditDefaultsOnly, BlueprintReadOnly, Category = Immunity, meta = (DisplayName = "GrantedApplicationImmunityTags", Categories="GrantedApplicationImmunityTagsCategory"))
+	UE_DEPRECATED(5.3, "Granted Application Immunity Tags is deprecated. Use the UImmunityGameplayEffectComponent instead.")
+	UPROPERTY(BlueprintReadOnly, Category = Deprecated, meta = (DisplayName = "GrantedApplicationImmunityTags", Categories="GrantedApplicationImmunityTagsCategory", DeprecatedProperty))
 	FGameplayTagRequirements GrantedApplicationImmunityTags;
 
 	/** Grants immunity to GameplayEffects that match this query. Queries are more powerful but slightly slower than GrantedApplicationImmunityTags. */
-	UPROPERTY(EditDefaultsOnly, BlueprintReadOnly, Category = Immunity)
+	UE_DEPRECATED(5.3, "Granted Application Immunity Query is deprecated. Use the UImmunityGameplayEffectComponent instead.")
+	UPROPERTY(BlueprintReadOnly, Category = Deprecated, meta = (DeprecatedProperty))
 	FGameplayEffectQuery GrantedApplicationImmunityQuery;
 
 	/** Cached !GrantedApplicationImmunityQuery.IsEmpty(). Set on PostLoad. */
+	UE_DEPRECATED(5.3, "HasGrantedApplicationImmunityQuery is deprecated.  Use the UImmunityGameplayEffectComponent instead.")
 	bool HasGrantedApplicationImmunityQuery = false;
 
 	/** On Application of an effect, any active effects with this this query that matches against the added effect will be removed. Queries are more powerful but slightly slower than RemoveGameplayEffectsWithTags. */
-	UPROPERTY(EditDefaultsOnly, BlueprintReadOnly, Category = Tags, meta = (DisplayAfter = "RemovalTagRequirements"))
+	UE_DEPRECATED(5.3, "Remove Gameplay Effect Query is deprecated.  Use the URemoveOtherGameplayEffectComponent instead.")
+	UPROPERTY(BlueprintReadOnly, Category = Tags, meta = (DisplayAfter = "RemovalTagRequirements", DeprecatedProperty))
 	FGameplayEffectQuery RemoveGameplayEffectQuery;
 
 	/** Cached !RemoveGameplayEffectsQuery.IsEmpty(). Set on PostLoad. */
+	UE_DEPRECATED(5.3, "HasRemoveGameplayEffectsQuery is deprecated.  Use the URemoveOtherGameplayEffectComponent instead.")
 	bool HasRemoveGameplayEffectsQuery = false;
 
 	// ----------------------------------------------------------------------
 	//	Stacking
 	// ----------------------------------------------------------------------
-	
+
 	/** How this GameplayEffect stacks with other instances of this same GameplayEffect */
 	UPROPERTY(EditDefaultsOnly, BlueprintReadOnly, Category = Stacking)
 	EGameplayEffectStackingType	StackingType;
 
 	/** Stack limit for StackingType */
-	UPROPERTY(EditDefaultsOnly, BlueprintReadOnly, Category = Stacking)
+	UPROPERTY(EditDefaultsOnly, BlueprintReadOnly, Category = Stacking, meta = (EditConditionHides, EditCondition = "StackingType != EGameplayEffectStackingType::None"))
 	int32 StackLimitCount;
 
 	/** Policy for how the effect duration should be refreshed while stacking */
-	UPROPERTY(EditDefaultsOnly, BlueprintReadOnly, Category = Stacking)
+	UPROPERTY(EditDefaultsOnly, BlueprintReadOnly, Category = Stacking, meta = (EditConditionHides, EditCondition = "StackingType != EGameplayEffectStackingType::None"))
 	EGameplayEffectStackingDurationPolicy StackDurationRefreshPolicy;
 
 	/** Policy for how the effect period should be reset (or not) while stacking */
-	UPROPERTY(EditDefaultsOnly, BlueprintReadOnly, Category = Stacking)
+	UPROPERTY(EditDefaultsOnly, BlueprintReadOnly, Category = Stacking, meta = (EditConditionHides, EditCondition = "StackingType != EGameplayEffectStackingType::None"))
 	EGameplayEffectStackingPeriodPolicy StackPeriodResetPolicy;
 
 	/** Policy for how to handle duration expiring on this gameplay effect */
-	UPROPERTY(EditDefaultsOnly, BlueprintReadOnly, Category = Stacking)
+	UPROPERTY(EditDefaultsOnly, BlueprintReadOnly, Category = Stacking, meta = (EditConditionHides, EditCondition = "StackingType != EGameplayEffectStackingType::None"))
 	EGameplayEffectStackingExpirationPolicy StackExpirationPolicy;
 
 	// ----------------------------------------------------------------------
 	//	Granted abilities
 	// ----------------------------------------------------------------------
 	
+	/** Policy for what abilities this GE will grant. */
 	UPROPERTY(EditDefaultsOnly, BlueprintReadOnly, Category = "Granted Abilities")
 	TArray<FGameplayAbilitySpecDef>	GrantedAbilities;
+
+	// ----------------------------------------------------------------------
+	//	Cached Component Data - Do not modify these at runtime!
+	//	If you want to manipulate these, write your own GameplayEffectComponent
+	//	set the data during PostLoad.
+	// ----------------------------------------------------------------------
+	
+	/** Cached copy of all the tags this GE has. Data populated during PostLoad. */
+	FGameplayTagContainer CachedAssetTags;
+
+	/** Cached copy of all the tags this GE grants to its target. Data populated during PostLoad. */
+	FGameplayTagContainer CachedGrantedTags;
+
+	/** Cached copy of all the tags this GE applies to its target to block Gameplay Abilities. Data populated during PostLoad. */
+	FGameplayTagContainer CachedBlockedAbilityTags;
+
+protected:
+	/** These Gameplay Effect Components define how this Gameplay Effect behaves when applied */
+	UPROPERTY(EditDefaultsOnly, BlueprintReadOnly, Instanced, Category = "GameplayEffect", meta = (DisplayName = "Components", TitleProperty = EditorFriendlyName, ShowOnlyInnerProperties, DisplayPriority = 0))
+	TArray<TObjectPtr<UGameplayEffectComponent>> GEComponents;
+
+#if WITH_EDITORONLY_DATA
+	/** Allow us to show the Status of the class (valid configurations or invalid configurations) while configuring in the Editor */
+	UPROPERTY(VisibleAnywhere, Transient, Category = Status)
+	FText EditorStatusText;
+
+private:
+	/** The saved version of this package (the value is not inherited from its parents). @see SetVersion and GetVersion. */
+	UPROPERTY()
+	FGameplayEffectVersion DataVersion;
+#endif
 };
+
+template<typename GEComponentClass>
+const GEComponentClass* UGameplayEffect::FindComponent() const
+{
+	static_assert(TIsDerivedFrom<GEComponentClass, UGameplayEffectComponent>::IsDerived, "GEComponentClass must be derived from UGameplayEffectComponent");
+
+	for (const TObjectPtr<UGameplayEffectComponent>& GEComponent : GEComponents)
+	{
+		if (GEComponentClass* CastComponent = Cast<GEComponentClass>(GEComponent))
+		{
+			return CastComponent;
+		}
+	}
+
+	return nullptr;
+}
+
+template<typename GEComponentClass>
+GEComponentClass& UGameplayEffect::AddComponent()
+{
+	static_assert( TIsDerivedFrom<GEComponentClass, UGameplayEffectComponent>::IsDerived, "GEComponentClass must be derived from UGameplayEffectComponent");
+
+	TObjectPtr<GEComponentClass> Instance = NewObject<GEComponentClass>(this, NAME_None, GetMaskedFlags(RF_PropagateToSubObjects) | RF_Transactional);
+	GEComponents.Add(Instance);
+
+	check(Instance);
+	return *Instance;
+}
+
+template<typename GEComponentClass>
+GEComponentClass& UGameplayEffect::FindOrAddComponent()
+{
+	static_assert(TIsDerivedFrom<GEComponentClass, UGameplayEffectComponent>::IsDerived, "GEComponentClass must be derived from UGameplayEffectComponent");
+
+	for (const TObjectPtr<UGameplayEffectComponent>& GEComponent : GEComponents)
+	{
+		if (GEComponentClass* CastComponent = Cast<GEComponentClass>(GEComponent))
+		{
+			return *CastComponent;
+		}
+	}
+
+	return AddComponent<GEComponentClass>();
+}
\ No newline at end of file
diff --git a/Engine/Plugins/Runtime/GameplayAbilities/Source/GameplayAbilities/Public/GameplayEffectComponent.h b/Engine/Plugins/Runtime/GameplayAbilities/Source/GameplayAbilities/Public/GameplayEffectComponent.h
new file mode 100644
--- /dev/null
+++ b/Engine/Plugins/Runtime/GameplayAbilities/Source/GameplayAbilities/Public/GameplayEffectComponent.h
@@ -0,0 +1,88 @@
+// Copyright Epic Games, Inc. All Rights Reserved.
+
+#pragma once
+
+#include "CoreMinimal.h"
+#include "GameplayEffectComponent.generated.h"
+
+struct FActiveGameplayEffect;
+struct FActiveGameplayEffectsContainer;
+struct FGameplayEffectSpec;
+struct FPredictionKey;
+class UGameplayEffect;
+
+/**
+ * Gameplay Effect Component (aka GEComponent)
+ * 
+ * GEComponents are what define how a GameplayEffect behaves.  Introduced in UE 5.3, there are very few calls from UGameplayEffect to UGameplayEffectComponent by design.
+ * Instead of providing a larger API for all desired functionality, the implementer of a GEComponent must read the GE flow carefully and register desired callbacks
+ * to achieve the desired results.  This effectively limits the implementation of GEComponents to native code for the time being.
+ * 
+ * GEComponents live Within a GameplayEffect (which is typically a data-only blueprint asset).  Thus, like GEs, only one GEComponent exists for all applied instances.
+ * One of the unintuitive caveats of this is that GEComponent should not contain any runtime manipulated/instanced data (e.g. stored state per execution).
+ * One must take careful consideration about where to store any data (and thus when it can be evaluated).  The early implementations typically work around this by
+ * storing small amounts of runtime data on the desired callbacks (e.g. by binding extra parameters on the delegate).  This may explain why some functionality is still
+ * in UGameplayEffect rather than a UGameplayEffectComponent.  Future implementations may need extra data stored on the FGameplayEffectSpec (i.e. Gameplay Effect Spec Components).
+ * 
+ * @see GameplayEffect.h for further notes, especially on the terminology used (Added vs. Executed vs. Apply).
+ */
+UCLASS(Abstract, Const, DefaultToInstanced, EditInlineNew, CollapseCategories, Within=GameplayEffect)
+class GAMEPLAYABILITIES_API UGameplayEffectComponent : public UObject
+{
+	GENERATED_BODY()
+
+public:
+	/** Constructor */
+	UGameplayEffectComponent();
+
+	/** Returns the GameplayEffect that owns this Component (the Outer) */
+	UGameplayEffect* GetOwner() const;
+
+	/** 
+	 * Can the GameplayEffectSpec apply to the passed-in ASC?  All Components of the GE must return true, or a single one can return false to prohibit the application.
+	 * Note: Application and Inhibition are two separate things.  If a GE can apply, we then either Add it (if it has duration/prediction) or Execute it (if it's instant).
+	 */
+	virtual bool CanGameplayEffectApply(const FActiveGameplayEffectsContainer& ActiveGEContainer, const FGameplayEffectSpec& GESpec) const { return true; }
+
+	/**
+     * Called when a Gameplay Effect is Added to the ActiveGameplayEffectsContainer.  GE's are added to that container when they have duration (or are predicting locally)
+     * Return if the effect should remain active, or false to inhibit.  Note: Inhibit does not remove the effect (it remains added but dormant, waiting to unhibit).
+     */
+	virtual bool OnActiveGameplayEffectAdded(FActiveGameplayEffectsContainer& ActiveGEContainer, FActiveGameplayEffect& ActiveGE) const { return true; }
+
+	/** 
+	 * Called when a Gameplay Effect is executed.  GE's can only Execute on ROLE_Authority.  GE's only Execute when they're applying an instant effect (otherwise they're added to the ActiveGameplayEffectsContainer).
+	 * Note:  Periodic effects Execute every period (and are also added to ActiveGameplayEffectsContainer).
+	 */
+	virtual void OnGameplayEffectExecuted(FActiveGameplayEffectsContainer& ActiveGEContainer, FGameplayEffectSpec& GESpec, FPredictionKey& PredictionKey) const {}
+
+	/**
+	 * Let us know that the owning GameplayEffect has finished its PostLoad, thus apply an asset-related changes to the owning GameplayEffect (e.g. any of its fields)
+	 */
+	virtual void OnOwnerPostLoad() {}
+
+#if WITH_EDITOR
+	/**
+	 * Allow each Gameplay Effect Component to validate its own data.  Any warnings/errors will immediately show up in the Gameplay Effect when in Editor.
+	 * The default implementation ensures we only have a single type of any given class.  Override this function to change that functionality and use Super::Super::IsDataValid if needed.
+	 */
+	virtual EDataValidationResult IsDataValid(class FDataValidationContext& Context) override;
+#endif
+
+protected:
+#if WITH_EDITORONLY_DATA
+	/** Friendly name for displaying in the Editor's Gameplay Effect Component Index (@see UGameplayEffect::GEComponents). We set EditCondition False here so it doesn't show up otherwise. */
+	UPROPERTY(VisibleDefaultsOnly, Transient, Category=AlwaysHidden, Meta=(EditCondition=False, EditConditionHides))
+	FString EditorFriendlyName;
+#endif
+};
+
+// Find the same component in the parent of the passed-in GameplayEffect.  Useful for having a child component inherit properties from the parent (e.g. inherited tags).
+template<typename GEComponentClass, typename LateBindGameplayEffect = UGameplayEffect> // LateBindGameplayEffect to avoid having to #include GameplayEffect.h
+const GEComponentClass* FindParentComponent(const GEComponentClass& ChildComponent)
+{
+	const LateBindGameplayEffect* ChildGE = ChildComponent.GetOwner();
+	const LateBindGameplayEffect* ParentGE = ChildGE ? CastChecked<LateBindGameplayEffect>(ChildGE->GetClass()->GetArchetypeForCDO(), ECastCheckedType::NullAllowed) : nullptr;
+	return ParentGE ? ParentGE->template FindComponent<GEComponentClass>() : nullptr;
+}
+
diff --git a/Engine/Plugins/Runtime/GameplayAbilities/Source/GameplayAbilities/Public/GameplayEffectComponents/AdditionalEffectsGameplayEffectComponent.h b/Engine/Plugins/Runtime/GameplayAbilities/Source/GameplayAbilities/Public/GameplayEffectComponents/AdditionalEffectsGameplayEffectComponent.h
new file mode 100644
--- /dev/null
+++ b/Engine/Plugins/Runtime/GameplayAbilities/Source/GameplayAbilities/Public/GameplayEffectComponents/AdditionalEffectsGameplayEffectComponent.h
@@ -0,0 +1,69 @@
+// Copyright Epic Games, Inc. All Rights Reserved.
+
+#pragma once
+
+#include "GameplayEffect.h"
+#include "GameplayEffectComponent.h"
+#include "AdditionalEffectsGameplayEffectComponent.generated.h"
+
+/** Add additional Gameplay Effects that attempt to activate under certain conditions (or no conditions) */
+UCLASS(CollapseCategories)
+class GAMEPLAYABILITIES_API UAdditionalEffectsGameplayEffectComponent : public UGameplayEffectComponent
+{
+	GENERATED_BODY()
+
+public:
+	/**
+     * Called when a Gameplay Effect is Added to the ActiveGameplayEffectsContainer.  GE's are added to that container when they have duration (or are predicting locally)
+     * Return if the effect should remain active, or false to inhibit.  Note: Inhibit does not remove the effect (it remains added but dormant, waiting to unhibit).
+     */
+	virtual bool OnActiveGameplayEffectAdded(FActiveGameplayEffectsContainer& GEContainer, FActiveGameplayEffect& ActiveGE) const override;
+
+	/**
+	 * Called when a Gameplay Effect is executed.  GE's can only Execute on ROLE_Authority.  GE's only Execute when they're instant (otherwise they're added to the ActiveGameplayEffectsContainer).
+	 */
+	virtual void OnGameplayEffectExecuted(FActiveGameplayEffectsContainer& ActiveGEContainer, FGameplayEffectSpec& GESpec, FPredictionKey& PredictionKey) const override;
+
+#if WITH_EDITOR
+	/**
+	 * There are some fields that are incompatible with one another, so let's catch them during configuration phase. 
+	 */
+	virtual EDataValidationResult IsDataValid(class FDataValidationContext& Context) override;
+#endif // WITH_EDITOR
+
+
+protected:
+	/** 
+	 * Define our own Applied which gets called in both cases (Added or Executed).
+	 */
+	void OnGameplayEffectApplied(UAbilitySystemComponent& AppliedToASC, FGameplayEffectSpec& GESpec, FPredictionKey& PredictionKey) const;
+
+	/**
+	 * Whenever the ActiveGE gets removed, we want to apply the configured OnComplete GameplayEffects
+	 */
+	void OnActiveGameplayEffectRemoved(const FGameplayEffectRemovalInfo& RemovalInfo, FActiveGameplayEffectsContainer* ActiveGEContainer) const;
+
+public:
+	/**
+	 * This will copy all of the data (e.g. SetByCallerMagnitudes) from the GESpec that Applied this GameplayEffect to the new OnApplicationGameplayEffect Specs.
+	 * One would think this is normally desirable (and how OnComplete works by default), but we default to false here for backwards compatability.
+	 */
+	UPROPERTY(EditDefaultsOnly, Category = OnApplication)
+	bool bOnApplicationCopyDataFromOriginalSpec = false;
+
+	/** Other gameplay effects that will be applied to the target of this effect if the owning effect applies */
+	UPROPERTY(EditDefaultsOnly, Category = OnApplication)
+	TArray<FConditionalGameplayEffect> OnApplicationGameplayEffects;
+
+	/** Effects to apply when this effect completes, regardless of how it ends */
+	UPROPERTY(EditDefaultsOnly, Category = OnComplete)
+	TArray<TSubclassOf<UGameplayEffect>> OnCompleteAlways;
+
+	/** Effects to apply when this effect expires naturally via its duration */
+	UPROPERTY(EditDefaultsOnly, Category = OnComplete)
+	TArray<TSubclassOf<UGameplayEffect>> OnCompleteNormal;
+
+	/** Effects to apply when this effect is made to expire prematurely (e.g. via a forced removal, clear tags, etc.) */
+	UPROPERTY(EditDefaultsOnly, Category = OnComplete)
+	TArray<TSubclassOf<UGameplayEffect>> OnCompletePrematurely;
+};
diff --git a/Engine/Plugins/Runtime/GameplayAbilities/Source/GameplayAbilities/Public/GameplayEffectComponents/AssetTagsGameplayEffectComponent.h b/Engine/Plugins/Runtime/GameplayAbilities/Source/GameplayAbilities/Public/GameplayEffectComponents/AssetTagsGameplayEffectComponent.h
new file mode 100644
--- /dev/null
+++ b/Engine/Plugins/Runtime/GameplayAbilities/Source/GameplayAbilities/Public/GameplayEffectComponents/AssetTagsGameplayEffectComponent.h
@@ -0,0 +1,45 @@
+// Copyright Epic Games, Inc. All Rights Reserved.
+
+#pragma once
+
+#include "GameplayEffect.h"
+#include "GameplayEffectComponent.h"
+#include "AssetTagsGameplayEffectComponent.generated.h"
+
+/** These are tags that the Gameplay Effect Asset itself 'has' (owns). These do _not_ transfer to any Actors */
+UCLASS()
+class GAMEPLAYABILITIES_API UAssetTagsGameplayEffectComponent : public UGameplayEffectComponent
+{
+	GENERATED_BODY()
+
+public:
+	/** Setup a better EditorFriendlyName */
+	virtual void PostInitProperties() override;
+
+	/** Needed to properly apply FInheritedTagContainer properties */
+	virtual void OnOwnerPostLoad() override;
+
+	/** Gets the Asset Tags inherited tag structure (as configured) */
+	const FInheritedTagContainer& GetConfiguredAssetTagChanges() const { return InheritableAssetTags; }
+
+	/** Applies the Asset Tags to the GameplayEffect (and saves those changes) */
+	void SetAndApplyAssetTagChanges(const FInheritedTagContainer& TagContainerMods);
+
+#if WITH_EDITOR
+	/** Needed to properly update FInheritedTagContainer properties */
+	virtual void PostEditChangeProperty(FPropertyChangedEvent& PropertyChangedEvent) override;
+
+private:
+	/** Get a cached version of the FProperty Name for PostEditChangeProperty */
+	static const FName& GetInheritableAssetTagsName()
+	{
+		static FName NAME_InheritableAssetTags = GET_MEMBER_NAME_CHECKED(UAssetTagsGameplayEffectComponent, InheritableAssetTags);
+		return NAME_InheritableAssetTags;
+	}
+#endif // WITH_EDITOR
+
+private:
+	/** The GameplayEffect's Tags: tags the the GE *has* and DOES NOT give to the actor. */
+	UPROPERTY(EditDefaultsOnly, Category = Tags, meta = (DisplayName = "AssetTags", Categories = "OwnedTagsCategory"))
+	FInheritedTagContainer InheritableAssetTags;
+};
diff --git a/Engine/Plugins/Runtime/GameplayAbilities/Source/GameplayAbilities/Public/GameplayEffectComponents/BlockAbilityTagsGameplayEffectComponent.h b/Engine/Plugins/Runtime/GameplayAbilities/Source/GameplayAbilities/Public/GameplayEffectComponents/BlockAbilityTagsGameplayEffectComponent.h
new file mode 100644
--- /dev/null
+++ b/Engine/Plugins/Runtime/GameplayAbilities/Source/GameplayAbilities/Public/GameplayEffectComponents/BlockAbilityTagsGameplayEffectComponent.h
@@ -0,0 +1,45 @@
+// Copyright Epic Games, Inc. All Rights Reserved.
+
+#pragma once
+
+#include "GameplayEffect.h"
+#include "GameplayEffectComponent.h"
+#include "BlockAbilityTagsGameplayEffectComponent.generated.h"
+
+/** Handles blocking the activation of Gameplay Abilities based on Gameplay Tags for the Target actor of the owner Gameplay Effect */
+UCLASS()
+class GAMEPLAYABILITIES_API UBlockAbilityTagsGameplayEffectComponent : public UGameplayEffectComponent
+{
+	GENERATED_BODY()
+	
+public:
+	/** Constructor to set the EditorFriendlyName */
+	UBlockAbilityTagsGameplayEffectComponent();
+
+	/** Needed to properly apply FInheritedTagContainer properties */
+	virtual void OnOwnerPostLoad() override;
+
+	/** Gets the Blocked Ability Tags inherited tag structure (as configured) */
+	const FInheritedTagContainer& GetConfiguredBlockedAbilityTagChanges() const { return InheritableBlockedAbilityTagsContainer; }
+
+	/** Applies the Blocked Ability Tags to the GameplayEffect (and saves those changes) so that when the Gameplay Effect is applied, these Blocking Tags are then applied to the Target Actor */
+	void SetAndApplyBlockedAbilityTagChanges(const FInheritedTagContainer& TagContainerMods);
+
+#if WITH_EDITOR
+	/** Needed to properly update FInheritedTagContainer properties */
+	virtual void PostEditChangeProperty(FPropertyChangedEvent& PropertyChangedEvent) override;
+
+private:
+	/** Get a cached version of the FProperty Name for PostEditChangeProperty */
+	static const FName& GetInheritableBlockedAbilityTagsContainerPropertyName()
+	{
+		static FName NAME_InheritableBlockedAbilityTagsContainer = GET_MEMBER_NAME_CHECKED(UBlockAbilityTagsGameplayEffectComponent, InheritableBlockedAbilityTagsContainer);
+		return NAME_InheritableBlockedAbilityTagsContainer;
+	}
+#endif // WITH_EDITOR
+
+private:
+	/** These tags are applied to the target actor of the Gameplay Effect.  Blocked Ability Tags prevent Gameplay Abilities with these tags from executing. */
+	UPROPERTY(EditDefaultsOnly, Category = None, meta = (DisplayName = "Block Abilities w/ Tags"))
+	FInheritedTagContainer InheritableBlockedAbilityTagsContainer;
+};
diff --git a/Engine/Plugins/Runtime/GameplayAbilities/Source/GameplayAbilities/Public/GameplayEffectComponents/ChanceToApplyGameplayEffectComponent.h b/Engine/Plugins/Runtime/GameplayAbilities/Source/GameplayAbilities/Public/GameplayEffectComponents/ChanceToApplyGameplayEffectComponent.h
new file mode 100644
--- /dev/null
+++ b/Engine/Plugins/Runtime/GameplayAbilities/Source/GameplayAbilities/Public/GameplayEffectComponents/ChanceToApplyGameplayEffectComponent.h
@@ -0,0 +1,33 @@
+// Copyright Epic Games, Inc. All Rights Reserved.
+
+#pragma once
+
+#include "GameplayEffectComponent.h"
+#include "ScalableFloat.h"
+
+#include "ChanceToApplyGameplayEffectComponent.generated.h"
+
+/** Applies a probablity to the application conditions of the Gameplay Effect. */
+UCLASS()
+class GAMEPLAYABILITIES_API UChanceToApplyGameplayEffectComponent : public UGameplayEffectComponent
+{
+	GENERATED_BODY()
+	
+public:
+	/** Warn of errors in the Editor */
+	virtual void PostLoad() override;
+
+	/** Determine if we can apply this GameplayEffect or not */
+	virtual bool CanGameplayEffectApply(const FActiveGameplayEffectsContainer& ActiveGEContainer, const FGameplayEffectSpec& GESpec) const override;
+
+	/** Get the configured chance to apply to target */
+	const FScalableFloat& GetChanceToApplyToTarget() const { return ChanceToApplyToTarget; }
+
+	/** Set the chance to apply to target */
+	void SetChanceToApplyToTarget(const FScalableFloat& ChanceToApplyToTarget);
+
+private:
+	/** Probability that this gameplay effect will be applied to the target actor (0.0 for never, 1.0 for always) */
+	UPROPERTY(EditDefaultsOnly, Category = Application, meta = (GameplayAttribute = "True"))
+	FScalableFloat ChanceToApplyToTarget;
+};
diff --git a/Engine/Plugins/Runtime/GameplayAbilities/Source/GameplayAbilities/Public/GameplayEffectComponents/CustomCanApplyGameplayEffectComponent.h b/Engine/Plugins/Runtime/GameplayAbilities/Source/GameplayAbilities/Public/GameplayEffectComponents/CustomCanApplyGameplayEffectComponent.h
new file mode 100644
--- /dev/null
+++ b/Engine/Plugins/Runtime/GameplayAbilities/Source/GameplayAbilities/Public/GameplayEffectComponents/CustomCanApplyGameplayEffectComponent.h
@@ -0,0 +1,26 @@
+// Copyright Epic Games, Inc. All Rights Reserved.
+
+#pragma once
+
+#include "GameplayEffectComponent.h"
+#include "GameplayEffectCustomApplicationRequirement.h"
+
+#include "CustomCanApplyGameplayEffectComponent.generated.h"
+
+/** Handles configuration of a CustomApplicationRequirement function to see if this GameplayEffect should apply */
+UCLASS()
+class GAMEPLAYABILITIES_API UCustomCanApplyGameplayEffectComponent : public UGameplayEffectComponent
+{
+	GENERATED_BODY()
+	
+public:
+	UCustomCanApplyGameplayEffectComponent();
+
+	/** Determine if we can apply this GameplayEffect or not */
+	virtual bool CanGameplayEffectApply(const FActiveGameplayEffectsContainer& ActiveGEContainer, const FGameplayEffectSpec& GESpec) const override;
+
+public:
+	/** Custom application requirements */
+	UPROPERTY(EditDefaultsOnly, Category = Application, meta = (DisplayName = "Custom Application Requirement"))
+	TArray<TSubclassOf<UGameplayEffectCustomApplicationRequirement>> ApplicationRequirements;
+};
diff --git a/Engine/Plugins/Runtime/GameplayAbilities/Source/GameplayAbilities/Public/GameplayEffectComponents/ImmunityGameplayEffectComponent.h b/Engine/Plugins/Runtime/GameplayAbilities/Source/GameplayAbilities/Public/GameplayEffectComponents/ImmunityGameplayEffectComponent.h
new file mode 100644
--- /dev/null
+++ b/Engine/Plugins/Runtime/GameplayAbilities/Source/GameplayAbilities/Public/GameplayEffectComponents/ImmunityGameplayEffectComponent.h
@@ -0,0 +1,41 @@
+// Copyright Epic Games, Inc. All Rights Reserved.
+
+#pragma once
+
+#include "GameplayEffect.h"
+#include "GameplayEffectComponent.h"
+#include "ImmunityGameplayEffectComponent.generated.h"
+
+/** 
+ * Immunity is blocking the application of _other_ GameplayEffectSpecs.
+ * This registers a global handler on the ASC to block the application of other GameplayEffectSpecs.
+ */
+UCLASS()
+class GAMEPLAYABILITIES_API UImmunityGameplayEffectComponent : public UGameplayEffectComponent
+{
+	friend class UGameplayEffect; // Needed for upgrade path
+	GENERATED_BODY()
+
+public:
+	/** Constructor to set EditorFriendlyName */
+	UImmunityGameplayEffectComponent();
+
+	/** We need to register our callback to check for Immunity */
+	virtual bool OnActiveGameplayEffectAdded(FActiveGameplayEffectsContainer& ActiveGEContainer, FActiveGameplayEffect& ActiveGE) const override;
+
+#if WITH_EDITOR
+	/**
+	 * Warn about instant Gameplay Effects as they do not persist over time (and therefore cannot grant immunity)
+	 */
+	virtual EDataValidationResult IsDataValid(class FDataValidationContext& Context) override;
+#endif // WITH_EDITOR
+
+protected:
+	/** We register with the AbilitySystemComponent to try and block any GESpecs we think we should be immune to */
+	bool AllowGameplayEffectApplication(const FActiveGameplayEffectsContainer& ActiveGEContainer, const FGameplayEffectSpec& GESpecToConsider, FActiveGameplayEffectHandle ImmunityActiveGE) const;
+
+public:
+	/** Grants immunity to GameplayEffects that match this query. */
+	UPROPERTY(EditDefaultsOnly, BlueprintReadOnly, Category = None)
+	TArray<FGameplayEffectQuery> ImmunityQueries;
+};
diff --git a/Engine/Plugins/Runtime/GameplayAbilities/Source/GameplayAbilities/Public/GameplayEffectComponents/RemoveOtherGameplayEffectComponent.h b/Engine/Plugins/Runtime/GameplayAbilities/Source/GameplayAbilities/Public/GameplayEffectComponents/RemoveOtherGameplayEffectComponent.h
new file mode 100644
--- /dev/null
+++ b/Engine/Plugins/Runtime/GameplayAbilities/Source/GameplayAbilities/Public/GameplayEffectComponents/RemoveOtherGameplayEffectComponent.h
@@ -0,0 +1,41 @@
+// Copyright Epic Games, Inc. All Rights Reserved.
+
+#pragma once
+
+#include "GameplayEffect.h"
+#include "GameplayEffectComponent.h"
+#include "RemoveOtherGameplayEffectComponent.generated.h"
+
+struct FGameplayEffectRemovalInfo;
+
+/** Remove other Gameplay Effects based on certain conditions */
+UCLASS()
+class GAMEPLAYABILITIES_API URemoveOtherGameplayEffectComponent : public UGameplayEffectComponent
+{
+	GENERATED_BODY()
+
+public:
+	/** Constructor to set EditorFriendlyName */
+	URemoveOtherGameplayEffectComponent();
+
+	/** Once we've applied, we need to register for ongoing requirements */
+	virtual bool OnActiveGameplayEffectAdded(FActiveGameplayEffectsContainer& GEContainer, FActiveGameplayEffect& ActiveGE) const override;
+
+	/** If we're only executed, it's an indication something has gone wrong and we should log it */
+	virtual void OnGameplayEffectExecuted(FActiveGameplayEffectsContainer& ActiveGEContainer, FGameplayEffectSpec& GESpec, FPredictionKey& PredictionKey) const override;
+
+#if WITH_EDITOR
+	/**
+	 * Warn about periodic Gameplay Effects, that you probably instead want the OngoingTagRequirements in TagRequirementsGameplayEffectComponent
+	 */
+	virtual EDataValidationResult IsDataValid(class FDataValidationContext& Context) override;
+#endif // WITH_EDITOR
+
+private:
+	void OnGameplayEffectApplied(FActiveGameplayEffectsContainer& ActiveGEContainer) const;
+
+public:
+	/** On Application of the owning Gameplay Effect, any Active GameplayEffects that *match* these queries will be removed. */
+	UPROPERTY(BlueprintReadOnly, Category = None, meta = (DeprecatedProperty))
+	TArray<FGameplayEffectQuery> RemoveGameplayEffectQueries;
+};
diff --git a/Engine/Plugins/Runtime/GameplayAbilities/Source/GameplayAbilities/Public/GameplayEffectComponents/TargetTagRequirementsGameplayEffectComponent.h b/Engine/Plugins/Runtime/GameplayAbilities/Source/GameplayAbilities/Public/GameplayEffectComponents/TargetTagRequirementsGameplayEffectComponent.h
new file mode 100644
--- /dev/null
+++ b/Engine/Plugins/Runtime/GameplayAbilities/Source/GameplayAbilities/Public/GameplayEffectComponents/TargetTagRequirementsGameplayEffectComponent.h
@@ -0,0 +1,56 @@
+// Copyright Epic Games, Inc. All Rights Reserved.
+
+#pragma once
+
+#include "GameplayEffect.h"
+#include "GameplayEffectComponent.h"
+#include "TargetTagRequirementsGameplayEffectComponent.generated.h"
+
+struct FGameplayEffectRemovalInfo;
+
+/** Specifies tag requirements that the Target (owner of the Gameplay Effect) must have if this GE should apply or continue to execute */
+UCLASS()
+class GAMEPLAYABILITIES_API UTargetTagRequirementsGameplayEffectComponent : public UGameplayEffectComponent
+{
+	GENERATED_BODY()
+
+public:
+	/** Constructor to set EditorFriendlyName */
+	UTargetTagRequirementsGameplayEffectComponent();
+
+	/** Can we apply to the ActiveGEContainer? */
+	virtual bool CanGameplayEffectApply(const FActiveGameplayEffectsContainer& ActiveGEContainer, const FGameplayEffectSpec& GESpec) const override;
+
+	/** Once we've applied, we need to register for ongoing requirements */
+	virtual bool OnActiveGameplayEffectAdded(FActiveGameplayEffectsContainer& GEContainer, FActiveGameplayEffect& ActiveGE) const override;
+
+	/** If we're only executed, it's an indication something has gone wrong and we should log it */
+	virtual void OnGameplayEffectExecuted(FActiveGameplayEffectsContainer& ActiveGEContainer, FGameplayEffectSpec& GESpec, FPredictionKey& PredictionKey) const override;
+
+#if WITH_EDITOR
+	/**
+	 * Validate incompatable configurations
+	 */
+	virtual EDataValidationResult IsDataValid(class FDataValidationContext& Context) override;
+#endif // WITH_EDITOR
+
+private:
+	/** We need to be notified when we're removed to unregister some callbacks */
+	void OnGameplayEffectRemoved(const FGameplayEffectRemovalInfo& GERemovalInfo, UAbilitySystemComponent* ASC, TArray<TTuple<FGameplayTag, FDelegateHandle>> AllBoundEvents) const;
+
+	/** We need to register a callback for when the owner changes its tags.  When that happens, we need to figure out if our GE should continue to execute */
+	void OnTagChanged(const FGameplayTag GameplayTag, int32 NewCount, FActiveGameplayEffectHandle ActiveGEHandle) const;
+
+public:
+	/** Tag requirements for this GameplayEffect to be applied to a target. This is pass/fail at the time of application. If fail, this GE fails to apply. */
+	UPROPERTY(EditDefaultsOnly, Category = Tags)
+	FGameplayTagRequirements ApplicationTagRequirements;
+
+	/** Once Applied, these tags requirements are used to determined if the GameplayEffect is "on" or "off". A GameplayEffect can be off and do nothing, but still applied. */
+	UPROPERTY(EditDefaultsOnly, Category = Tags)
+	FGameplayTagRequirements OngoingTagRequirements;
+
+	/** Tag requirements that if met will remove this effect. Also prevents effect application. */
+	UPROPERTY(EditDefaultsOnly, Category = Tags)
+	FGameplayTagRequirements RemovalTagRequirements;
+};
diff --git a/Engine/Plugins/Runtime/GameplayAbilities/Source/GameplayAbilities/Public/GameplayEffectComponents/TargetTagsGameplayEffectComponent.h b/Engine/Plugins/Runtime/GameplayAbilities/Source/GameplayAbilities/Public/GameplayEffectComponents/TargetTagsGameplayEffectComponent.h
new file mode 100644
--- /dev/null
+++ b/Engine/Plugins/Runtime/GameplayAbilities/Source/GameplayAbilities/Public/GameplayEffectComponents/TargetTagsGameplayEffectComponent.h
@@ -0,0 +1,45 @@
+// Copyright Epic Games, Inc. All Rights Reserved.
+
+#pragma once
+
+#include "GameplayEffect.h"
+#include "GameplayEffectComponent.h"
+#include "TargetTagsGameplayEffectComponent.generated.h"
+
+/** Handles granting Tags to the Target (sometimes referred to as the Owner) of the Gameplay Effect */
+UCLASS()
+class GAMEPLAYABILITIES_API UTargetTagsGameplayEffectComponent : public UGameplayEffectComponent
+{
+	GENERATED_BODY()
+
+public:
+	/** Constructor that sets EditorFriendlyName */
+	UTargetTagsGameplayEffectComponent();
+
+	/** Needed to properly load FInheritedTagContainer properties */
+	virtual void OnOwnerPostLoad() override;
+
+	/** Gets the (Granted) Target Tags inherited tag structure (as configured) */
+	const FInheritedTagContainer& GetConfiguredTargetTagChanges() const { return InheritableGrantedTagsContainer; }
+
+	/** Applies the (Granted) Target Tags to the GameplayEffect (and saves those changes) so that it may grant these Tags the Target Actor when Applied */
+	void SetAndApplyTargetTagChanges(const FInheritedTagContainer& TagContainerMods);
+
+#if WITH_EDITOR
+	/** Needed to properly update FInheritedTagContainer properties */
+	virtual void PostEditChangeProperty(FPropertyChangedEvent& PropertyChangedEvent) override;
+
+private:
+	/** Get a cached version of the FProperty Name for PostEditChangeProperty */
+	static const FName& GetInheritableGrantedTagsContainerName()
+	{
+		static FName NAME_InheritableGrantedTagsContainer = GET_MEMBER_NAME_CHECKED(UTargetTagsGameplayEffectComponent, InheritableGrantedTagsContainer);
+		return NAME_InheritableGrantedTagsContainer;
+	}
+#endif // WITH_EDITOR
+
+private:
+	/** These tags are applied (granted) to the target actor of the Gameplay Effect.  The Target would then "own" these Tags. */
+	UPROPERTY(EditDefaultsOnly, Category = None, meta = (DisplayName = "Add Tags", Categories = "OwnedTagsCategory"))
+	FInheritedTagContainer InheritableGrantedTagsContainer;
+};
diff --git a/Engine/Plugins/Runtime/GameplayAbilities/Source/GameplayAbilities/Public/GameplayEffectCustomApplicationRequirement.h b/Engine/Plugins/Runtime/GameplayAbilities/Source/GameplayAbilities/Public/GameplayEffectCustomApplicationRequirement.h
--- a/Engine/Plugins/Runtime/GameplayAbilities/Source/GameplayAbilities/Public/GameplayEffectCustomApplicationRequirement.h
+++ b/Engine/Plugins/Runtime/GameplayAbilities/Source/GameplayAbilities/Public/GameplayEffectCustomApplicationRequirement.h
@@ -2,13 +2,11 @@
 
 #pragma once
 
-
 #include "CoreMinimal.h"
-#include "UObject/ObjectMacros.h"
-#include "UObject/Object.h"
-#include "GameplayEffect.h"
 #include "GameplayEffectCustomApplicationRequirement.generated.h"
 
+struct FGameplayEffectSpec;
+class UGameplayEffect;
 class UAbilitySystemComponent;
 
 /** Class used to perform custom gameplay effect modifier calculations, either via blueprint or native code */ 
diff --git a/Engine/Plugins/Runtime/GameplayAbilities/Source/GameplayAbilities/Public/GameplayEffectTypes.h b/Engine/Plugins/Runtime/GameplayAbilities/Source/GameplayAbilities/Public/GameplayEffectTypes.h
--- a/Engine/Plugins/Runtime/GameplayAbilities/Source/GameplayAbilities/Public/GameplayEffectTypes.h
+++ b/Engine/Plugins/Runtime/GameplayAbilities/Source/GameplayAbilities/Public/GameplayEffectTypes.h
@@ -245,7 +245,7 @@ struct GAMEPLAYABILITIES_API FGameplayEffectContext
 	{
 	}
 
-	/** Returns the list of gameplay tags applicable to this effect, defaults to the owner's tags */
+	/** Returns the list of gameplay tags applicable to this effect, defaults to the owner's tags. SpecTagContainer remains untouched by default. */
 	virtual void GetOwnedGameplayTags(OUT FGameplayTagContainer& ActorTagContainer, OUT FGameplayTagContainer& SpecTagContainer) const;
 
 	/** Sets the instigator and effect causer. Instigator is who owns the ability that spawned this, EffectCauser is the actor that is the physical source of the effect, such as a weapon. They can be the same. */
@@ -780,18 +780,21 @@ USTRUCT(BlueprintType)
 struct FGameplayEffectRemovalInfo
 {
 	GENERATED_USTRUCT_BODY()
-
+	
 	/** True when the gameplay effect's duration has not expired, meaning the gameplay effect is being forcefully removed.  */
-	UPROPERTY(EditAnywhere, BlueprintReadWrite, Category="Removal")
+	UPROPERTY(VisibleAnywhere, BlueprintReadOnly, Category="Removal")
 	bool bPrematureRemoval = false;
 
 	/** Number of Stacks this gameplay effect had before it was removed. */
-	UPROPERTY(EditAnywhere, BlueprintReadWrite, Category="Removal")
+	UPROPERTY(VisibleAnywhere, BlueprintReadOnly, Category="Removal")
 	int32 StackCount = 0;
 
 	/** Actor this gameplay effect was targeting. */
-	UPROPERTY(EditAnywhere, BlueprintReadWrite, Category="Removal")
+	UPROPERTY(VisibleAnywhere, BlueprintReadOnly, Category="Removal")
 	FGameplayEffectContextHandle EffectContext;
+
+	/** The Effect being Removed */
+	const FActiveGameplayEffect* ActiveEffect = nullptr;
 };
 
 
@@ -1353,11 +1356,11 @@ struct GAMEPLAYABILITIES_API FGameplayTagRequirements
 	GENERATED_USTRUCT_BODY()
 
 	/** All of these tags must be present */
-	UPROPERTY(EditAnywhere, BlueprintReadWrite, Category = GameplayModifier)
+	UPROPERTY(EditAnywhere, BlueprintReadWrite, Category = GameplayModifier, meta=(DisplayName="Must Have Tags"))
 	FGameplayTagContainer RequireTags;
 
 	/** None of these tags may be present */
-	UPROPERTY(EditAnywhere, BlueprintReadWrite, Category = GameplayModifier)
+	UPROPERTY(EditAnywhere, BlueprintReadWrite, Category = GameplayModifier, meta=(DisplayName="Must Not Have Tags"))
 	FGameplayTagContainer IgnoreTags;
 
 	/** True if all required tags and no ignore tags found */
@@ -1368,6 +1371,9 @@ struct GAMEPLAYABILITIES_API FGameplayTagRequirements
 
 	/** Return debug string */
 	FString ToString() const;
+
+	bool operator==(const FGameplayTagRequirements& Other) const;
+	bool operator!=(const FGameplayTagRequirements& Other) const;
 };
 
 
@@ -1382,7 +1388,6 @@ struct GAMEPLAYABILITIES_API FTagContainerAggregator
 	FTagContainerAggregator(FTagContainerAggregator&& Other)
 		: CapturedActorTags(MoveTemp(Other.CapturedActorTags))
 		, CapturedSpecTags(MoveTemp(Other.CapturedSpecTags))
-		, ScopedTags(MoveTemp(Other.ScopedTags))
 		, CachedAggregator(MoveTemp(Other.CachedAggregator))
 		, CacheIsValid(Other.CacheIsValid)
 	{
@@ -1391,7 +1396,6 @@ struct GAMEPLAYABILITIES_API FTagContainerAggregator
 	FTagContainerAggregator(const FTagContainerAggregator& Other)
 		: CapturedActorTags(Other.CapturedActorTags)
 		, CapturedSpecTags(Other.CapturedSpecTags)
-		, ScopedTags(Other.ScopedTags)
 		, CachedAggregator(Other.CachedAggregator)
 		, CacheIsValid(Other.CacheIsValid)
 	{
@@ -1401,7 +1405,6 @@ struct GAMEPLAYABILITIES_API FTagContainerAggregator
 	{
 		CapturedActorTags = MoveTemp(Other.CapturedActorTags);
 		CapturedSpecTags = MoveTemp(Other.CapturedSpecTags);
-		ScopedTags = MoveTemp(Other.ScopedTags);
 		CachedAggregator = MoveTemp(Other.CachedAggregator);
 		CacheIsValid = Other.CacheIsValid;
 		return *this;
@@ -1411,7 +1414,6 @@ struct GAMEPLAYABILITIES_API FTagContainerAggregator
 	{
 		CapturedActorTags = Other.CapturedActorTags;
 		CapturedSpecTags = Other.CapturedSpecTags;
-		ScopedTags = Other.ScopedTags;
 		CachedAggregator = Other.CachedAggregator;
 		CacheIsValid = Other.CacheIsValid;
 		return *this;
@@ -1436,6 +1438,7 @@ private:
 	UPROPERTY()
 	FGameplayTagContainer CapturedSpecTags;
 
+	UE_DEPRECATED(5.3, "This variable is unused and will removed")
 	UPROPERTY()
 	FGameplayTagContainer ScopedTags;
 
diff --git a/Engine/Plugins/Runtime/GameplayAbilities/Source/GameplayAbilities/Public/GameplayEffectUIData.h b/Engine/Plugins/Runtime/GameplayAbilities/Source/GameplayAbilities/Public/GameplayEffectUIData.h
--- a/Engine/Plugins/Runtime/GameplayAbilities/Source/GameplayAbilities/Public/GameplayEffectUIData.h
+++ b/Engine/Plugins/Runtime/GameplayAbilities/Source/GameplayAbilities/Public/GameplayEffectUIData.h
@@ -2,17 +2,16 @@
 
 #pragma once
 
-#include "CoreMinimal.h"
-#include "UObject/ObjectMacros.h"
-#include "UObject/Object.h"
+#include "GameplayEffectComponent.h"
 #include "GameplayEffectUIData.generated.h"
 
 /**
  * UGameplayEffectUIData
  * Base class to provide game-specific data about how to describe a Gameplay Effect in the UI. Subclass with data to use in your game.
+ * In Unreal Engine 5.3, this now derives from UGameplayEffectComponent so you can use it directly as a GameplayEffectComponent.
  */
 UCLASS(Blueprintable, Abstract, EditInlineNew, CollapseCategories)
-class GAMEPLAYABILITIES_API UGameplayEffectUIData : public UObject
+class GAMEPLAYABILITIES_API UGameplayEffectUIData : public UGameplayEffectComponent
 {
-	GENERATED_UCLASS_BODY()
+	GENERATED_BODY()
 };
diff --git a/Engine/Plugins/Runtime/GameplayAbilities/Source/GameplayAbilities/Public/GameplayEffectUIData_TextOnly.h b/Engine/Plugins/Runtime/GameplayAbilities/Source/GameplayAbilities/Public/GameplayEffectUIData_TextOnly.h
--- a/Engine/Plugins/Runtime/GameplayAbilities/Source/GameplayAbilities/Public/GameplayEffectUIData_TextOnly.h
+++ b/Engine/Plugins/Runtime/GameplayAbilities/Source/GameplayAbilities/Public/GameplayEffectUIData_TextOnly.h
@@ -2,8 +2,6 @@
 
 #pragma once
 
-#include "CoreMinimal.h"
-#include "UObject/ObjectMacros.h"
 #include "GameplayEffectUIData.h"
 #include "GameplayEffectUIData_TextOnly.generated.h"
 
@@ -14,8 +12,9 @@
 UCLASS()
 class GAMEPLAYABILITIES_API UGameplayEffectUIData_TextOnly : public UGameplayEffectUIData
 {
-	GENERATED_UCLASS_BODY()
-
+	GENERATED_BODY()
+	
+public:
 	UPROPERTY(EditDefaultsOnly, BlueprintReadOnly, Category = Data, meta = (MultiLine = "true"))
 	FText Description;
 };
diff --git a/Engine/Plugins/Runtime/GameplayAbilities/Source/GameplayAbilitiesEditor/Private/GameplayEffectDetails.cpp b/Engine/Plugins/Runtime/GameplayAbilities/Source/GameplayAbilitiesEditor/Private/GameplayEffectDetails.cpp
--- a/Engine/Plugins/Runtime/GameplayAbilities/Source/GameplayAbilitiesEditor/Private/GameplayEffectDetails.cpp
+++ b/Engine/Plugins/Runtime/GameplayAbilities/Source/GameplayAbilitiesEditor/Private/GameplayEffectDetails.cpp
@@ -29,10 +29,18 @@ void FGameplayEffectDetails::CustomizeDetails(IDetailLayoutBuilder& DetailLayout
 		return;
 	}
 
-	TSharedPtr<IPropertyHandle> DurationPolicyProperty = DetailLayout.GetProperty(GET_MEMBER_NAME_CHECKED(UGameplayEffect, DurationPolicy), UGameplayEffect::StaticClass());
 	FSimpleDelegate UpdateDurationPolicyDelegate = FSimpleDelegate::CreateSP(this, &FGameplayEffectDetails::OnDurationPolicyChange);
+
+	TSharedPtr<IPropertyHandle> DurationPolicyProperty = DetailLayout.GetProperty(GET_MEMBER_NAME_CHECKED(UGameplayEffect, DurationPolicy), UGameplayEffect::StaticClass());
 	DurationPolicyProperty->SetOnPropertyValueChanged(UpdateDurationPolicyDelegate);
-	
+
+	TSharedPtr<IPropertyHandle> PeriodProperty = DetailLayout.GetProperty(GET_MEMBER_NAME_CHECKED(UGameplayEffect, Period), UGameplayEffect::StaticClass());
+	TSharedPtr<IPropertyHandle> PeriodValue = PeriodProperty ? PeriodProperty->GetChildHandle(GET_MEMBER_NAME_CHECKED(FScalableFloat, Value)) : nullptr;
+	if (PeriodValue)
+	{
+		PeriodValue->SetOnPropertyValueChanged(UpdateDurationPolicyDelegate);
+	}
+
 	// Hide properties where necessary
 	UGameplayEffect* Obj = Cast<UGameplayEffect>(Objects[0].Get());
 	if (Obj)
@@ -43,11 +51,11 @@ void FGameplayEffectDetails::CustomizeDetails(IDetailLayoutBuilder& DetailLayout
 			DetailLayout.HideProperty(DurationMagnitudeProperty);
 		}
 
-		if (Obj->DurationPolicy == EGameplayEffectDurationType::Instant)
+		if (Obj->DurationPolicy == EGameplayEffectDurationType::Instant || Obj->Period.Value <= 0.0f)
 		{
-			TSharedPtr<IPropertyHandle> PeriodProperty = DetailLayout.GetProperty(GET_MEMBER_NAME_CHECKED(UGameplayEffect, Period), UGameplayEffect::StaticClass());
+			TSharedPtr<IPropertyHandle> PeriodInhibitionPolicyProperty = DetailLayout.GetProperty(GET_MEMBER_NAME_CHECKED(UGameplayEffect, PeriodicInhibitionPolicy), UGameplayEffect::StaticClass());
 			TSharedPtr<IPropertyHandle> ExecutePeriodicEffectOnApplicationProperty = DetailLayout.GetProperty(GET_MEMBER_NAME_CHECKED(UGameplayEffect, bExecutePeriodicEffectOnApplication), UGameplayEffect::StaticClass());
-			DetailLayout.HideProperty(PeriodProperty);
+			DetailLayout.HideProperty(PeriodInhibitionPolicyProperty);
 			DetailLayout.HideProperty(ExecutePeriodicEffectOnApplicationProperty);
 		}
 
diff --git a/Engine/Plugins/Runtime/GameplayAbilities/Source/GameplayAbilitiesEditor/Private/GameplayEffectDetails.h b/Engine/Plugins/Runtime/GameplayAbilities/Source/GameplayAbilitiesEditor/Private/GameplayEffectDetails.h
--- a/Engine/Plugins/Runtime/GameplayAbilities/Source/GameplayAbilitiesEditor/Private/GameplayEffectDetails.h
+++ b/Engine/Plugins/Runtime/GameplayAbilities/Source/GameplayAbilitiesEditor/Private/GameplayEffectDetails.h
@@ -20,12 +20,6 @@ private:
 	/** IDetailCustomization interface */
 	virtual void CustomizeDetails(IDetailLayoutBuilder& DetailLayout) override;
 
-	
-
-	TArray<TSharedPtr<FString>> PropertyOptions;
-
-	TSharedPtr<IPropertyHandle> MyProperty;
-
 	IDetailLayoutBuilder* MyDetailLayout;
 
 	void OnDurationPolicyChange();
diff --git a/Engine/Source/Runtime/Engine/Private/BlueprintGeneratedClass.cpp b/Engine/Source/Runtime/Engine/Private/BlueprintGeneratedClass.cpp
--- a/Engine/Source/Runtime/Engine/Private/BlueprintGeneratedClass.cpp
+++ b/Engine/Source/Runtime/Engine/Private/BlueprintGeneratedClass.cpp
@@ -1574,7 +1574,7 @@ void UBlueprintGeneratedClass::GetPreloadDependencies(TArray<UObject*>& OutDeps)
 	{
 		ForEachObjectWithOuter(CDO, [&OutDeps](UObject* SubObj)
 		{
-			if (SubObj->HasAllFlags(RF_DefaultSubObject))
+			if (SubObj->HasAnyFlags(RF_DefaultSubObject | RF_ArchetypeObject))
 			{
 				OutDeps.Add(SubObj->GetClass());
 				OutDeps.Add(SubObj->GetArchetype());
diff --git a/Engine/Source/Runtime/GameplayTags/Classes/GameplayTagContainer.h b/Engine/Source/Runtime/GameplayTags/Classes/GameplayTagContainer.h
--- a/Engine/Source/Runtime/GameplayTags/Classes/GameplayTagContainer.h
+++ b/Engine/Source/Runtime/GameplayTags/Classes/GameplayTagContainer.h
@@ -806,6 +806,9 @@ public:
 	FGameplayTagQuery& operator=(FGameplayTagQuery const& Other);
 	FGameplayTagQuery& operator=(FGameplayTagQuery&& Other);
 
+	bool operator==(const FGameplayTagQuery& Other) const;
+	bool operator!=(const FGameplayTagQuery& Other) const;
+
 private:
 	// Note: Properties need to be editable to allow FComponentPropertyWriter to serialize them, but are hidden in the editor by the customizations mentioned above.
 
diff --git a/Engine/Source/Runtime/GameplayTags/Private/GameplayTagContainer.cpp b/Engine/Source/Runtime/GameplayTags/Private/GameplayTagContainer.cpp
--- a/Engine/Source/Runtime/GameplayTags/Private/GameplayTagContainer.cpp
+++ b/Engine/Source/Runtime/GameplayTags/Private/GameplayTagContainer.cpp
@@ -1376,6 +1376,20 @@ FGameplayTagQuery& FGameplayTagQuery::operator=(FGameplayTagQuery&& Other)
 	return *this;
 }
 
+bool FGameplayTagQuery::operator==(const FGameplayTagQuery& Other) const
+{
+	return TokenStreamVersion == Other.TokenStreamVersion &&
+		TagDictionary == Other.TagDictionary &&
+		QueryTokenStream == Other.QueryTokenStream &&
+		UserDescription == Other.UserDescription &&
+		AutoDescription == Other.AutoDescription;
+}
+
+bool FGameplayTagQuery::operator!=(const FGameplayTagQuery& Other) const
+{
+	return !(*this == Other);
+}
+
 bool FGameplayTagQuery::Matches(FGameplayTagContainer const& Tags) const
 {
 	FQueryEvaluator QE(*this);
