Fix for crash with dynamicly created GE's and GEComponents.

diff --git a/Engine/Plugins/Runtime/GameplayAbilities/Source/GameplayAbilities/Public/GameplayEffectComponent.h b/Engine/Plugins/Runtime/GameplayAbilities/Source/GameplayAbilities/Public/GameplayEffectComponent.h
--- a/Engine/Plugins/Runtime/GameplayAbilities/Source/GameplayAbilities/Public/GameplayEffectComponent.h
+++ b/Engine/Plugins/Runtime/GameplayAbilities/Source/GameplayAbilities/Public/GameplayEffectComponent.h
@@ -82,7 +82,7 @@ template<typename GEComponentClass, typename LateBindGameplayEffect = UGameplayE
 const GEComponentClass* FindParentComponent(const GEComponentClass& ChildComponent)
 {
 	const LateBindGameplayEffect* ChildGE = ChildComponent.GetOwner();
-	const LateBindGameplayEffect* ParentGE = ChildGE ? CastChecked<LateBindGameplayEffect>(ChildGE->GetClass()->GetArchetypeForCDO(), ECastCheckedType::NullAllowed) : nullptr;
+	const LateBindGameplayEffect* ParentGE = ChildGE ? Cast<LateBindGameplayEffect>(ChildGE->GetClass()->GetArchetypeForCDO()) : nullptr;
 	return ParentGE ? ParentGE->template FindComponent<GEComponentClass>() : nullptr;
 }
 
