Fixed a bug in RemoveOtherGameplayEffectComponent where it would consider itself for removal.

diff --git a/Engine/Plugins/Runtime/GameplayAbilities/Source/GameplayAbilities/Private/GameplayEffect.cpp b/Engine/Plugins/Runtime/GameplayAbilities/Source/GameplayAbilities/Private/GameplayEffect.cpp
--- a/Engine/Plugins/Runtime/GameplayAbilities/Source/GameplayAbilities/Private/GameplayEffect.cpp
+++ b/Engine/Plugins/Runtime/GameplayAbilities/Source/GameplayAbilities/Private/GameplayEffect.cpp
@@ -511,6 +511,7 @@ void UGameplayEffect::ConvertRemoveOtherComponent()
 
 		if (bTagsChanged)
 		{
+			// Note: This is keeping the old functionality which incorrectly conflates Asset Tags with Owned Tags.
 			RemoveOtherComponent.RemoveGameplayEffectQueries[0] = FGameplayEffectQuery::MakeQuery_MatchAnyOwningTags(RemoveGameplayEffectsWithTags.CombinedTags);
 		}
 
diff --git a/Engine/Plugins/Runtime/GameplayAbilities/Source/GameplayAbilities/Private/GameplayEffectComponents/RemoveOtherGameplayEffectComponent.cpp b/Engine/Plugins/Runtime/GameplayAbilities/Source/GameplayAbilities/Private/GameplayEffectComponents/RemoveOtherGameplayEffectComponent.cpp
--- a/Engine/Plugins/Runtime/GameplayAbilities/Source/GameplayAbilities/Private/GameplayEffectComponents/RemoveOtherGameplayEffectComponent.cpp
+++ b/Engine/Plugins/Runtime/GameplayAbilities/Source/GameplayAbilities/Private/GameplayEffectComponents/RemoveOtherGameplayEffectComponent.cpp
@@ -17,7 +17,7 @@ bool URemoveOtherGameplayEffectComponent::OnActiveGameplayEffectAdded(FActiveGam
 {
 	if (ActiveGEContainer.OwnerIsNetAuthority)
 	{
-		OnGameplayEffectApplied(ActiveGEContainer);
+		OnGameplayEffectApplied(ActiveGEContainer, ActiveGE.Handle);
 	}
 	return true;
 }
@@ -27,18 +27,32 @@ void URemoveOtherGameplayEffectComponent::OnGameplayEffectExecuted(FActiveGamepl
 	const bool bInstantEffect = (GetOwner()->DurationPolicy == EGameplayEffectDurationType::Instant);
 	if (bInstantEffect && ActiveGEContainer.OwnerIsNetAuthority)
 	{
-		OnGameplayEffectApplied(ActiveGEContainer);
+		FActiveGameplayEffectHandle InvalidHandle;
+		OnGameplayEffectApplied(ActiveGEContainer, InvalidHandle);
 	}
 }
 
-void URemoveOtherGameplayEffectComponent::OnGameplayEffectApplied(FActiveGameplayEffectsContainer& ActiveGEContainer) const
+void URemoveOtherGameplayEffectComponent::OnGameplayEffectApplied(FActiveGameplayEffectsContainer& ActiveGEContainer, const FActiveGameplayEffectHandle& ActiveGEHandle) const
 {
 	constexpr int32 RemoveAllStacks = -1;
 	for (const FGameplayEffectQuery& RemoveQuery : RemoveGameplayEffectQueries)
 	{
 		if (!RemoveQuery.IsEmpty())
 		{
-			ActiveGEContainer.RemoveActiveEffects(RemoveQuery, RemoveAllStacks);
+			// If we have an ActiveGEHandle, make sure we never remove ourselves.
+			// If we don't, there's no need to make a copy.
+			if (!ActiveGEHandle.IsValid())
+			{
+				// Faster path: No copy needed
+				ActiveGEContainer.RemoveActiveEffects(RemoveQuery, RemoveAllStacks);
+			}
+			else
+			{
+				FGameplayEffectQuery MutableRemoveQuery = RemoveQuery;
+				MutableRemoveQuery.IgnoreHandles.Add(ActiveGEHandle);
+
+				ActiveGEContainer.RemoveActiveEffects(MutableRemoveQuery, RemoveAllStacks);
+			}
 		}
 	}
 }
diff --git a/Engine/Plugins/Runtime/GameplayAbilities/Source/GameplayAbilities/Public/GameplayEffectComponents/RemoveOtherGameplayEffectComponent.h b/Engine/Plugins/Runtime/GameplayAbilities/Source/GameplayAbilities/Public/GameplayEffectComponents/RemoveOtherGameplayEffectComponent.h
--- a/Engine/Plugins/Runtime/GameplayAbilities/Source/GameplayAbilities/Public/GameplayEffectComponents/RemoveOtherGameplayEffectComponent.h
+++ b/Engine/Plugins/Runtime/GameplayAbilities/Source/GameplayAbilities/Public/GameplayEffectComponents/RemoveOtherGameplayEffectComponent.h
@@ -32,10 +32,10 @@ public:
 #endif // WITH_EDITOR
 
 private:
-	void OnGameplayEffectApplied(FActiveGameplayEffectsContainer& ActiveGEContainer) const;
+	void OnGameplayEffectApplied(FActiveGameplayEffectsContainer& ActiveGEContainer, const FActiveGameplayEffectHandle& ActiveGEHandle) const;
 
 public:
 	/** On Application of the owning Gameplay Effect, any Active GameplayEffects that *match* these queries will be removed. */
-	UPROPERTY(BlueprintReadOnly, Category = None, meta = (DeprecatedProperty))
+	UPROPERTY(EditDefaultsOnly, Category = None)
 	TArray<FGameplayEffectQuery> RemoveGameplayEffectQueries;
 };
