Fixed mismatched copy constructor/assignment operator in FGameplayTagBlueprintPropertyMapping.

diff --git a/Engine/Plugins/Runtime/GameplayAbilities/Source/GameplayAbilities/Public/GameplayEffectTypes.h b/Engine/Plugins/Runtime/GameplayAbilities/Source/GameplayAbilities/Public/GameplayEffectTypes.h
--- a/Engine/Plugins/Runtime/GameplayAbilities/Source/GameplayAbilities/Public/GameplayEffectTypes.h
+++ b/Engine/Plugins/Runtime/GameplayAbilities/Source/GameplayAbilities/Public/GameplayEffectTypes.h
@@ -1280,6 +1280,9 @@ public:
 		PropertyGuid = Other.PropertyGuid;
 	}
 
+	// Pretty weird that the assignment operator copies the handle when the copy constructor doesn't - bug?
+	FGameplayTagBlueprintPropertyMapping& operator=(const FGameplayTagBlueprintPropertyMapping& Other) = default;
+
 	/** Gameplay tag being counted. */
 	UPROPERTY(EditAnywhere, Category = GameplayTagBlueprintProperty)
 	FGameplayTag TagToMap;
