Fix for FAttributeSetInitterDiscreteLevels not working in Cooked Builds due to Curve Simplification.

PR from Gabriel Wreczycki

diff --git a/Engine/Plugins/Runtime/GameplayAbilities/Source/GameplayAbilities/Private/AttributeSet.cpp b/Engine/Plugins/Runtime/GameplayAbilities/Source/GameplayAbilities/Private/AttributeSet.cpp
--- a/Engine/Plugins/Runtime/GameplayAbilities/Source/GameplayAbilities/Private/AttributeSet.cpp
+++ b/Engine/Plugins/Runtime/GameplayAbilities/Source/GameplayAbilities/Private/AttributeSet.cpp
@@ -545,7 +545,7 @@ TSubclassOf<UAttributeSet> FindBestAttributeClass(TArray<TSubclassOf<UAttributeS
  *	Each curve in the table represents a *single attribute's values for all levels*.
  *	At runtime, we want *all attribute values at given level*.
  *
- *	This code assumes that your curve data starts with a key of 1 and increases by 1 with each key.
+ *	This code assumes that your curve data starts with a key of 1.
  */
 void FAttributeSetInitterDiscreteLevels::PreloadAttributeSetData(const TArray<UCurveTable*>& CurveData)
 {
@@ -612,46 +612,25 @@ void FAttributeSetInitterDiscreteLevels::PreloadAttributeSetData(const TArray<UC
 			FName ClassFName = FName(*ClassName);
 			FAttributeSetDefaultsCollection& DefaultCollection = Defaults.FindOrAdd(ClassFName);
 
-			// Check our curve to make sure the keys match the expected format
-			int32 ExpectedLevel = 1;
-			bool bShouldSkip = false;
-			for (auto KeyIter = Curve->GetKeyHandleIterator(); KeyIter; ++KeyIter)
-			{
-				const FKeyHandle& KeyHandle = *KeyIter;
-				if (KeyHandle == FKeyHandle::Invalid())
-				{
-					ABILITY_LOG(Verbose, TEXT("FAttributeSetInitterDiscreteLevels::PreloadAttributeSetData Data contains an invalid key handle (row: %s)"), *RowName);
-					bShouldSkip = true;
-					break;
-				}
-
-				int32 Level = Curve->GetKeyTimeValuePair(KeyHandle).Key;
-				if (ExpectedLevel != Level)
-				{
-					ABILITY_LOG(Verbose, TEXT("FAttributeSetInitterDiscreteLevels::PreloadAttributeSetData Keys are expected to start at 1 and increase by 1 for every key (row: %s)"), *RowName);
-					bShouldSkip = true;
-					break;
-				}
+			float FirstLevelFloat = 0.f;
+			float LastLevelFloat = 0.f;
+			Curve->GetTimeRange(FirstLevelFloat, LastLevelFloat);
 
-				++ExpectedLevel;
-			}
+			int32 FirstLevel = FMath::RoundToInt32(FirstLevelFloat);
+			int32 LastLevel = FMath::RoundToInt32(LastLevelFloat);
 
-			if (bShouldSkip)
+			// Only log these as warnings, as they're not deal breakers.
+			if (FirstLevel != 1)
 			{
+				ABILITY_LOG(Warning, TEXT("FAttributeSetInitterDiscreteLevels::PreloadAttributeSetData First level should be 1"));
 				continue;
 			}
 
-			int32 LastLevel = Curve->GetKeyTime(Curve->GetLastKeyHandle());
 			DefaultCollection.LevelData.SetNum(FMath::Max(LastLevel, DefaultCollection.LevelData.Num()));
 
-			//At this point we know the Name of this "class"/"group", the AttributeSet, and the Property Name. Now loop through the values on the curve to get the attribute default value at each level.
-			for (auto KeyIter = Curve->GetKeyHandleIterator(); KeyIter; ++KeyIter)
+			for (int32 Level = 1; Level <= LastLevel; ++Level)
 			{
-				const FKeyHandle& KeyHandle = *KeyIter;
-
-				TPair<float, float> LevelValuePair = Curve->GetKeyTimeValuePair(KeyHandle);
-				int32 Level = LevelValuePair.Key;
-				float Value = LevelValuePair.Value;
+				float Value = Curve->Eval(float(Level));
 
 				FAttributeSetDefaults& SetDefaults = DefaultCollection.LevelData[Level-1];
 
