Stop tickable attribute sets from ticking if ShouldTick returns false.

diff --git a/Engine/Plugins/Runtime/GameplayAbilities/Source/GameplayAbilities/Private/AbilitySystemComponent_Abilities.cpp b/Engine/Plugins/Runtime/GameplayAbilities/Source/GameplayAbilities/Private/AbilitySystemComponent_Abilities.cpp
--- a/Engine/Plugins/Runtime/GameplayAbilities/Source/GameplayAbilities/Private/AbilitySystemComponent_Abilities.cpp
+++ b/Engine/Plugins/Runtime/GameplayAbilities/Source/GameplayAbilities/Private/AbilitySystemComponent_Abilities.cpp
@@ -124,7 +124,7 @@ void UAbilitySystemComponent::TickComponent(float DeltaTime, enum ELevelTick Tic
 	for (UAttributeSet* AttributeSet : GetSpawnedAttributes())
 	{
 		ITickableAttributeSetInterface* TickableSet = Cast<ITickableAttributeSetInterface>(AttributeSet);
-		if (TickableSet)
+		if (TickableSet && TickableSet->ShouldTick())
 		{
 			TickableSet->Tick(DeltaTime);
 		}
