GameplayAbilityWorldReticle should face towards the source actor instead of the TargetingActor

diff --git a/Engine/Plugins/Runtime/GameplayAbilities/Source/GameplayAbilities/Private/Abilities/GameplayAbilityWorldReticle.cpp b/Engine/Plugins/Runtime/GameplayAbilities/Source/GameplayAbilities/Private/Abilities/GameplayAbilityWorldReticle.cpp
--- a/Engine/Plugins/Runtime/GameplayAbilities/Source/GameplayAbilities/Private/Abilities/GameplayAbilityWorldReticle.cpp
+++ b/Engine/Plugins/Runtime/GameplayAbilities/Source/GameplayAbilities/Private/Abilities/GameplayAbilityWorldReticle.cpp
@@ -1,6 +1,7 @@
 // Copyright Epic Games, Inc. All Rights Reserved.
 
 #include "Abilities/GameplayAbilityWorldReticle.h"
+#include "Abilities/GameplayAbilityTargetActor.h"
 #include "GameFramework/PlayerController.h"
 
 #include UE_INLINE_GENERATED_CPP_BY_NAME(GameplayAbilityWorldReticle)
@@ -29,7 +30,7 @@ void AGameplayAbilityWorldReticle::Tick(float DeltaTime)
 	FaceTowardSource(bFaceOwnerFlat);
 }
 
-void AGameplayAbilityWorldReticle::InitializeReticle(AActor* InTargetingActor, APlayerController* PlayerController, FWorldReticleParameters InParameters)
+void AGameplayAbilityWorldReticle::InitializeReticle(AGameplayAbilityTargetActor* InTargetingActor, APlayerController* PlayerController, FWorldReticleParameters InParameters)
 {
 	check(InTargetingActor);
 	TargetingActor = InTargetingActor;
@@ -70,11 +71,18 @@ void AGameplayAbilityWorldReticle::SetIsTargetAnActor(bool bNewValue)
 
 void AGameplayAbilityWorldReticle::FaceTowardSource(bool bFaceIn2D)
 {
-	if (TargetingActor)
+	AActor* FacingActor = TargetingActor.Get();
+
+	if (TargetingActor.Get() && TargetingActor->SourceActor)
+	{
+		FacingActor = TargetingActor->SourceActor;
+	}
+	
+	if (FacingActor)
 	{
 		if (bFaceIn2D)
 		{
-			FVector FacingVector = (TargetingActor->GetActorLocation() - GetActorLocation()).GetSafeNormal2D();
+			FVector FacingVector = (FacingActor->GetActorLocation() - GetActorLocation()).GetSafeNormal2D();
 			if (FacingVector.IsZero())
 			{
 				FacingVector = -GetActorForwardVector().GetSafeNormal2D();
@@ -86,7 +94,7 @@ void AGameplayAbilityWorldReticle::FaceTowardSource(bool bFaceIn2D)
 		}
 		else
 		{
-			FVector FacingVector = (TargetingActor->GetActorLocation() - GetActorLocation()).GetSafeNormal();
+			FVector FacingVector = (FacingActor->GetActorLocation() - GetActorLocation()).GetSafeNormal();
 			if (FacingVector.IsZero())
 			{
 				FacingVector = -GetActorForwardVector().GetSafeNormal();
diff --git a/Engine/Plugins/Runtime/GameplayAbilities/Source/GameplayAbilities/Public/Abilities/GameplayAbilityWorldReticle.h b/Engine/Plugins/Runtime/GameplayAbilities/Source/GameplayAbilities/Public/Abilities/GameplayAbilityWorldReticle.h
--- a/Engine/Plugins/Runtime/GameplayAbilities/Source/GameplayAbilities/Public/Abilities/GameplayAbilityWorldReticle.h
+++ b/Engine/Plugins/Runtime/GameplayAbilities/Source/GameplayAbilities/Public/Abilities/GameplayAbilityWorldReticle.h
@@ -7,6 +7,7 @@
 #include "GameplayAbilityWorldReticle.generated.h"
 
 class APlayerController;
+class AGameplayAbilityTargetActor;
 
 USTRUCT(BlueprintType)
 struct FWorldReticleParameters
@@ -37,7 +38,7 @@ public:
 
 	virtual bool IsNetRelevantFor(const AActor* RealViewer, const AActor* ViewTarget, const FVector& SrcLocation) const override;
 
-	void InitializeReticle(AActor* InTargetingActor, APlayerController* PlayerController, FWorldReticleParameters InParameters);
+	void InitializeReticle(AGameplayAbilityTargetActor* InTargetingActor, APlayerController* PlayerController, FWorldReticleParameters InParameters);
 
 	void SetIsTargetValid(bool bNewValue);
 	void SetIsTargetAnActor(bool bNewValue);
@@ -94,5 +95,5 @@ protected:
 
 	/** In the future, we may want to grab things like sockets off of this. */
 	UPROPERTY(BlueprintReadOnly, Category = "Network")
-	TObjectPtr<AActor> TargetingActor;
+	TObjectPtr<AGameplayAbilityTargetActor> TargetingActor;
 };
