AnimTask_PlayMontageAndWait used to consider BlendOut as having finished.  However, there are cases where you can BlendOut and still want to know if the Montage has Completed or was Interrupted.

diff --git a/Engine/Plugins/Runtime/GameplayAbilities/Source/GameplayAbilities/Private/Abilities/Tasks/AbilityTask.cpp b/Engine/Plugins/Runtime/GameplayAbilities/Source/GameplayAbilities/Private/Abilities/Tasks/AbilityTask.cpp
--- a/Engine/Plugins/Runtime/GameplayAbilities/Source/GameplayAbilities/Private/Abilities/Tasks/AbilityTask.cpp
+++ b/Engine/Plugins/Runtime/GameplayAbilities/Source/GameplayAbilities/Private/Abilities/Tasks/AbilityTask.cpp
@@ -185,7 +185,7 @@ bool UAbilityTask::ShouldBroadcastAbilityTaskDelegates() const
 
 	if (!ShouldBroadcast && AbilityTaskWarnIfBroadcastSuppress)
 	{
-		ABILITY_LOG(Warning, TEXT("Suppressing ability task %s broadcsat"), *GetDebugString());
+		ABILITY_LOG(Warning, TEXT("Suppressing ability task %s broadcast"), *GetDebugString());
 	}
 
 	return ShouldBroadcast;
diff --git a/Engine/Plugins/Runtime/GameplayAbilities/Source/GameplayAbilities/Private/Abilities/Tasks/AbilityTask_PlayMontageAndWait.cpp b/Engine/Plugins/Runtime/GameplayAbilities/Source/GameplayAbilities/Private/Abilities/Tasks/AbilityTask_PlayMontageAndWait.cpp
--- a/Engine/Plugins/Runtime/GameplayAbilities/Source/GameplayAbilities/Private/Abilities/Tasks/AbilityTask_PlayMontageAndWait.cpp
+++ b/Engine/Plugins/Runtime/GameplayAbilities/Source/GameplayAbilities/Private/Abilities/Tasks/AbilityTask_PlayMontageAndWait.cpp
@@ -9,45 +9,43 @@
 
 #include UE_INLINE_GENERATED_CPP_BY_NAME(AbilityTask_PlayMontageAndWait)
 
-UAbilityTask_PlayMontageAndWait::UAbilityTask_PlayMontageAndWait(const FObjectInitializer& ObjectInitializer)
-	: Super(ObjectInitializer)
-{
-	Rate = 1.f;
-	bStopWhenAbilityEnds = true;
-}
+static bool GUseAggressivePlayMontageAndWaitEndTask = true;
+static FAutoConsoleVariableRef CVarAggressivePlayMontageAndWaitEndTask(TEXT("AbilitySystem.PlayMontage.AggressiveEndTask"), GUseAggressivePlayMontageAndWaitEndTask, TEXT("This should be set to true in order to avoid multiple callbacks off an AbilityTask_PlayMontageAndWait node"));
 
 void UAbilityTask_PlayMontageAndWait::OnMontageBlendingOut(UAnimMontage* Montage, bool bInterrupted)
 {
-	if (Ability && Ability->GetCurrentMontage() == MontageToPlay)
+	const bool bPlayingThisMontage = (Montage == MontageToPlay) && Ability && Ability->GetCurrentMontage() == MontageToPlay;
+	if (bPlayingThisMontage)
 	{
-		if (Montage == MontageToPlay)
+		// Reset AnimRootMotionTranslationScale
+		ACharacter* Character = Cast<ACharacter>(GetAvatarActor());
+		if (Character && (Character->GetLocalRole() == ROLE_Authority ||
+							(Character->GetLocalRole() == ROLE_AutonomousProxy && Ability->GetNetExecutionPolicy() == EGameplayAbilityNetExecutionPolicy::LocalPredicted)))
 		{
-			if (UAbilitySystemComponent* ASC = AbilitySystemComponent.Get())
-			{
-				ASC->ClearAnimatingAbility(Ability);
-			}
-
-			// Reset AnimRootMotionTranslationScale
-			ACharacter* Character = Cast<ACharacter>(GetAvatarActor());
-			if (Character && (Character->GetLocalRole() == ROLE_Authority ||
-							  (Character->GetLocalRole() == ROLE_AutonomousProxy && Ability->GetNetExecutionPolicy() == EGameplayAbilityNetExecutionPolicy::LocalPredicted)))
-			{
-				Character->SetAnimRootMotionTranslationScale(1.f);
-			}
-
+			Character->SetAnimRootMotionTranslationScale(1.f);
 		}
 	}
 
-	if (bInterrupted)
+	if (bPlayingThisMontage && (bInterrupted || !bAllowInterruptAfterBlendOut))
 	{
-		if (ShouldBroadcastAbilityTaskDelegates())
+		if (UAbilitySystemComponent* ASC = AbilitySystemComponent.Get())
 		{
-			OnInterrupted.Broadcast();
+			ASC->ClearAnimatingAbility(Ability);
 		}
 	}
-	else
+
+	if (ShouldBroadcastAbilityTaskDelegates())
 	{
-		if (ShouldBroadcastAbilityTaskDelegates())
+		if (bInterrupted)
+		{
+			OnInterrupted.Broadcast();
+
+			if (GUseAggressivePlayMontageAndWaitEndTask)
+			{
+				EndTask();
+			}
+		}
+		else
 		{
 			OnBlendOut.Broadcast();
 		}
@@ -56,7 +54,13 @@ void UAbilityTask_PlayMontageAndWait::OnMontageBlendingOut(UAnimMontage* Montage
 
 void UAbilityTask_PlayMontageAndWait::OnMontageInterrupted()
 {
-	if (StopPlayingMontage())
+	// Call the new function
+	OnGameplayAbilityCancelled();
+}
+
+void UAbilityTask_PlayMontageAndWait::OnGameplayAbilityCancelled()
+{
+	if (StopPlayingMontage() || bAllowInterruptAfterBlendOut)
 	{
 		// Let the BP handle the interrupt as well
 		if (ShouldBroadcastAbilityTaskDelegates())
@@ -64,6 +68,11 @@ void UAbilityTask_PlayMontageAndWait::OnMontageInterrupted()
 			OnInterrupted.Broadcast();
 		}
 	}
+
+	if (GUseAggressivePlayMontageAndWaitEndTask)
+	{
+		EndTask();
+	}
 }
 
 void UAbilityTask_PlayMontageAndWait::OnMontageEnded(UAnimMontage* Montage, bool bInterrupted)
@@ -80,7 +89,7 @@ void UAbilityTask_PlayMontageAndWait::OnMontageEnded(UAnimMontage* Montage, bool
 }
 
 UAbilityTask_PlayMontageAndWait* UAbilityTask_PlayMontageAndWait::CreatePlayMontageAndWaitProxy(UGameplayAbility* OwningAbility,
-	FName TaskInstanceName, UAnimMontage *MontageToPlay, float Rate, FName StartSection, bool bStopWhenAbilityEnds, float AnimRootMotionTranslationScale, float StartTimeSeconds)
+	FName TaskInstanceName, UAnimMontage *MontageToPlay, float Rate, FName StartSection, bool bStopWhenAbilityEnds, float AnimRootMotionTranslationScale, float StartTimeSeconds, bool bAllowInterruptAfterBlendOut)
 {
 
 	UAbilitySystemGlobals::NonShipping_ApplyGlobalAbilityScaler_Rate(Rate);
@@ -91,6 +100,7 @@ UAbilityTask_PlayMontageAndWait* UAbilityTask_PlayMontageAndWait::CreatePlayMont
 	MyObj->StartSection = StartSection;
 	MyObj->AnimRootMotionTranslationScale = AnimRootMotionTranslationScale;
 	MyObj->bStopWhenAbilityEnds = bStopWhenAbilityEnds;
+	MyObj->bAllowInterruptAfterBlendOut = bAllowInterruptAfterBlendOut;
 	MyObj->StartTimeSeconds = StartTimeSeconds;
 	
 	return MyObj;
@@ -119,7 +129,7 @@ void UAbilityTask_PlayMontageAndWait::Activate()
 					return;
 				}
 
-				InterruptedHandle = Ability->OnGameplayAbilityCancelled.AddUObject(this, &UAbilityTask_PlayMontageAndWait::OnMontageInterrupted);
+				InterruptedHandle = Ability->OnGameplayAbilityCancelled.AddUObject(this, &UAbilityTask_PlayMontageAndWait::OnGameplayAbilityCancelled);
 
 				BlendingOutDelegate.BindUObject(this, &UAbilityTask_PlayMontageAndWait::OnMontageBlendingOut);
 				AnimInstance->Montage_SetBlendingOutDelegate(BlendingOutDelegate, MontageToPlay);
diff --git a/Engine/Plugins/Runtime/GameplayAbilities/Source/GameplayAbilities/Public/Abilities/Tasks/AbilityTask_PlayMontageAndWait.h b/Engine/Plugins/Runtime/GameplayAbilities/Source/GameplayAbilities/Public/Abilities/Tasks/AbilityTask_PlayMontageAndWait.h
--- a/Engine/Plugins/Runtime/GameplayAbilities/Source/GameplayAbilities/Public/Abilities/Tasks/AbilityTask_PlayMontageAndWait.h
+++ b/Engine/Plugins/Runtime/GameplayAbilities/Source/GameplayAbilities/Public/Abilities/Tasks/AbilityTask_PlayMontageAndWait.h
@@ -13,7 +13,8 @@ DECLARE_DYNAMIC_MULTICAST_DELEGATE(FMontageWaitSimpleDelegate);
 UCLASS()
 class GAMEPLAYABILITIES_API UAbilityTask_PlayMontageAndWait : public UAbilityTask
 {
-	GENERATED_UCLASS_BODY()
+	GENERATED_BODY()
+public:
 
 	UPROPERTY(BlueprintAssignable)
 	FMontageWaitSimpleDelegate	OnCompleted;
@@ -30,9 +31,14 @@ class GAMEPLAYABILITIES_API UAbilityTask_PlayMontageAndWait : public UAbilityTas
 	UFUNCTION()
 	void OnMontageBlendingOut(UAnimMontage* Montage, bool bInterrupted);
 
+	UE_DEPRECATED(5.3, "Please use OnGameplayAbilityCancelled instead. This function naming implied the Montage was already interrupted (instead, we are about to interrupt it).")
 	UFUNCTION()
 	void OnMontageInterrupted();
 
+	/** Callback function for when the owning Gameplay Ability is cancelled */
+	UFUNCTION()
+	void OnGameplayAbilityCancelled();
+
 	UFUNCTION()
 	void OnMontageEnded(UAnimMontage* Montage, bool bInterrupted);
 
@@ -49,11 +55,12 @@ class GAMEPLAYABILITIES_API UAbilityTask_PlayMontageAndWait : public UAbilityTas
 	 * @param bStopWhenAbilityEnds If true, this montage will be aborted if the ability ends normally. It is always stopped when the ability is explicitly cancelled
 	 * @param AnimRootMotionTranslationScale Change to modify size of root motion or set to 0 to block it entirely
 	 * @param StartTimeSeconds Starting time offset in montage, this will be overridden by StartSection if that is also set
+	 * @param bAllowInterruptAfterBlendOut If true, you can receive OnInterrupted after an OnBlendOut started (otherwise OnInterrupted will not fire when interrupted, but you will not get OnComplete).
 	 */
 	UFUNCTION(BlueprintCallable, Category="Ability|Tasks", meta = (DisplayName="PlayMontageAndWait",
 		HidePin = "OwningAbility", DefaultToSelf = "OwningAbility", BlueprintInternalUseOnly = "TRUE"))
 	static UAbilityTask_PlayMontageAndWait* CreatePlayMontageAndWaitProxy(UGameplayAbility* OwningAbility,
-		FName TaskInstanceName, UAnimMontage* MontageToPlay, float Rate = 1.f, FName StartSection = NAME_None, bool bStopWhenAbilityEnds = true, float AnimRootMotionTranslationScale = 1.f, float StartTimeSeconds = 0.f);
+		FName TaskInstanceName, UAnimMontage* MontageToPlay, float Rate = 1.f, FName StartSection = NAME_None, bool bStopWhenAbilityEnds = true, float AnimRootMotionTranslationScale = 1.f, float StartTimeSeconds = 0.f, bool bAllowInterruptAfterBlendOut = false);
 
 	virtual void Activate() override;
 
@@ -90,4 +97,7 @@ protected:
 
 	UPROPERTY()
 	bool bStopWhenAbilityEnds;
+
+	UPROPERTY()
+	bool bAllowInterruptAfterBlendOut;
 };
