Fix Gameplay Effect Components not being able to apply their Tags after changes in the Editor (and prior to a reload of the Gameplay Effect).

diff --git a/Engine/Plugins/Runtime/GameplayAbilities/Source/GameplayAbilities/Private/GameplayEffect.cpp b/Engine/Plugins/Runtime/GameplayAbilities/Source/GameplayAbilities/Private/GameplayEffect.cpp
--- a/Engine/Plugins/Runtime/GameplayAbilities/Source/GameplayAbilities/Private/GameplayEffect.cpp
+++ b/Engine/Plugins/Runtime/GameplayAbilities/Source/GameplayAbilities/Private/GameplayEffect.cpp
@@ -179,15 +179,7 @@ void UGameplayEffect::PostLoad()
 {
 	Super::PostLoad();
 
-	// Call PostLoad on components and cache tag info from components for more efficient lookups
-	for (UGameplayEffectComponent* GEComponent : GEComponents)
-	{
-		if (GEComponent)
-		{
-			GEComponent->ConditionalPostLoad(); // Ensure the SubObject is PostLoaded
-			GEComponent->OnOwnerPostLoad();
-		}
-	}
+	OnGameplayEffectChanged();
 
 #if WITH_EDITOR
 	SCALABLEFLOAT_REPORTERROR_WITHPOSTLOAD(Period);
@@ -212,6 +204,31 @@ void UGameplayEffect::PostLoad()
 #endif // WITH_EDITOR
 }
 
+void UGameplayEffect::OnGameplayEffectChanged()
+{
+	if (HasAnyFlags(RF_NeedPostLoad))
+	{
+		ensureMsgf(false, TEXT("%s: OnGameplayEffectChanged can only be called after the GameplayEffect is fully loaded"), *GetName());
+		return;
+	}
+
+	// Reset these tags so we can reaggregate them properly from the GEComponents
+	CachedAssetTags.Reset();
+	CachedGrantedTags.Reset();
+	CachedBlockedAbilityTags.Reset();
+
+	// Call PostLoad on components and cache tag info from components for more efficient lookups
+	for (UGameplayEffectComponent* GEComponent : GEComponents)
+	{
+		if (GEComponent)
+		{
+			// Ensure the SubObject is fully loaded
+			GEComponent->ConditionalPostLoad();
+			GEComponent->OnGameplayEffectChanged();
+		}
+	}
+}
+
 #if WITH_EDITOR
 /**
  * We now support GameplayEffectComponents which are Subobjects.
@@ -287,6 +304,12 @@ void UGameplayEffect::PostCDOCompiled(const FPostCDOCompiledContext& Context)
 	ConvertTagRequirementsComponent();
 	ConvertTargetTagsComponent();
 	ConvertUIComponent();
+
+	const bool bAlreadyLoaded = !HasAnyFlags(RF_NeedPostLoad);
+	if (bAlreadyLoaded)
+	{
+		OnGameplayEffectChanged();
+	}
 }
 
 bool UGameplayEffect::VersionAllowsCreateGEComponentsDuringUpgrade() const
diff --git a/Engine/Plugins/Runtime/GameplayAbilities/Source/GameplayAbilities/Private/GameplayEffectComponents/AssetTagsGameplayEffectComponent.cpp b/Engine/Plugins/Runtime/GameplayAbilities/Source/GameplayAbilities/Private/GameplayEffectComponents/AssetTagsGameplayEffectComponent.cpp
--- a/Engine/Plugins/Runtime/GameplayAbilities/Source/GameplayAbilities/Private/GameplayEffectComponents/AssetTagsGameplayEffectComponent.cpp
+++ b/Engine/Plugins/Runtime/GameplayAbilities/Source/GameplayAbilities/Private/GameplayEffectComponents/AssetTagsGameplayEffectComponent.cpp
@@ -20,10 +20,10 @@ void UAssetTagsGameplayEffectComponent::PostInitProperties()
 #endif // WITH_EDITORONLY_DATA
 }
 
-void UAssetTagsGameplayEffectComponent::OnOwnerPostLoad()
+void UAssetTagsGameplayEffectComponent::OnGameplayEffectChanged() const
 {
-	Super::OnOwnerPostLoad();
-	SetAndApplyAssetTagChanges(InheritableAssetTags);
+	Super::OnGameplayEffectChanged();
+	ApplyAssetTagChanges();
 }
 
 #if WITH_EDITOR
@@ -34,6 +34,10 @@ void UAssetTagsGameplayEffectComponent::PostEditChangeProperty(FPropertyChangedE
 	if (PropertyChangedEvent.GetMemberPropertyName()  == GetInheritableAssetTagsName())
 	{
 		SetAndApplyAssetTagChanges(InheritableAssetTags);
+		
+		// Tell the GE it needs to reconfigure itself based on these updated properties (this will reaggregate the tags)
+		UGameplayEffect* Owner = GetOwner();
+		Owner->OnGameplayEffectChanged();
 	}
 }
 #endif // WITH_EDITOR
@@ -46,6 +50,11 @@ void UAssetTagsGameplayEffectComponent::SetAndApplyAssetTagChanges(const FInheri
 	const UAssetTagsGameplayEffectComponent* Parent = FindParentComponent(*this);
 	InheritableAssetTags.UpdateInheritedTagProperties(Parent ? &Parent->InheritableAssetTags : nullptr);
 
+	ApplyAssetTagChanges();
+}
+
+void UAssetTagsGameplayEffectComponent::ApplyAssetTagChanges() const
+{
 	// Apply to the owning Gameplay Effect Component
 	UGameplayEffect* Owner = GetOwner();
 	InheritableAssetTags.ApplyTo(Owner->CachedAssetTags);
diff --git a/Engine/Plugins/Runtime/GameplayAbilities/Source/GameplayAbilities/Private/GameplayEffectComponents/BlockAbilityTagsGameplayEffectComponent.cpp b/Engine/Plugins/Runtime/GameplayAbilities/Source/GameplayAbilities/Private/GameplayEffectComponents/BlockAbilityTagsGameplayEffectComponent.cpp
--- a/Engine/Plugins/Runtime/GameplayAbilities/Source/GameplayAbilities/Private/GameplayEffectComponents/BlockAbilityTagsGameplayEffectComponent.cpp
+++ b/Engine/Plugins/Runtime/GameplayAbilities/Source/GameplayAbilities/Private/GameplayEffectComponents/BlockAbilityTagsGameplayEffectComponent.cpp
@@ -9,10 +9,10 @@ UBlockAbilityTagsGameplayEffectComponent::UBlockAbilityTagsGameplayEffectCompone
 #endif
 }
 
-void UBlockAbilityTagsGameplayEffectComponent::OnOwnerPostLoad()
+void UBlockAbilityTagsGameplayEffectComponent::OnGameplayEffectChanged() const
 {
-	Super::OnOwnerPostLoad();
-	SetAndApplyBlockedAbilityTagChanges(InheritableBlockedAbilityTagsContainer);
+	Super::OnGameplayEffectChanged();
+	ApplyBlockedAbilityTagChanges();
 }
 
 #if WITH_EDITOR
@@ -23,6 +23,10 @@ void UBlockAbilityTagsGameplayEffectComponent::PostEditChangeProperty(FPropertyC
 	if (PropertyChangedEvent.GetMemberPropertyName() == GetInheritableBlockedAbilityTagsContainerPropertyName())
 	{
 		SetAndApplyBlockedAbilityTagChanges(InheritableBlockedAbilityTagsContainer);
+		
+		// Tell the GE it needs to reconfigure itself based on these updated properties (this will reaggregate the tags)
+		UGameplayEffect* Owner = GetOwner();
+		Owner->OnGameplayEffectChanged();
 	}
 }
 #endif // WITH_EDITOR
@@ -35,7 +39,12 @@ void UBlockAbilityTagsGameplayEffectComponent::SetAndApplyBlockedAbilityTagChang
 	const UBlockAbilityTagsGameplayEffectComponent* Parent = FindParentComponent(*this);
 	InheritableBlockedAbilityTagsContainer.UpdateInheritedTagProperties(Parent ? &Parent->InheritableBlockedAbilityTagsContainer : nullptr);
 
+	ApplyBlockedAbilityTagChanges();
+}
+
+void UBlockAbilityTagsGameplayEffectComponent::ApplyBlockedAbilityTagChanges() const
+{
 	// Apply to the owning Gameplay Effect Component
 	UGameplayEffect* Owner = GetOwner();
 	InheritableBlockedAbilityTagsContainer.ApplyTo(Owner->CachedBlockedAbilityTags);
-}
+}
\ No newline at end of file
diff --git a/Engine/Plugins/Runtime/GameplayAbilities/Source/GameplayAbilities/Private/GameplayEffectComponents/TargetTagsGameplayEffectComponent.cpp b/Engine/Plugins/Runtime/GameplayAbilities/Source/GameplayAbilities/Private/GameplayEffectComponents/TargetTagsGameplayEffectComponent.cpp
--- a/Engine/Plugins/Runtime/GameplayAbilities/Source/GameplayAbilities/Private/GameplayEffectComponents/TargetTagsGameplayEffectComponent.cpp
+++ b/Engine/Plugins/Runtime/GameplayAbilities/Source/GameplayAbilities/Private/GameplayEffectComponents/TargetTagsGameplayEffectComponent.cpp
@@ -9,13 +9,14 @@ UTargetTagsGameplayEffectComponent::UTargetTagsGameplayEffectComponent()
 #endif // WITH_EDITORONLY_DATA
 }
 
-void UTargetTagsGameplayEffectComponent::OnOwnerPostLoad()
+void UTargetTagsGameplayEffectComponent::OnGameplayEffectChanged() const
 {
-	Super::OnOwnerPostLoad();
-	SetAndApplyTargetTagChanges(InheritableGrantedTagsContainer);
+	Super::OnGameplayEffectChanged();
+	ApplyTargetTagChanges();
 }
 
 #if WITH_EDITOR
+
 void UTargetTagsGameplayEffectComponent::PostEditChangeProperty(FPropertyChangedEvent& PropertyChangedEvent)
 {
 	Super::PostEditChangeProperty(PropertyChangedEvent);
@@ -23,6 +24,10 @@ void UTargetTagsGameplayEffectComponent::PostEditChangeProperty(FPropertyChanged
 	if (PropertyChangedEvent.GetMemberPropertyName() == GetInheritableGrantedTagsContainerName())
 	{
 		SetAndApplyTargetTagChanges(InheritableGrantedTagsContainer);
+
+		// Tell the GE it needs to reconfigure itself based on these updated properties (this will reaggregate the tags)
+		UGameplayEffect* Owner = GetOwner();
+		Owner->OnGameplayEffectChanged();
 	}
 }
 #endif // WITH_EDITOR
@@ -36,6 +41,11 @@ void UTargetTagsGameplayEffectComponent::SetAndApplyTargetTagChanges(const FInhe
 	InheritableGrantedTagsContainer.UpdateInheritedTagProperties(Parent ? &Parent->InheritableGrantedTagsContainer : nullptr);
 
 	// Apply to the owning Gameplay Effect Component
+	ApplyTargetTagChanges();
+}
+
+void UTargetTagsGameplayEffectComponent::ApplyTargetTagChanges() const
+{
 	UGameplayEffect* Owner = GetOwner();
 	InheritableGrantedTagsContainer.ApplyTo(Owner->CachedGrantedTags);
 }
diff --git a/Engine/Plugins/Runtime/GameplayAbilities/Source/GameplayAbilities/Public/GameplayEffect.h b/Engine/Plugins/Runtime/GameplayAbilities/Source/GameplayAbilities/Public/GameplayEffect.h
--- a/Engine/Plugins/Runtime/GameplayAbilities/Source/GameplayAbilities/Public/GameplayEffect.h
+++ b/Engine/Plugins/Runtime/GameplayAbilities/Source/GameplayAbilities/Public/GameplayEffect.h
@@ -2043,11 +2043,15 @@ public:
 
 	virtual void GetBlockedAbilityTags(FGameplayTagContainer& OutTagContainer) const;
 
-	// Needed to properly disable inheriting the version value from its parent.
+	/** Needed to properly disable inheriting the version value from its parent. */
 	virtual void PostInitProperties() override;
 
+	/** PostLoad gets called once after the asset has been loaded.  It does not get called again on Blueprint recompile (@see PostCDOCompiled) */
 	virtual void PostLoad() override;
 
+	/** Called when the Gameplay Effect has finished loading.  This is used to catch both PostLoad (the initial load) and PostCDOCompiled (any subsequent changes) */
+	virtual void OnGameplayEffectChanged();
+
 #if WITH_EDITOR
 	/** Do our upgrades in PostCDOCompiled */
 	virtual void PostCDOCompiled(const FPostCDOCompiledContext& Context) override;
diff --git a/Engine/Plugins/Runtime/GameplayAbilities/Source/GameplayAbilities/Public/GameplayEffectComponent.h b/Engine/Plugins/Runtime/GameplayAbilities/Source/GameplayAbilities/Public/GameplayEffectComponent.h
--- a/Engine/Plugins/Runtime/GameplayAbilities/Source/GameplayAbilities/Public/GameplayEffectComponent.h
+++ b/Engine/Plugins/Runtime/GameplayAbilities/Source/GameplayAbilities/Public/GameplayEffectComponent.h
@@ -57,9 +57,9 @@ public:
 	virtual void OnGameplayEffectExecuted(FActiveGameplayEffectsContainer& ActiveGEContainer, FGameplayEffectSpec& GESpec, FPredictionKey& PredictionKey) const {}
 
 	/**
-	 * Let us know that the owning GameplayEffect has finished its PostLoad, thus apply an asset-related changes to the owning GameplayEffect (e.g. any of its fields)
+	 * Let us know that the owning GameplayEffect has been modified, thus apply an asset-related changes to the owning GameplayEffect (e.g. any of its fields)
 	 */
-	virtual void OnOwnerPostLoad() {}
+	virtual void OnGameplayEffectChanged() const {}
 
 #if WITH_EDITOR
 	/**
diff --git a/Engine/Plugins/Runtime/GameplayAbilities/Source/GameplayAbilities/Public/GameplayEffectComponents/AssetTagsGameplayEffectComponent.h b/Engine/Plugins/Runtime/GameplayAbilities/Source/GameplayAbilities/Public/GameplayEffectComponents/AssetTagsGameplayEffectComponent.h
--- a/Engine/Plugins/Runtime/GameplayAbilities/Source/GameplayAbilities/Public/GameplayEffectComponents/AssetTagsGameplayEffectComponent.h
+++ b/Engine/Plugins/Runtime/GameplayAbilities/Source/GameplayAbilities/Public/GameplayEffectComponents/AssetTagsGameplayEffectComponent.h
@@ -17,7 +17,7 @@ public:
 	virtual void PostInitProperties() override;
 
 	/** Needed to properly apply FInheritedTagContainer properties */
-	virtual void OnOwnerPostLoad() override;
+	virtual void OnGameplayEffectChanged() const override;
 
 	/** Gets the Asset Tags inherited tag structure (as configured) */
 	const FInheritedTagContainer& GetConfiguredAssetTagChanges() const { return InheritableAssetTags; }
@@ -38,6 +38,10 @@ private:
 	}
 #endif // WITH_EDITOR
 
+private:
+	/** Applies the Asset Tags to the GameplayEffect */
+	void ApplyAssetTagChanges() const;
+
 private:
 	/** The GameplayEffect's Tags: tags the the GE *has* and DOES NOT give to the actor. */
 	UPROPERTY(EditDefaultsOnly, Category = Tags, meta = (DisplayName = "AssetTags", Categories = "OwnedTagsCategory"))
diff --git a/Engine/Plugins/Runtime/GameplayAbilities/Source/GameplayAbilities/Public/GameplayEffectComponents/BlockAbilityTagsGameplayEffectComponent.h b/Engine/Plugins/Runtime/GameplayAbilities/Source/GameplayAbilities/Public/GameplayEffectComponents/BlockAbilityTagsGameplayEffectComponent.h
--- a/Engine/Plugins/Runtime/GameplayAbilities/Source/GameplayAbilities/Public/GameplayEffectComponents/BlockAbilityTagsGameplayEffectComponent.h
+++ b/Engine/Plugins/Runtime/GameplayAbilities/Source/GameplayAbilities/Public/GameplayEffectComponents/BlockAbilityTagsGameplayEffectComponent.h
@@ -17,7 +17,7 @@ public:
 	UBlockAbilityTagsGameplayEffectComponent();
 
 	/** Needed to properly apply FInheritedTagContainer properties */
-	virtual void OnOwnerPostLoad() override;
+	virtual void OnGameplayEffectChanged() const override;
 
 	/** Gets the Blocked Ability Tags inherited tag structure (as configured) */
 	const FInheritedTagContainer& GetConfiguredBlockedAbilityTagChanges() const { return InheritableBlockedAbilityTagsContainer; }
@@ -38,6 +38,10 @@ private:
 	}
 #endif // WITH_EDITOR
 
+private:
+	/** Applies the Blocked Ability Tags to the owning GameplayEffect */
+	void ApplyBlockedAbilityTagChanges() const;
+
 private:
 	/** These tags are applied to the target actor of the Gameplay Effect.  Blocked Ability Tags prevent Gameplay Abilities with these tags from executing. */
 	UPROPERTY(EditDefaultsOnly, Category = None, meta = (DisplayName = "Block Abilities w/ Tags"))
diff --git a/Engine/Plugins/Runtime/GameplayAbilities/Source/GameplayAbilities/Public/GameplayEffectComponents/TargetTagsGameplayEffectComponent.h b/Engine/Plugins/Runtime/GameplayAbilities/Source/GameplayAbilities/Public/GameplayEffectComponents/TargetTagsGameplayEffectComponent.h
--- a/Engine/Plugins/Runtime/GameplayAbilities/Source/GameplayAbilities/Public/GameplayEffectComponents/TargetTagsGameplayEffectComponent.h
+++ b/Engine/Plugins/Runtime/GameplayAbilities/Source/GameplayAbilities/Public/GameplayEffectComponents/TargetTagsGameplayEffectComponent.h
@@ -17,7 +17,7 @@ public:
 	UTargetTagsGameplayEffectComponent();
 
 	/** Needed to properly load FInheritedTagContainer properties */
-	virtual void OnOwnerPostLoad() override;
+	virtual void OnGameplayEffectChanged() const override;
 
 	/** Gets the (Granted) Target Tags inherited tag structure (as configured) */
 	const FInheritedTagContainer& GetConfiguredTargetTagChanges() const { return InheritableGrantedTagsContainer; }
@@ -38,6 +38,10 @@ private:
 	}
 #endif // WITH_EDITOR
 
+private:
+	/** Applies the configured (Granted) Target Tags to the owning GameplayEffect. */
+	void ApplyTargetTagChanges() const;
+
 private:
 	/** These tags are applied (granted) to the target actor of the Gameplay Effect.  The Target would then "own" these Tags. */
 	UPROPERTY(EditDefaultsOnly, Category = None, meta = (DisplayName = "Add Tags", Categories = "OwnedTagsCategory"))
