Call OnAvatarSet on the primary instance instead of the CDO for instanced GameplayAbilities

diff --git a/Engine/Plugins/Runtime/GameplayAbilities/Source/GameplayAbilities/Private/AbilitySystemComponent_Abilities.cpp b/Engine/Plugins/Runtime/GameplayAbilities/Source/GameplayAbilities/Private/AbilitySystemComponent_Abilities.cpp
--- a/Engine/Plugins/Runtime/GameplayAbilities/Source/GameplayAbilities/Private/AbilitySystemComponent_Abilities.cpp
+++ b/Engine/Plugins/Runtime/GameplayAbilities/Source/GameplayAbilities/Private/AbilitySystemComponent_Abilities.cpp
@@ -159,7 +159,18 @@ void UAbilitySystemComponent::InitAbilityActorInfo(AActor* InOwnerActor, AActor*
 		{
 			if (Spec.Ability)
 			{
-				Spec.Ability->OnAvatarSet(AbilityActorInfo.Get(), Spec);
+				if (Spec.Ability->GetInstancingPolicy() == EGameplayAbilityInstancingPolicy::InstancedPerActor)
+				{
+					UGameplayAbility* AbilityInstance = Spec.GetPrimaryInstance();
+					if (ensureMsgf(AbilityInstance, TEXT("%s: Could not find a Primary Instance for InstancedPerActor Ability Spec!"), ANSI_TO_TCHAR(__func__)))
+					{
+						AbilityInstance->OnAvatarSet(AbilityActorInfo.Get(), Spec);
+					}
+				}
+				else
+				{
+					Spec.Ability->OnAvatarSet(AbilityActorInfo.Get(), Spec);
+				}
 			}
 		}
 	}
