Remove unused variable in FScopedPredictionWindow that caused some confusion.

diff --git a/Engine/Plugins/Runtime/GameplayAbilities/Source/GameplayAbilities/Private/GameplayPrediction.cpp b/Engine/Plugins/Runtime/GameplayAbilities/Source/GameplayAbilities/Private/GameplayPrediction.cpp
--- a/Engine/Plugins/Runtime/GameplayAbilities/Source/GameplayAbilities/Private/GameplayPrediction.cpp
+++ b/Engine/Plugins/Runtime/GameplayAbilities/Source/GameplayAbilities/Private/GameplayPrediction.cpp
@@ -311,7 +311,7 @@ FScopedPredictionWindow::FScopedPredictionWindow(UAbilitySystemComponent* InAbil
 	// Original ensure has been left in to catch other cases of invalid Owner ASCs
 	if ((!InAbilitySystemComponent) || (InAbilitySystemComponent->IsBeingDestroyed()) || (!IsValidChecked(InAbilitySystemComponent) || InAbilitySystemComponent->IsUnreachable()))
 	{
-		ABILITY_LOG(Verbose, TEXT("FScopedPredictionWindow() aborting due to Owner (ASC) being null, destroyed or pending kill / unreachable [%s]"), *ScopedPredictionKey.ToString());
+		ABILITY_LOG(Verbose, TEXT("FScopedPredictionWindow() aborting due to Owner (ASC) being null, destroyed or pending kill / unreachable"));
 		return;
 	}
 
diff --git a/Engine/Plugins/Runtime/GameplayAbilities/Source/GameplayAbilities/Public/GameplayPrediction.h b/Engine/Plugins/Runtime/GameplayAbilities/Source/GameplayAbilities/Public/GameplayPrediction.h
--- a/Engine/Plugins/Runtime/GameplayAbilities/Source/GameplayAbilities/Public/GameplayPrediction.h
+++ b/Engine/Plugins/Runtime/GameplayAbilities/Source/GameplayAbilities/Public/GameplayPrediction.h
@@ -496,8 +496,6 @@ struct GAMEPLAYABILITIES_API FScopedPredictionWindow
 
 	~FScopedPredictionWindow();
 
-	FPredictionKey	ScopedPredictionKey;
-
 private:
 
 	TWeakObjectPtr<UAbilitySystemComponent> Owner;
