Allow both Activate Ability and Activate Ability From Event in the same Gameplay Ability Graph.
It is now safe to call Super::ActivateAbility in a child class (previously, it would call CommitAbility).

diff --git a/Engine/Plugins/Runtime/GameplayAbilities/Source/GameplayAbilities/Private/Abilities/GameplayAbility.cpp b/Engine/Plugins/Runtime/GameplayAbilities/Source/GameplayAbilities/Private/Abilities/GameplayAbility.cpp
--- a/Engine/Plugins/Runtime/GameplayAbilities/Source/GameplayAbilities/Private/Abilities/GameplayAbility.cpp
+++ b/Engine/Plugins/Runtime/GameplayAbilities/Source/GameplayAbilities/Private/Abilities/GameplayAbility.cpp
@@ -468,7 +468,7 @@ bool UGameplayAbility::CommitCheck(const FGameplayAbilitySpecHandle Handle, cons
 	 *		-An ability can start activating, play an animation, wait for a user confirmation/target data, and then actually commit
 	 *		-Commit = spend resources/cooldowns. It's possible the source has changed state since it started activation, so a commit may fail.
 	 *		-We don't want to just call CanActivateAbility() since right now that also checks things like input inhibition.
-	 *			-E.g., its possible the act of starting your ability makes it no longer activatable (CanaCtivateAbility() may be false if called here).
+	 *			-E.g., its possible the act of starting your ability makes it no longer activatable (CanActivateAbility() may be false if called here).
 	 */
 
 	const bool bValidHandle = Handle.IsValid();
@@ -724,36 +724,37 @@ void UGameplayAbility::EndAbility(const FGameplayAbilitySpecHandle Handle, const
 
 void UGameplayAbility::ActivateAbility(const FGameplayAbilitySpecHandle Handle, const FGameplayAbilityActorInfo* ActorInfo, const FGameplayAbilityActivationInfo ActivationInfo, const FGameplayEventData* TriggerEventData)
 {
-	if (bHasBlueprintActivate)
+	if (TriggerEventData && bHasBlueprintActivateFromEvent)
+	{
+		// A Blueprinted ActivateAbility function must call CommitAbility somewhere in its execution chain.
+		K2_ActivateAbilityFromEvent(*TriggerEventData);
+	}
+	else if (bHasBlueprintActivate)
 	{
 		// A Blueprinted ActivateAbility function must call CommitAbility somewhere in its execution chain.
 		K2_ActivateAbility();
 	}
 	else if (bHasBlueprintActivateFromEvent)
 	{
-		if (TriggerEventData)
-		{
-			// A Blueprinted ActivateAbility function must call CommitAbility somewhere in its execution chain.
-			K2_ActivateAbilityFromEvent(*TriggerEventData);
-		}
-		else
-		{
-			UE_LOG(LogAbilitySystem, Warning, TEXT("Ability %s expects event data but none is being supplied. Use Activate Ability instead of Activate Ability From Event."), *GetName());
-			bool bReplicateEndAbility = false;
-			bool bWasCancelled = true;
-			EndAbility(Handle, ActorInfo, ActivationInfo, bReplicateEndAbility, bWasCancelled);
-		}
+		UE_LOG(LogAbilitySystem, Warning, TEXT("Ability %s expects event data but none is being supplied. Use 'Activate Ability' instead of 'Activate Ability From Event' in the Blueprint."), *GetName());
+		constexpr bool bReplicateEndAbility = false;
+		constexpr bool bWasCancelled = true;
+		EndAbility(Handle, ActorInfo, ActivationInfo, bReplicateEndAbility, bWasCancelled);
 	}
 	else
 	{
-		// Native child classes may want to override ActivateAbility and do something like this:
-
-		// Do stuff...
-
-		if (CommitAbility(Handle, ActorInfo, ActivationInfo))		// ..then commit the ability...
-		{			
-			//	Then do more stuff...
-		}
+		// Native child classes should override ActivateAbility and call CommitAbility.
+		// CommitAbility is used to do one last check for spending resources.
+		// Previous versions of this function called CommitAbility but that prevents the callers
+		// from knowing the result. Your override should call it and check the result.
+		// Here is some starter code:
+		
+		//	if (!CommitAbility(Handle, ActorInfo, ActivationInfo))
+		//	{			
+		//		constexpr bool bReplicateEndAbility = true;
+		//		constexpr bool bWasCancelled = true;
+		//		EndAbility(Handle, ActorInfo, ActivationInfo, bReplicateEndAbility, bWasCancelled);
+		//	}
 	}
 }
 
