Add log warning when WaitGameplayTag is used before an ASC is valid, resulting in not registering for gameplay tag count changes.

diff --git a/Engine/Plugins/Runtime/GameplayAbilities/Source/GameplayAbilities/Private/Abilities/Async/AbilityAsync_WaitGameplayTag.cpp b/Engine/Plugins/Runtime/GameplayAbilities/Source/GameplayAbilities/Private/Abilities/Async/AbilityAsync_WaitGameplayTag.cpp
--- a/Engine/Plugins/Runtime/GameplayAbilities/Source/GameplayAbilities/Private/Abilities/Async/AbilityAsync_WaitGameplayTag.cpp
+++ b/Engine/Plugins/Runtime/GameplayAbilities/Source/GameplayAbilities/Private/Abilities/Async/AbilityAsync_WaitGameplayTag.cpp
@@ -1,8 +1,9 @@
 // Copyright Epic Games, Inc. All Rights Reserved.
 
 #include "Abilities/Async/AbilityAsync_WaitGameplayTag.h"
-#include "AbilitySystemGlobals.h"
 #include "AbilitySystemComponent.h"
+#include "AbilitySystemGlobals.h"
+#include "AbilitySystemLog.h"
 
 #include UE_INLINE_GENERATED_CPP_BY_NAME(AbilityAsync_WaitGameplayTag)
 
@@ -27,6 +28,7 @@ void UAbilityAsync_WaitGameplayTag::Activate()
 	}
 	else
 	{
+		ABILITY_LOG(Warning, TEXT("%s: AbilitySystemComponent is nullptr! Could not register for gameplay tag event with Tag = %s."), *GetName(), *Tag.GetTagName().ToString());
 		EndAction();
 	}
 }
