Use existing GameplayCueParameters instead of default Parameters in RemoveGameplayCue_Internal

diff --git a/Engine/Plugins/Runtime/GameplayAbilities/Source/GameplayAbilities/Private/AbilitySystemComponent.cpp b/Engine/Plugins/Runtime/GameplayAbilities/Source/GameplayAbilities/Private/AbilitySystemComponent.cpp
--- a/Engine/Plugins/Runtime/GameplayAbilities/Source/GameplayAbilities/Private/AbilitySystemComponent.cpp
+++ b/Engine/Plugins/Runtime/GameplayAbilities/Source/GameplayAbilities/Private/AbilitySystemComponent.cpp
@@ -1348,11 +1348,14 @@ void UAbilitySystemComponent::RemoveGameplayCue_Internal(const FGameplayTag Game
 
 		if (bWasInList)
 		{
-			FGameplayCueParameters Parameters;
-			InitDefaultGameplayCueParameters(Parameters);
-
-			// Call on server here, clients get it from repnotify
-			InvokeGameplayCueEvent(GameplayCueTag, EGameplayCueEvent::Removed, Parameters);
+			for (const FActiveGameplayCue& ActiveGameplayCue : ActiveGameplayCues.GameplayCues)
+			{
+				if (ActiveGameplayCue.GameplayCueTag == GameplayCueTag)
+				{
+					// Call on server here, clients get it from repnotify
+					InvokeGameplayCueEvent(GameplayCueTag, EGameplayCueEvent::Removed, ActiveGameplayCue.Parameters);
+				}
+			}
 		}
 		// Don't need to multicast broadcast this, ActiveGameplayCues replication handles it
 	}
