Use existing GameplayCueParameters instead of default Parameters in RemoveGameplayCue_Internal

diff --git a/Engine/Plugins/Runtime/GameplayAbilities/Source/GameplayAbilities/Private/AbilitySystemComponent.cpp b/Engine/Plugins/Runtime/GameplayAbilities/Source/GameplayAbilities/Private/AbilitySystemComponent.cpp
--- a/Engine/Plugins/Runtime/GameplayAbilities/Source/GameplayAbilities/Private/AbilitySystemComponent.cpp
+++ b/Engine/Plugins/Runtime/GameplayAbilities/Source/GameplayAbilities/Private/AbilitySystemComponent.cpp
@@ -1348,11 +1348,14 @@ void UAbilitySystemComponent::RemoveGameplayCue_Internal(const FGameplayTag Game
 
 		if (bWasInList)
 		{
-			FGameplayCueParameters Parameters;
-			InitDefaultGameplayCueParameters(Parameters);
-
-			// Call on server here, clients get it from repnotify
-			InvokeGameplayCueEvent(GameplayCueTag, EGameplayCueEvent::Removed, Parameters);
+			for (const FActiveGameplayCue& ActiveGameplayCue : ActiveGameplayCues.GameplayCues)
+			{
+				if (ActiveGameplayCue.GameplayCueTag == GameplayCueTag)
+				{
+					// Call on server here, clients get it from repnotify
+					InvokeGameplayCueEvent(GameplayCueTag, EGameplayCueEvent::Removed, ActiveGameplayCue.Parameters);
+				}
+			}
 		}
 		// Don't need to multicast broadcast this, ActiveGameplayCues replication handles it
 	}
diff --git a/Engine/Plugins/Runtime/GameplayAbilities/Source/GameplayAbilities/Private/GameplayCueInterface.cpp b/Engine/Plugins/Runtime/GameplayAbilities/Source/GameplayAbilities/Private/GameplayCueInterface.cpp
--- a/Engine/Plugins/Runtime/GameplayAbilities/Source/GameplayAbilities/Private/GameplayCueInterface.cpp
+++ b/Engine/Plugins/Runtime/GameplayAbilities/Source/GameplayAbilities/Private/GameplayCueInterface.cpp
@@ -243,19 +243,25 @@ void FActiveGameplayCueContainer::RemoveCue(const FGameplayTag& Tag)
 		return;
 	}
 
-	for (int32 idx=0; idx < GameplayCues.Num(); ++idx)
+	
+	// Iterate backwards so we can remove during loop
+	int32 CountDelta = 0;
+	for (int32 idx=GameplayCues.Num()-1; idx >= 0; --idx)
 	{
 		FActiveGameplayCue& Cue = GameplayCues[idx];
 
 		if (Cue.GameplayCueTag == Tag)
 		{
 			GameplayCues.RemoveAt(idx);
-			MarkArrayDirty();
-			Owner->UpdateTagMap(Tag, -1);
-			return;
+			CountDelta -= 1;
 		}
 	}
 
+	if (CountDelta < 0)
+	{
+		MarkArrayDirty();
+		Owner->UpdateTagMap(Tag, CountDelta);
+	}
 }
 
 void FActiveGameplayCueContainer::RemoveAllCues()
