ImmunityGameplayEffectComponent was not building its immunity query properly.

diff --git a/Engine/Plugins/Runtime/GameplayAbilities/Source/GameplayAbilities/Private/GameplayEffect.cpp b/Engine/Plugins/Runtime/GameplayAbilities/Source/GameplayAbilities/Private/GameplayEffect.cpp
--- a/Engine/Plugins/Runtime/GameplayAbilities/Source/GameplayAbilities/Private/GameplayEffect.cpp
+++ b/Engine/Plugins/Runtime/GameplayAbilities/Source/GameplayAbilities/Private/GameplayEffect.cpp
@@ -489,11 +489,32 @@ void UGameplayEffect::ConvertImmunityComponent()
 		{
 			// Build up a query that matches the equivalent GrantedApplicationImmunityTags
 			FGameplayEffectQuery TagQuery;
-			TagQuery.SourceTagQuery.BuildQuery(
-				FGameplayTagQueryExpression().AnyTagsMatch().AddTags(GrantedApplicationImmunityTags.RequireTags).NoTagsMatch().AddTags(GrantedApplicationImmunityTags.IgnoreTags),
-				TEXT("GrantedApplicationImmunityTags")
-			);
 
+			const bool bHasRequireTags = !GrantedApplicationImmunityTags.RequireTags.IsEmpty();
+			const bool bHasIgnoreTags = !GrantedApplicationImmunityTags.IgnoreTags.IsEmpty();
+			if (bHasRequireTags || bHasIgnoreTags)
+			{
+				// Previously the Tags query was GrantedApplicationImmunityTags.RequirementsMet. FGameplayTagContainer::RequirementsMet is HasAll(RequireTags) && !HasAny(IgnoreTags);
+				FGameplayTagQueryExpression RequiredTagsQueryExpression = FGameplayTagQueryExpression().AllTagsMatch().AddTags(GrantedApplicationImmunityTags.RequireTags);
+				FGameplayTagQueryExpression IgnoreTagsQueryExpression = FGameplayTagQueryExpression().NoTagsMatch().AddTags(GrantedApplicationImmunityTags.IgnoreTags);
+
+				FGameplayTagQueryExpression RootQueryExpression;
+				if (bHasRequireTags && bHasIgnoreTags)
+				{
+					RootQueryExpression = FGameplayTagQueryExpression().AllExprMatch().AddExpr(RequiredTagsQueryExpression).AddExpr(IgnoreTagsQueryExpression);
+				}
+				else if (bHasRequireTags)
+				{
+					RootQueryExpression = RequiredTagsQueryExpression;
+				}
+				else // bHasIgnoreTags
+				{
+					RootQueryExpression = IgnoreTagsQueryExpression;
+				}
+
+				// Build the expression
+				TagQuery.SourceTagQuery.Build(RootQueryExpression,TEXT("GrantedApplicationImmunityTags"));
+			}
 			ImmunityComponent.ImmunityQueries[0] = MoveTemp(TagQuery);
 		}
 
