Undo //UE5/Main/Engine/Plugins/Runtime/GameplayAbilities/Source/GameplayAbilities/... changelist 24078469

Reason: had fallout in Lyra, see jira UE-181130, and possibly other projects. After discussion, rolling back until we reach a consensus on how to handle this setting in all scenarios.

diff --git a/Engine/Plugins/Runtime/GameplayAbilities/Source/GameplayAbilities/Private/GameplayEffect.cpp b/Engine/Plugins/Runtime/GameplayAbilities/Source/GameplayAbilities/Private/GameplayEffect.cpp
--- a/Engine/Plugins/Runtime/GameplayAbilities/Source/GameplayAbilities/Private/GameplayEffect.cpp
+++ b/Engine/Plugins/Runtime/GameplayAbilities/Source/GameplayAbilities/Private/GameplayEffect.cpp
@@ -2511,14 +2511,14 @@ void FActiveGameplayEffectsContainer::PredictivelyExecuteEffectSpec(FGameplayEff
 	//		These will modify the base value of attributes
 	// ------------------------------------------------------
 
-	bool bAnyModifierSuccessfullyExecuted = false;
+	bool ModifierSuccessfullyExecuted = false;
 
 	for (int32 ModIdx = 0; ModIdx < SpecToUse.Modifiers.Num(); ++ModIdx)
 	{
 		const FGameplayModifierInfo& ModDef = SpecToUse.Def->Modifiers[ModIdx];
 
 		FGameplayModifierEvaluatedData EvalData(ModDef.Attribute, ModDef.ModifierOp, SpecToUse.GetModifierMagnitude(ModIdx, true));
-		bAnyModifierSuccessfullyExecuted |= InternalExecuteMod(SpecToUse, EvalData);
+		ModifierSuccessfullyExecuted |= InternalExecuteMod(SpecToUse, EvalData);
 	}
 
 	// ------------------------------------------------------
@@ -2526,8 +2526,7 @@ void FActiveGameplayEffectsContainer::PredictivelyExecuteEffectSpec(FGameplayEff
 	//		This will run custom code to 'do stuff'
 	// ------------------------------------------------------
 
-	bool bGameplayCuesWereManuallyHandled = false;
-	bool bExecutionsProducedOutModifiers = false;
+	bool GameplayCuesWereManuallyHandled = false;
 
 	for (const FGameplayEffectExecutionDefinition& CurExecDef : SpecToUse.Def->Executions)
 	{
@@ -2543,7 +2542,6 @@ void FActiveGameplayEffectsContainer::PredictivelyExecuteEffectSpec(FGameplayEff
 
 			// Execute any mods the custom execution yielded
 			TArray<FGameplayModifierEvaluatedData>& OutModifiers = ExecutionOutput.GetOutputModifiersRef();
-			bExecutionsProducedOutModifiers |= OutModifiers.Num() > 0;
 
 			const bool bApplyStackCountToEmittedMods = !ExecutionOutput.IsStackCountHandledManually();
 			const int32 SpecStackCount = SpecToUse.GetStackCount();
@@ -2555,13 +2553,13 @@ void FActiveGameplayEffectsContainer::PredictivelyExecuteEffectSpec(FGameplayEff
 				{
 					CurExecMod.Magnitude = GameplayEffectUtilities::ComputeStackedModifierMagnitude(CurExecMod.Magnitude, SpecStackCount, CurExecMod.ModifierOp);
 				}
-				bAnyModifierSuccessfullyExecuted |= InternalExecuteMod(SpecToUse, CurExecMod);
+				ModifierSuccessfullyExecuted |= InternalExecuteMod(SpecToUse, CurExecMod);
 			}
 
 			// If execution handled GameplayCues, we dont have to.
 			if (ExecutionOutput.AreGameplayCuesHandledManually())
 			{
-				bGameplayCuesWereManuallyHandled = true;
+				GameplayCuesWereManuallyHandled = true;
 			}
 		}
 	}
@@ -2572,18 +2570,25 @@ void FActiveGameplayEffectsContainer::PredictivelyExecuteEffectSpec(FGameplayEff
 	if (bPredictGameplayCues)
 	{
 		// If there are no modifiers or we don't require modifier success to trigger, we apply the GameplayCue.
-		const bool bHasModifiersFromData = SpecToUse.Modifiers.Num() > 0;
+		const bool bHasModifiers = SpecToUse.Modifiers.Num() > 0;
+		const bool bHasExecutions = SpecToUse.Def->Executions.Num() > 0;
+		const bool bHasModifiersOrExecutions = bHasModifiers || bHasExecutions;
 
-		// Apply the GameplayCue if:
-		// - No execution handled it manually AND
-		// -- It didn't require a successful modifier OR there was a successful modifier
-		// -- Or it didn't have any modifiers at all
-		const bool bInvokeGameplayCueExecute = !bGameplayCuesWereManuallyHandled && (
-			(!Spec.Def->bRequireModifierSuccessToTriggerCues || bAnyModifierSuccessfullyExecuted) ||
-			(!bHasModifiersFromData && !bExecutionsProducedOutModifiers)
-		);
+		// If there are no modifiers or we don't require modifier success to trigger, we apply the GameplayCue.
+		bool InvokeGameplayCueExecute = (!bHasModifiersOrExecutions) || !Spec.Def->bRequireModifierSuccessToTriggerCues;
+
+		if (bHasModifiersOrExecutions && ModifierSuccessfullyExecuted)
+		{
+			InvokeGameplayCueExecute = true;
+		}
 
-		if (bInvokeGameplayCueExecute && SpecToUse.Def->GameplayCues.Num())
+		// Don't trigger gameplay cues if one of the executions says it manually handled them
+		if (GameplayCuesWereManuallyHandled)
+		{
+			InvokeGameplayCueExecute = false;
+		}
+
+		if (InvokeGameplayCueExecute && SpecToUse.Def->GameplayCues.Num())
 		{
 			// TODO: check replication policy. Right now we will replicate every execute via a multicast RPC
 
@@ -2621,14 +2626,14 @@ void FActiveGameplayEffectsContainer::ExecuteActiveEffectsFrom(FGameplayEffectSp
 	//		These will modify the base value of attributes
 	// ------------------------------------------------------
 	
-	bool bAnyModifierSuccessfullyExecuted = false;
+	bool ModifierSuccessfullyExecuted = false;
 
 	for (int32 ModIdx = 0; ModIdx < SpecToUse.Modifiers.Num(); ++ModIdx)
 	{
 		const FGameplayModifierInfo& ModDef = SpecToUse.Def->Modifiers[ModIdx];
 		
 		FGameplayModifierEvaluatedData EvalData(ModDef.Attribute, ModDef.ModifierOp, SpecToUse.GetModifierMagnitude(ModIdx, true));
-		bAnyModifierSuccessfullyExecuted |= InternalExecuteMod(SpecToUse, EvalData);
+		ModifierSuccessfullyExecuted |= InternalExecuteMod(SpecToUse, EvalData);
 	}
 
 	// ------------------------------------------------------
@@ -2638,8 +2643,7 @@ void FActiveGameplayEffectsContainer::ExecuteActiveEffectsFrom(FGameplayEffectSp
 	
 	TArray< FGameplayEffectSpecHandle, TInlineAllocator<4> > ConditionalEffectSpecs;
 
-	bool bGameplayCuesWereManuallyHandled = false;
-	bool bExecutionsProducedOutModifiers = false;
+	bool GameplayCuesWereManuallyHandled = false;
 
 	for (const FGameplayEffectExecutionDefinition& CurExecDef : SpecToUse.Def->Executions)
 	{
@@ -2659,7 +2663,6 @@ void FActiveGameplayEffectsContainer::ExecuteActiveEffectsFrom(FGameplayEffectSp
 
 			// Execute any mods the custom execution yielded
 			TArray<FGameplayModifierEvaluatedData>& OutModifiers = ExecutionOutput.GetOutputModifiersRef();
-			bExecutionsProducedOutModifiers |= OutModifiers.Num() > 0;
 
 			const bool bApplyStackCountToEmittedMods = !ExecutionOutput.IsStackCountHandledManually();
 			const int32 SpecStackCount = SpecToUse.GetStackCount();
@@ -2671,13 +2674,13 @@ void FActiveGameplayEffectsContainer::ExecuteActiveEffectsFrom(FGameplayEffectSp
 				{
 					CurExecMod.Magnitude = GameplayEffectUtilities::ComputeStackedModifierMagnitude(CurExecMod.Magnitude, SpecStackCount, CurExecMod.ModifierOp);
 				}
-				bAnyModifierSuccessfullyExecuted |= InternalExecuteMod(SpecToUse, CurExecMod);
+				ModifierSuccessfullyExecuted |= InternalExecuteMod(SpecToUse, CurExecMod);
 			}
 
 			// If execution handled GameplayCues, we dont have to.
 			if (ExecutionOutput.AreGameplayCuesHandledManually())
 			{
-				bGameplayCuesWereManuallyHandled = true;
+				GameplayCuesWereManuallyHandled = true;
 			}
 		}
 
@@ -2703,18 +2706,25 @@ void FActiveGameplayEffectsContainer::ExecuteActiveEffectsFrom(FGameplayEffectSp
 	// ------------------------------------------------------
 	
 	// If there are no modifiers or we don't require modifier success to trigger, we apply the GameplayCue.
-	const bool bHasModifiersFromData = SpecToUse.Modifiers.Num() > 0;
-
-	// Apply the GameplayCue if:
-	// - No execution handled it manually AND
-	// -- It didn't require a successful modifier OR there was a successful modifier
-	// -- Or it didn't have any modifiers at all
-	const bool bInvokeGameplayCueExecute = !bGameplayCuesWereManuallyHandled && (
-		(!Spec.Def->bRequireModifierSuccessToTriggerCues || bAnyModifierSuccessfullyExecuted) ||
-		(!bHasModifiersFromData && !bExecutionsProducedOutModifiers)
-	);
+	const bool bHasModifiers = SpecToUse.Modifiers.Num() > 0;
+	const bool bHasExecutions = SpecToUse.Def->Executions.Num() > 0;
+	const bool bHasModifiersOrExecutions = bHasModifiers || bHasExecutions;
+
+	// If there are no modifiers or we don't require modifier success to trigger, we apply the GameplayCue.
+	bool InvokeGameplayCueExecute = (!bHasModifiersOrExecutions) || !Spec.Def->bRequireModifierSuccessToTriggerCues;
+
+	if (bHasModifiersOrExecutions && ModifierSuccessfullyExecuted)
+	{
+		InvokeGameplayCueExecute = true;
+	}
+
+	// Don't trigger gameplay cues if one of the executions says it manually handled them
+	if (GameplayCuesWereManuallyHandled)
+	{
+		InvokeGameplayCueExecute = false;
+	}
 
-	if (bInvokeGameplayCueExecute && SpecToUse.Def->GameplayCues.Num())
+	if (InvokeGameplayCueExecute && SpecToUse.Def->GameplayCues.Num())
 	{
 		// TODO: check replication policy. Right now we will replicate every execute via a multicast RPC
 
diff --git a/Engine/Plugins/Runtime/GameplayAbilities/Source/GameplayAbilities/Public/GameplayEffect.h b/Engine/Plugins/Runtime/GameplayAbilities/Source/GameplayAbilities/Public/GameplayEffect.h
--- a/Engine/Plugins/Runtime/GameplayAbilities/Source/GameplayAbilities/Public/GameplayEffect.h
+++ b/Engine/Plugins/Runtime/GameplayAbilities/Source/GameplayAbilities/Public/GameplayEffect.h
@@ -2228,7 +2228,7 @@ public:
 	UPROPERTY(BlueprintReadOnly, Category="Deprecated|Expiration", meta = (DeprecatedProperty))
 	TArray<TSubclassOf<UGameplayEffect>> RoutineExpirationEffectClasses;
 
-	/** If true, cues will only trigger when any GE modifiers succeed being applied (whether through modifiers or executions), or if there were no modifiers application attempts. */
+	/** If true, cues will only trigger when GE modifiers succeed being applied (whether through modifiers or executions) */
 	UPROPERTY(EditDefaultsOnly, BlueprintReadOnly, Category = "GameplayCues")
 	bool bRequireModifierSuccessToTriggerCues;
 
