UAbilitySystemComponent::RemoveGameplayCue_Internal should be iterating over passed in GameplayCueContainer, not ActiveGameplayCues and should also iterate before removing the cues from the container

diff --git a/Engine/Plugins/Runtime/GameplayAbilities/Source/GameplayAbilities/Private/AbilitySystemComponent.cpp b/Engine/Plugins/Runtime/GameplayAbilities/Source/GameplayAbilities/Private/AbilitySystemComponent.cpp
--- a/Engine/Plugins/Runtime/GameplayAbilities/Source/GameplayAbilities/Private/AbilitySystemComponent.cpp
+++ b/Engine/Plugins/Runtime/GameplayAbilities/Source/GameplayAbilities/Private/AbilitySystemComponent.cpp
@@ -1344,19 +1344,21 @@ void UAbilitySystemComponent::RemoveGameplayCue_Internal(const FGameplayTag Game
 
 		// Force replication so GameplayCue removals are properly replicated to all clients during Mixed and Minimal replication modes
 		ForceReplication();
-		GameplayCueContainer.RemoveCue(GameplayCueTag);
 
 		if (bWasInList)
 		{
-			for (const FActiveGameplayCue& ActiveGameplayCue : ActiveGameplayCues.GameplayCues)
+			for (const FActiveGameplayCue& GameplayCue : GameplayCueContainer.GameplayCues)
 			{
-				if (ActiveGameplayCue.GameplayCueTag == GameplayCueTag)
+				if (GameplayCue.GameplayCueTag == GameplayCueTag)
 				{
 					// Call on server here, clients get it from repnotify
-					InvokeGameplayCueEvent(GameplayCueTag, EGameplayCueEvent::Removed, ActiveGameplayCue.Parameters);
+					InvokeGameplayCueEvent(GameplayCueTag, EGameplayCueEvent::Removed, GameplayCue.Parameters);
 				}
 			}
 		}
+		
+		GameplayCueContainer.RemoveCue(GameplayCueTag);
+		
 		// Don't need to multicast broadcast this, ActiveGameplayCues replication handles it
 	}
 	else if (ScopedPredictionKey.IsLocalClientKey())
