Removing an ensure in ImmunityGameplayEffectComponent because complex interactions can cause this callback while 'removed' and still valid.

diff --git a/Engine/Plugins/Runtime/GameplayAbilities/Source/GameplayAbilities/Private/GameplayEffectComponents/ImmunityGameplayEffectComponent.cpp b/Engine/Plugins/Runtime/GameplayAbilities/Source/GameplayAbilities/Private/GameplayEffectComponents/ImmunityGameplayEffectComponent.cpp
--- a/Engine/Plugins/Runtime/GameplayAbilities/Source/GameplayAbilities/Private/GameplayEffectComponents/ImmunityGameplayEffectComponent.cpp
+++ b/Engine/Plugins/Runtime/GameplayAbilities/Source/GameplayAbilities/Private/GameplayEffectComponents/ImmunityGameplayEffectComponent.cpp
@@ -54,8 +54,9 @@ bool UImmunityGameplayEffectComponent::AllowGameplayEffectApplication(const FAct
 		return false;
 	}
 
+	// ActiveGE can be null (IsPendingRemove) but we could have not yet received OnEffectRemoved. This can happen when applying more effects during its removal.
 	const FActiveGameplayEffect* ActiveGE = ASC->GetActiveGameplayEffect(ImmunityActiveGEHandle);
-	if (!ensure(ActiveGE) || ActiveGE->bIsInhibited)
+	if (!ActiveGE || ActiveGE->bIsInhibited)
 	{
 		return true;
 	}
