Fix spelling error
#rnx

diff --git a/Engine/Plugins/Runtime/GameplayAbilities/Source/GameplayAbilities/Private/GameplayEffectComponents/AdditionalEffectsGameplayEffectComponent.cpp b/Engine/Plugins/Runtime/GameplayAbilities/Source/GameplayAbilities/Private/GameplayEffectComponents/AdditionalEffectsGameplayEffectComponent.cpp
--- a/Engine/Plugins/Runtime/GameplayAbilities/Source/GameplayAbilities/Private/GameplayEffectComponents/AdditionalEffectsGameplayEffectComponent.cpp
+++ b/Engine/Plugins/Runtime/GameplayAbilities/Source/GameplayAbilities/Private/GameplayEffectComponents/AdditionalEffectsGameplayEffectComponent.cpp
@@ -133,7 +133,7 @@ EDataValidationResult UAdditionalEffectsGameplayEffectComponent::IsDataValid(FDa
 	if (GetOwner()->DurationPolicy == EGameplayEffectDurationType::Instant)
 	{
 		const bool bHasOnCompleteEffects = (OnCompleteAlways.Num() + OnCompleteNormal.Num() + OnCompletePrematurely.Num() > 0);
-		Context.AddError(FText::FormatOrdered(LOCTEXT("InstantDoesNotWorkWithOnComplete", "Instant GE will never received OnComplete for {0}."), FText::FromString(GetClass()->GetName())));
+		Context.AddError(FText::FormatOrdered(LOCTEXT("InstantDoesNotWorkWithOnComplete", "Instant GE will never receive OnComplete for {0}."), FText::FromString(GetClass()->GetName())));
 		Result = EDataValidationResult::Invalid;
 	}
 	else if (GetOwner()->Period.Value > 0.0f)
