Catching the case where I removed a deprecated IsDataValid function without deprecating it.

diff --git a/Engine/Plugins/Runtime/GameplayAbilities/Source/GameplayAbilities/Public/GameplayEffectTypes.h b/Engine/Plugins/Runtime/GameplayAbilities/Source/GameplayAbilities/Public/GameplayEffectTypes.h
--- a/Engine/Plugins/Runtime/GameplayAbilities/Source/GameplayAbilities/Public/GameplayEffectTypes.h
+++ b/Engine/Plugins/Runtime/GameplayAbilities/Source/GameplayAbilities/Public/GameplayEffectTypes.h
@@ -1328,6 +1328,9 @@ public:
 	void ApplyCurrentTags();
 
 #if WITH_EDITOR
+	UE_DEPRECATED(5.3, "The signature for IsDataValid has changed.  Please use the one that takes a FDataValidationContext")
+	EDataValidationResult IsDataValid(const UObject* ContainingAsset, TArray<FText>& ValidationErrors) { return EDataValidationResult::NotValidated; }
+
 	/** This can optionally be called in the owner's IsDataValid() for data validation. */
 	EDataValidationResult IsDataValid(const UObject* ContainingAsset, class FDataValidationContext& Context) const;
 #endif // #if WITH_EDITOR
