Fixing an issue where the GlobalActiveGameplayEffectHandles map can be reset by a client PIE session joining a game.
This caused really hard to track down bugs where GameplayEffectComponents could fail to get the ActiveGE's AbilitySystemComponent (because the handle map was reset by the client).

diff --git a/Engine/Plugins/Runtime/GameplayAbilities/Source/GameplayAbilities/Private/AbilitySystemGlobals.cpp b/Engine/Plugins/Runtime/GameplayAbilities/Source/GameplayAbilities/Private/AbilitySystemGlobals.cpp
--- a/Engine/Plugins/Runtime/GameplayAbilities/Source/GameplayAbilities/Private/AbilitySystemGlobals.cpp
+++ b/Engine/Plugins/Runtime/GameplayAbilities/Source/GameplayAbilities/Private/AbilitySystemGlobals.cpp
@@ -67,7 +67,7 @@ void UAbilitySystemGlobals::InitGlobalData()
 	InitTargetDataScriptStructCache();
 
 	// Register for PreloadMap so cleanup can occur on map transitions
-	FCoreUObjectDelegates::PreLoadMap.AddUObject(this, &UAbilitySystemGlobals::HandlePreLoadMap);
+	FCoreUObjectDelegates::PreLoadMapWithContext.AddUObject(this, &UAbilitySystemGlobals::HandlePreLoadMap);
 
 #if WITH_EDITOR
 	// Register in editor for PreBeginPlay so cleanup can occur when we start a PIE session
@@ -607,8 +607,14 @@ void UAbilitySystemGlobals::ResetCachedData()
 	FActiveGameplayEffectHandle::ResetGlobalHandleMap();
 }
 
-void UAbilitySystemGlobals::HandlePreLoadMap(const FString& MapName)
+void UAbilitySystemGlobals::HandlePreLoadMap(const FWorldContext& WorldContext, const FString& MapName)
 {
+	// We don't want to reset for PIE since this is shared memory (which would have received OnPreBeginPIE).
+	if (WorldContext.PIEInstance > 0)
+	{
+		return;
+	}
+
 	ResetCachedData();
 }
 
diff --git a/Engine/Plugins/Runtime/GameplayAbilities/Source/GameplayAbilities/Public/AbilitySystemGlobals.h b/Engine/Plugins/Runtime/GameplayAbilities/Source/GameplayAbilities/Public/AbilitySystemGlobals.h
--- a/Engine/Plugins/Runtime/GameplayAbilities/Source/GameplayAbilities/Public/AbilitySystemGlobals.h
+++ b/Engine/Plugins/Runtime/GameplayAbilities/Source/GameplayAbilities/Public/AbilitySystemGlobals.h
@@ -414,7 +414,7 @@ protected:
 #endif
 
 	void ResetCachedData();
-	void HandlePreLoadMap(const FString& MapName);
+	void HandlePreLoadMap(const FWorldContext& WorldContext, const FString& MapName);
 
 #if WITH_EDITORONLY_DATA
 	bool RegisteredReimportCallback;
