Fixing non-Unity build

diff --git a/Engine/Plugins/Runtime/GameplayAbilities/Source/GameplayAbilities/Private/AbilitySystemGlobals.cpp b/Engine/Plugins/Runtime/GameplayAbilities/Source/GameplayAbilities/Private/AbilitySystemGlobals.cpp
--- a/Engine/Plugins/Runtime/GameplayAbilities/Source/GameplayAbilities/Private/AbilitySystemGlobals.cpp
+++ b/Engine/Plugins/Runtime/GameplayAbilities/Source/GameplayAbilities/Private/AbilitySystemGlobals.cpp
@@ -13,6 +13,7 @@
 #include "GameplayCueManager.h"
 #include "GameplayTagResponseTable.h"
 #include "GameplayTagsManager.h"
+#include "Engine/Engine.h"
 #include "UObject/UObjectIterator.h"
 
 #if UE_WITH_IRIS
