GAS: Move ShouldAbilityRespondToEvent from client-only codepath to both server and client.

diff --git a/Engine/Plugins/Runtime/GameplayAbilities/Source/GameplayAbilities/Private/AbilitySystemComponent_Abilities.cpp b/Engine/Plugins/Runtime/GameplayAbilities/Source/GameplayAbilities/Private/AbilitySystemComponent_Abilities.cpp
--- a/Engine/Plugins/Runtime/GameplayAbilities/Source/GameplayAbilities/Private/AbilitySystemComponent_Abilities.cpp
+++ b/Engine/Plugins/Runtime/GameplayAbilities/Source/GameplayAbilities/Private/AbilitySystemComponent_Abilities.cpp
@@ -1543,18 +1543,22 @@ bool UAbilitySystemComponent::InternalTryActivateAbility(FGameplayAbilitySpecHan
 		return false;
 	}
 
+	if (TriggerEventData)
+	{
+		if (!Ability->ShouldAbilityRespondToEvent(ActorInfo, TriggerEventData))
+		{
+			NotifyAbilityFailed(Handle, Ability, InternalTryActivateAbilityFailureTags);
+			return false;
+		}
+	}
+
 	// If it's instance once the instanced ability will be set, otherwise it will be null
 	UGameplayAbility* InstancedAbility = Spec->GetPrimaryInstance();
 
-	const FGameplayTagContainer* SourceTags = nullptr;
-	const FGameplayTagContainer* TargetTags = nullptr;
-	if (TriggerEventData != nullptr)
 	{
-		SourceTags = &TriggerEventData->InstigatorTags;
-		TargetTags = &TriggerEventData->TargetTags;
-	}
+		const FGameplayTagContainer* SourceTags = TriggerEventData ? &TriggerEventData->InstigatorTags : nullptr;
+		const FGameplayTagContainer* TargetTags = TriggerEventData ? &TriggerEventData->TargetTags : nullptr;
 
-	{
 		// If we have an instanced ability, call CanActivateAbility on it.
 		// Otherwise we always do a non instanced CanActivateAbility check using the CDO of the Ability.
 		UGameplayAbility* const CanActivateAbilitySource = InstancedAbility ? InstancedAbility : Ability;
@@ -2214,14 +2218,7 @@ bool UAbilitySystemComponent::TriggerAbilityFromGameplayEvent(FGameplayAbilitySp
 	TempEventData.EventTag = EventTag;
 
 	// Run on the non-instanced ability
-	if (Ability->ShouldAbilityRespondToEvent(ActorInfo, &TempEventData))
-	{
-		if (InternalTryActivateAbility(Handle, ScopedPredictionKey, nullptr, nullptr, &TempEventData))
-		{
-			return true;
-		}
-	}
-	return false;
+	return InternalTryActivateAbility(Handle, ScopedPredictionKey, nullptr, nullptr, &TempEventData);
 }
 
 bool UAbilitySystemComponent::GetUserAbilityActivationInhibited() const
