Adding/extending the functonality to execute and cancel Gameplay Abilities & Gameplay Effects from a console command.

diff --git a/Engine/Plugins/Runtime/GameplayAbilities/Source/GameplayAbilities/Private/AbilitySystemGlobals.cpp b/Engine/Plugins/Runtime/GameplayAbilities/Source/GameplayAbilities/Private/AbilitySystemGlobals.cpp
--- a/Engine/Plugins/Runtime/GameplayAbilities/Source/GameplayAbilities/Private/AbilitySystemGlobals.cpp
+++ b/Engine/Plugins/Runtime/GameplayAbilities/Source/GameplayAbilities/Private/AbilitySystemGlobals.cpp
@@ -3,6 +3,7 @@
 #include "AbilitySystemGlobals.h"
 #include "Abilities/GameplayAbilityTypes.h"
 #include "AbilitySystemStats.h"
+#include "Engine/Blueprint.h"
 #include "GameFramework/Pawn.h"
 #include "GameplayCueInterface.h"
 #include "AbilitySystemComponent.h"
@@ -28,6 +29,18 @@
 
 #include UE_INLINE_GENERATED_CPP_BY_NAME(AbilitySystemGlobals)
 
+namespace UE::AbilitySystemGlobals
+{
+	bool bIgnoreAbilitySystemCooldowns = false;
+	bool bIgnoreAbilitySystemCosts = false;
+
+	FAutoConsoleVariableRef CVarAbilitySystemIgnoreCooldowns(TEXT("AbilitySystem.IgnoreCooldowns"), bIgnoreAbilitySystemCooldowns, TEXT("Ignore cooldowns for all Gameplay Abilities."), ECVF_Cheat);
+	FAutoConsoleVariableRef CVarAbilitySystemIgnoreCosts(TEXT("AbilitySystem.IgnoreCosts"), bIgnoreAbilitySystemCosts, TEXT("Ignore costs for all Gameplay Abilities."), ECVF_Cheat);
+
+	template<typename ClassToFind>
+	TSubclassOf<ClassToFind> FuzzyFindClass(FString SearchString);
+}
+
 UAbilitySystemGlobals::UAbilitySystemGlobals(const FObjectInitializer& ObjectInitializer)
 : Super(ObjectInitializer)
 {
@@ -46,11 +59,6 @@ UAbilitySystemGlobals::UAbilitySystemGlobals(const FObjectInitializer& ObjectIni
 #if WITH_EDITORONLY_DATA
 	RegisteredReimportCallback = false;
 #endif // #if WITH_EDITORONLY_DATA
-
-#if !(UE_BUILD_SHIPPING || UE_BUILD_TEST)
-	bIgnoreAbilitySystemCooldowns = false;
-	bIgnoreAbilitySystemCosts = false;
-#endif // #if !(UE_BUILD_SHIPPING || UE_BUILD_TEST)
 }
 
 void UAbilitySystemGlobals::InitGlobalData()
@@ -471,128 +479,14 @@ void UAbilitySystemGlobals::GlobalPreGameplayEffectSpecApply(FGameplayEffectSpec
 
 }
 
-void UAbilitySystemGlobals::ToggleIgnoreAbilitySystemCooldowns()
-{
-#if WITH_ABILITY_CHEATS
-	bIgnoreAbilitySystemCooldowns = !bIgnoreAbilitySystemCooldowns;
-#endif // WITH_ABILITY_CHEATS
-}
-
-void UAbilitySystemGlobals::ToggleIgnoreAbilitySystemCosts()
-{
-#if WITH_ABILITY_CHEATS
-	bIgnoreAbilitySystemCosts = !bIgnoreAbilitySystemCosts;
-#endif // WITH_ABILITY_CHEATS
-}
-
 bool UAbilitySystemGlobals::ShouldIgnoreCooldowns() const
 {
-#if WITH_ABILITY_CHEATS
-	return bIgnoreAbilitySystemCooldowns;
-#else
-	return false;
-#endif // WITH_ABILITY_CHEATS
+	return UE::AbilitySystemGlobals::bIgnoreAbilitySystemCooldowns;
 }
 
 bool UAbilitySystemGlobals::ShouldIgnoreCosts() const
 {
-#if WITH_ABILITY_CHEATS
-	return bIgnoreAbilitySystemCosts;
-#else
-	return false;
-#endif // WITH_ABILITY_CHEATS
-}
-
-void UAbilitySystemGlobals::ListPlayerAbilities()
-{
-#if WITH_ABILITY_CHEATS	
-	APlayerController* PC = GWorld->GetFirstPlayerController();
-	UAbilitySystemComponent* AbilityComponent = PC ? GetAbilitySystemComponentFromActor(PC->GetPawn()) : nullptr;
-	if(AbilityComponent)
-	{
-		const UEnum* ExecutionEnumPtr = FindObject<UEnum>(nullptr, TEXT("/Script/GameplayAbilities.EGameplayAbilityNetExecutionPolicy"), true);
-		check(ExecutionEnumPtr && TEXT("Couldn't locate EGameplayAbilityNetExecutionPolicy enum!"));
-		const UEnum* SecurityEnumPtr = FindObject<UEnum>(nullptr, TEXT("/Script/GameplayAbilities.EGameplayAbilityNetSecurityPolicy"), true);
-		check(SecurityEnumPtr && TEXT("Couldn't locate EGameplayAbilityNetSecurityPolicy enum!"));
-
-		PC->ClientMessage(TEXT("Available abilities:"));
-
-		check(AbilityComponent && TEXT("Failed to find ability component on player pawn."));
-		for (FGameplayAbilitySpec &Activatable : AbilityComponent->GetActivatableAbilities())
-		{
-			PC->ClientMessage(FString::Printf(TEXT("   %s (%s - %s)"), *Activatable.Ability->GetName(), *ExecutionEnumPtr->GetDisplayNameTextByIndex(Activatable.Ability->GetNetExecutionPolicy()).ToString(), *SecurityEnumPtr->GetDisplayNameTextByIndex(Activatable.Ability->GetNetSecurityPolicy()).ToString()));
-		}
-	}
-#endif // WITH_ABILITY_CHEATS
-}
-
-void UAbilitySystemGlobals::ServerActivatePlayerAbility(FString AbilityNameMatch)
-{
-#if WITH_ABILITY_CHEATS
-	if (AbilityNameMatch.IsEmpty())
-	{
-		ListPlayerAbilities();
-		return;
-	}
-
-	APlayerController* PC = GWorld->GetFirstPlayerController();
-	UAbilitySystemComponent* AbilityComponent = PC ? GetAbilitySystemComponentFromActor(PC->GetPawn()) : nullptr;
-	if(AbilityComponent)
-	{
-		for (FGameplayAbilitySpec &Activatable : AbilityComponent->GetActivatableAbilities())
-		{
-			// Trigger on first match only
-			if (Activatable.Ability->GetName().Contains(AbilityNameMatch))
-			{
-				PC->ClientMessage(FString::Printf(TEXT("Triggering forced server activation of %s"), *Activatable.Ability->GetName()));
-				FPredictionKey MyKey = FPredictionKey::CreateNewPredictionKey(AbilityComponent);
-				AbilityComponent->ServerTryActivateAbility(Activatable.Handle, false, MyKey);
-				return;
-			}
-		}
-
-		PC->ClientMessage(FString::Printf(TEXT("Failed to locate any ability matching %s"), *AbilityNameMatch));
-	}
-#endif // WITH_ABILITY_CHEATS
-}
-
-#if WITH_ABILITY_CHEATS
-static void TerminatePlayerAbility(const FString &AbilityNameMatch, TFunction<FString(UAbilitySystemComponent*, FGameplayAbilitySpec&)> TerminateFn)
-{
-	APlayerController* PC = GWorld->GetFirstPlayerController();
-	UAbilitySystemComponent* AbilityComponent = PC ? UAbilitySystemGlobals::GetAbilitySystemComponentFromActor(PC->GetPawn()) : nullptr;
-	if (AbilityComponent)
-	{
-		for (FGameplayAbilitySpec &ActivatableSpec : AbilityComponent->GetActivatableAbilities())
-		{
-			if (ActivatableSpec.Ability->IsActive() &&
-				(AbilityNameMatch.IsEmpty() || ActivatableSpec.Ability->GetName().Contains(AbilityNameMatch)))
-			{
-				PC->ClientMessage(TerminateFn(AbilityComponent, ActivatableSpec));
-			}
-		}
-	}
-}
-#endif // WITH_ABILITY_CHEATS
-
-void UAbilitySystemGlobals::ServerEndPlayerAbility(FString AbilityNameMatch)
-{
-#if WITH_ABILITY_CHEATS
-	TerminatePlayerAbility(AbilityNameMatch, [](UAbilitySystemComponent* AbilityComponent, FGameplayAbilitySpec& AbilitySpec) {
-		AbilityComponent->ServerEndAbility(AbilitySpec.Handle, AbilitySpec.Ability->GetCurrentActivationInfo(), AbilityComponent->ScopedPredictionKey);
-		return FString::Printf(TEXT("Triggering forced server ending of %s"), *AbilitySpec.Ability->GetName());
-	});
-#endif // WITH_ABILITY_CHEATS
-}
-
-void UAbilitySystemGlobals::ServerCancelPlayerAbility(FString AbilityNameMatch)
-{
-#if WITH_ABILITY_CHEATS
-	TerminatePlayerAbility(AbilityNameMatch, [](UAbilitySystemComponent* AbilityComponent, FGameplayAbilitySpec& AbilitySpec) {
-		AbilityComponent->ServerCancelAbility(AbilitySpec.Handle, AbilitySpec.Ability->GetCurrentActivationInfo());
-		return FString::Printf(TEXT("Triggering forced server cancellation of %s"), *AbilitySpec.Ability->GetName());
-	});
-#endif // WITH_ABILITY_CHEATS
+	return UE::AbilitySystemGlobals::bIgnoreAbilitySystemCosts;
 }
 
 #if WITH_EDITOR
@@ -697,3 +591,381 @@ bool FNetSerializeScriptStructCache::NetSerialize(FArchive& Ar, UScriptStruct*&
 		return false;
 	}
 }
+
+namespace UE::AbilitySystemGlobals
+{
+	/**
+	 * Common logic used to fuzzy-find a requested class (or alternatively, a passed-in asset path).
+	 */
+	template<typename ClassToFind>
+	TSubclassOf<ClassToFind> FuzzyFindClass(FString SearchString)
+	{
+		TSubclassOf<ClassToFind> FoundClass;
+
+		// See if we passed in a class name of a Class that already exists in memory.
+		// If we passed-in Default__, just remove that part since we're looking for Classes, not CDO names.
+		SearchString.RemoveFromStart(TEXT("Default__"));
+		const int SearchStringLen = SearchString.Len();
+		int BestClassMatchLen = INT_MAX;
+		for (TObjectIterator<UClass> ClassIt; ClassIt; ++ClassIt)
+		{
+			const bool bClassMatches = ClassIt->IsChildOf(ClassToFind::StaticClass());
+			if (!bClassMatches)
+			{
+				continue;
+			}
+
+			// Class name search
+			const FString ClassName = ClassIt->GetName();
+			const int ClassNameLen = ClassName.Len();
+			if (ClassNameLen < BestClassMatchLen && ClassNameLen >= SearchStringLen)
+			{
+				bool bContains = ClassName.Contains(SearchString);
+				if (bContains)
+				{
+					FoundClass = *ClassIt;
+					BestClassMatchLen = ClassNameLen;
+				}
+			}
+		}
+
+		// If it wasn't a class name, then perhaps it was the path to a specific asset
+		if (!FoundClass)
+		{
+			FSoftObjectPath SoftObjectPath{ SearchString };
+			if (UObject* ReferencedObject = SoftObjectPath.ResolveObject())
+			{
+				if (UPackage* ReferencedPackage = Cast<UPackage>(ReferencedObject))
+				{
+					ReferencedObject = ReferencedPackage->FindAssetInPackage();
+				}
+
+				if (UBlueprint* ReferencedBlueprint = Cast<UBlueprint>(ReferencedObject))
+				{
+					FoundClass = ReferencedBlueprint->GeneratedClass;
+				}
+				else if (ClassToFind* ReferencedGA = Cast<ClassToFind>(ReferencedObject))
+				{
+					FoundClass = ReferencedGA->GetClass();
+				}
+			}
+		}
+
+		return FoundClass;
+	}
+}
+
+//
+// AbilitySystem.Ability Debug Commands
+//
+
+FAutoConsoleCommand DebugAbilitySystemAbilityListGrantedCommand(TEXT("AbilitySystem.Ability.ListGranted"), TEXT("Lists all of the Gameplay Abilities currently granted to the Player"),
+	FConsoleCommandWithWorldArgsAndOutputDeviceDelegate::CreateLambda([](const TArray<FString>& Args, UWorld* World, FOutputDevice& OutputDevice)
+	{
+		UAbilitySystemGlobals& AbilitySystemGlobals = UAbilitySystemGlobals::Get();
+
+		APlayerController* PC = World->GetFirstPlayerController();
+		UAbilitySystemComponent* ASC = PC ? AbilitySystemGlobals.GetAbilitySystemComponentFromActor(PC->GetPawn()) : nullptr;
+		if (!ASC)
+		{
+			OutputDevice.Logf(TEXT("Could not find Player (%s) with AbilitySystemComponent in World (%s)"), *GetNameSafe(PC), *GetNameSafe(World));
+			return;
+		}
+
+		const UEnum* ExecutionEnumPtr = FindObject<UEnum>(nullptr, TEXT("/Script/GameplayAbilities.EGameplayAbilityNetExecutionPolicy"), true);
+		check(ExecutionEnumPtr&& TEXT("Couldn't locate EGameplayAbilityNetExecutionPolicy enum!"));
+
+		const UEnum* SecurityEnumPtr = FindObject<UEnum>(nullptr, TEXT("/Script/GameplayAbilities.EGameplayAbilityNetSecurityPolicy"), true);
+		check(SecurityEnumPtr&& TEXT("Couldn't locate EGameplayAbilityNetSecurityPolicy enum!"));
+
+		OutputDevice.Logf(TEXT("Granted abilities to %s (ASC: '%s'):"), *PC->GetName(), *ASC->GetFullName());
+
+		for (FGameplayAbilitySpec& Activatable : ASC->GetActivatableAbilities())
+		{
+			const TCHAR* ActiveText = Activatable.IsActive() ? TEXT("**ACTIVE**") : TEXT("");
+			OutputDevice.Logf(TEXT("   %s (%s - %s) %s"), *Activatable.Ability->GetName(), *ExecutionEnumPtr->GetDisplayNameTextByIndex(Activatable.Ability->GetNetExecutionPolicy()).ToString(), *SecurityEnumPtr->GetDisplayNameTextByIndex(Activatable.Ability->GetNetSecurityPolicy()).ToString(), ActiveText);
+		}
+	}), ECVF_Cheat);
+
+FAutoConsoleCommand DebugAbilitySystemAbilityCancelCommand(TEXT("AbilitySystem.Ability.Cancel"), TEXT("<Name>. Cancels (prematurely Ends) a currently executing Gameplay Ability"),
+	FConsoleCommandWithWorldArgsAndOutputDeviceDelegate::CreateLambda([](const TArray<FString>& Args, UWorld* World, FOutputDevice& OutputDevice)
+		{
+			UAbilitySystemGlobals& AbilitySystemGlobals = UAbilitySystemGlobals::Get();
+
+			APlayerController* PC = World->GetFirstPlayerController();
+			UAbilitySystemComponent* ASC = PC ? AbilitySystemGlobals.GetAbilitySystemComponentFromActor(PC->GetPawn()) : nullptr;
+			if (!ASC)
+			{
+				OutputDevice.Logf(TEXT("Could not find Player (%s) with AbilitySystemComponent in World (%s)"), *GetNameSafe(PC), *GetNameSafe(World));
+				return;
+			}
+
+			if (Args.Num() < 1)
+			{
+				if (IConsoleObject* GetGrantedCommand = IConsoleManager::Get().FindConsoleObject(TEXT("AbilitySystem.Ability.ListGranted"), false))
+				{
+					GetGrantedCommand->AsCommand()->Execute(Args, World, OutputDevice);
+				}
+				return;
+			}
+
+			FString SearchString = Args[0];
+			SearchString.RemoveFromStart(TEXT("Default__"));
+
+			bool bFound = false;
+			for (FGameplayAbilitySpec& GASpec : ASC->GetActivatableAbilities())
+			{
+				if (GASpec.IsActive())
+				{
+					if (!GASpec.Ability || !GASpec.Ability->GetName().Contains(SearchString))
+					{
+						continue;
+					}
+
+					TArray<UGameplayAbility*> ActiveAbilities = GASpec.GetAbilityInstances();
+					if (ActiveAbilities.Num() > 0)
+					{
+						for (UGameplayAbility* Instance : ActiveAbilities)
+						{
+							if (Instance)
+							{
+								OutputDevice.Logf(TEXT("%s (%s): Cancelling (instanced) %s"), *PC->GetName(), *ASC->GetName(), *Instance->GetName());
+								Instance->CancelAbility(GASpec.Handle, ASC->AbilityActorInfo.Get(), GASpec.ActivationInfo, true);
+								bFound = true;
+							}
+						}
+					}
+					else
+					{
+						OutputDevice.Logf(TEXT("%s (%s): Cancelling (non-instanced) %s"), *PC->GetName(), *ASC->GetName(), *GASpec.Ability->GetName());
+						GASpec.Ability->CancelAbility(GASpec.Handle, ASC->AbilityActorInfo.Get(), GASpec.ActivationInfo, true);
+						bFound = true;
+					}
+				}
+			}
+
+			if (!bFound)
+			{
+				OutputDevice.Logf(TEXT("Found no Active Gameplay Abilities on %s (%s) that matched '%s'"), *PC->GetName(), *ASC->GetName(), *SearchString);
+			}
+	}), ECVF_Cheat);
+
+FAutoConsoleCommand DebugAbilitySystemAbilityActivateCommand(TEXT("AbilitySystem.Ability.Activate"), TEXT("<TagName/ClassName/AssetName>. Activate a Gameplay Ability.\nSubstring name matching works for Activation Tags (on already granted abilities), Asset Paths (on non-granted abilities), or Class Names on both."),
+	FConsoleCommandWithWorldArgsAndOutputDeviceDelegate::CreateLambda([](const TArray<FString>& Args, UWorld* World, FOutputDevice& OutputDevice)
+	{
+		UAbilitySystemGlobals& AbilitySystemGlobals = UAbilitySystemGlobals::Get();
+
+		APlayerController* PC = World->GetFirstPlayerController();
+		UAbilitySystemComponent* ASC = PC ? AbilitySystemGlobals.GetAbilitySystemComponentFromActor(PC->GetPawn()) : nullptr;
+		if (!ASC)
+		{
+			OutputDevice.Logf(TEXT("Could not find Player (%s) with AbilitySystemComponent in World (%s)"), *GetNameSafe(PC), *GetNameSafe(World));
+			return;
+		}
+
+		if (Args.Num() < 1)
+		{
+			if (IConsoleObject* GetGrantedCommand = IConsoleManager::Get().FindConsoleObject(TEXT("AbilitySystem.Ability.ListGranted"), false))
+			{
+				GetGrantedCommand->AsCommand()->Execute(Args, World, OutputDevice);
+			}
+
+			return;
+		}
+
+		const FString SearchString = Args[0];
+
+		// Start by figuring out if we're trying to execute by GameplayTag.
+		const FName SearchName{ TCHAR_TO_ANSI(*SearchString), FNAME_Find };
+		if (SearchName != NAME_None)
+		{
+			const FGameplayTag GameplayTag = FGameplayTag::RequestGameplayTag(SearchName, false);
+			if (GameplayTag.IsValid())
+			{
+				FGameplayTagContainer GameplayTagContainer{ GameplayTag };
+				TArray<FGameplayAbilitySpec*> MatchingSpecs;
+				constexpr bool bTagRequirementsMustMatch = false;
+				ASC->GetActivatableGameplayAbilitySpecsByAllMatchingTags(GameplayTagContainer, MatchingSpecs, bTagRequirementsMustMatch);
+				if (MatchingSpecs.Num() > 0)
+				{
+					FString MatchingAbilitiesString = FString::JoinBy(MatchingSpecs, TEXT(", "), [](FGameplayAbilitySpec* Item) { return Item->GetDebugString(); });
+
+					bool bSuccess = ASC->TryActivateAbilitiesByTag(GameplayTagContainer);
+					if (bSuccess)
+					{
+						OutputDevice.Logf(TEXT("Requested Tag '%s' successfully executed one of: %s."), *SearchName.ToString(), *MatchingAbilitiesString);
+					}
+					else
+					{
+						OutputDevice.Logf(TEXT("Requested Tag '%s' was expected to execute one of: %s. But it failed to do so due to tag requirements."), *SearchName.ToString(), *MatchingAbilitiesString);
+					}
+				}
+				else
+				{
+					OutputDevice.Logf(TEXT("Requested Tag '%s' matched no given Gameplay Abilities to %s."), *SearchName.ToString(), *GetNameSafe(PC));
+				}
+
+				return;
+			}
+		}
+
+		// We couldn't find anything the user was searching for, so early out
+		TSubclassOf<UGameplayAbility> ActivateAbilityClass = UE::AbilitySystemGlobals::FuzzyFindClass<UGameplayAbility>(SearchString);
+		if (!ActivateAbilityClass)
+		{
+			OutputDevice.Logf(TEXT("Could not find a valid Gameplay Ability based on Search String '%s'"), *SearchString);
+			return;
+		}
+
+		// We found what the user was searching for, so let's try to activate it.  First, let's assume the Ability was already granted.
+		bool bSuccess = ASC->TryActivateAbilityByClass(ActivateAbilityClass);
+		if (bSuccess)
+		{
+			OutputDevice.Logf(TEXT("Successfully Activated Granted Ability '%s' on Player '%s'."), *GetNameSafe(ActivateAbilityClass.Get()), *GetNameSafe(PC));
+			return;
+		}
+
+		// It wasn't granted (or we failed to activate it even though it was granted).  We can't really differentiate those two, so let's grant it, then activate it.
+		FGameplayAbilitySpec AbilitySpec{ ActivateAbilityClass };
+		FGameplayAbilitySpecHandle SpecHandle = ASC->GiveAbilityAndActivateOnce(AbilitySpec);
+		if (SpecHandle.IsValid())
+		{
+			OutputDevice.Logf(TEXT("Successfully Granted, then Activated '%s' on Player '%s'."), *GetNameSafe(ActivateAbilityClass.Get()), *GetNameSafe(PC));
+		}
+		else
+		{
+			OutputDevice.Logf(TEXT("Failed to Grant & Activate '%s' on Player '%s'."), *GetNameSafe(ActivateAbilityClass.Get()), *GetNameSafe(PC));
+		}
+	}), ECVF_Cheat);
+
+
+//
+// AbilitySystem.Effect Debug Commands
+//
+
+FAutoConsoleCommand DebugAbilitySystemEffectListActiveCommand(TEXT("AbilitySystem.Effect.ListActive"), TEXT("Lists all of the Gameplay Effects currently active on the Player"),
+	FConsoleCommandWithWorldArgsAndOutputDeviceDelegate::CreateLambda([](const TArray<FString>& Args, UWorld* World, FOutputDevice& OutputDevice)
+	{
+		UAbilitySystemGlobals& AbilitySystemGlobals = UAbilitySystemGlobals::Get();
+
+		APlayerController* PC = World->GetFirstPlayerController();
+		UAbilitySystemComponent* ASC = PC ? AbilitySystemGlobals.GetAbilitySystemComponentFromActor(PC->GetPawn()) : nullptr;
+		if (!ASC)
+		{
+			OutputDevice.Logf(TEXT("Could not find Player (%s) with AbilitySystemComponent in World (%s)"), *GetNameSafe(PC), *GetNameSafe(World));
+			return;
+		}
+
+		OutputDevice.Logf(TEXT("Active Gameplay Effects on %s (%s):"), *PC->GetName(), *ASC->GetName());
+
+		const FActiveGameplayEffectsContainer& ActiveGEContainer = ASC->GetActiveGameplayEffects();
+		for (const FActiveGameplayEffectHandle& ActiveGEHandle : ActiveGEContainer.GetAllActiveEffectHandles())
+		{
+			if (const FActiveGameplayEffect* ActiveGE = ASC->GetActiveGameplayEffect(ActiveGEHandle))
+			{
+				const UGameplayEffect* GESpecDef = ActiveGE->Spec.Def;
+				const TCHAR* ActiveState = ActiveGE->bIsInhibited ? TEXT("Inhibited") : TEXT("Active");
+
+				OutputDevice.Logf(TEXT(" [%u] %s: %s"), GetTypeHash(ActiveGEHandle), *GetNameSafe(GESpecDef), ActiveState);
+			}
+		}
+	}), ECVF_Cheat);
+
+FAutoConsoleCommand DebugAbilitySystemEffectRemoveCommand(TEXT("AbilitySystem.Effect.Remove"), TEXT("<Handle/Name>. Remove a Gameplay Effect that is currently active on the Player"),
+	FConsoleCommandWithWorldArgsAndOutputDeviceDelegate::CreateLambda([](const TArray<FString>& Args, UWorld* World, FOutputDevice& OutputDevice)
+		{
+			UAbilitySystemGlobals& AbilitySystemGlobals = UAbilitySystemGlobals::Get();
+
+			APlayerController* PC = World->GetFirstPlayerController();
+			UAbilitySystemComponent* ASC = PC ? AbilitySystemGlobals.GetAbilitySystemComponentFromActor(PC->GetPawn()) : nullptr;
+			if (!ASC)
+			{
+				OutputDevice.Logf(TEXT("Could not find Player (%s) with AbilitySystemComponent in World (%s)"), *GetNameSafe(PC), *GetNameSafe(World));
+				return;
+			}
+
+			if (Args.Num() < 1)
+			{
+				if (IConsoleObject* ListActiveCommand = IConsoleManager::Get().FindConsoleObject(TEXT("AbilitySystem.Effect.ListActive"), false))
+				{
+					ListActiveCommand->AsCommand()->Execute(Args, World, OutputDevice);
+				}
+
+				return;
+			}
+
+			uint32 SearchHandleValue = 0;
+			FString SearchString = Args[0];
+			if (SearchString.IsNumeric())
+			{
+				int64 SearchHash64 = FCString::Atoi64(*SearchString);
+				if (SearchHash64 > 0)
+				{
+					SearchHandleValue = static_cast<uint32>(SearchHash64 & UINT_MAX);
+				}
+			}
+
+			bool bAnyMatch = false;
+			const FActiveGameplayEffectsContainer& ActiveGEContainer = ASC->GetActiveGameplayEffects();
+			for (const FActiveGameplayEffectHandle& ActiveGEHandle : ActiveGEContainer.GetAllActiveEffectHandles())
+			{
+				if (const FActiveGameplayEffect* ActiveGE = ASC->GetActiveGameplayEffect(ActiveGEHandle))
+				{
+					const UGameplayEffect* GESpecDef = ActiveGE->Spec.Def;
+					const uint32 ActiveGEHandleValue = GetTypeHash(ActiveGEHandle);
+
+					const bool bMatch = (SearchHandleValue > 0) ? (ActiveGEHandleValue == SearchHandleValue) : GESpecDef->GetName().Contains(SearchString);
+					if (bMatch)
+					{
+						OutputDevice.Logf(TEXT("Removing Active GE Handle %u (%s)"), ActiveGEHandleValue, *GESpecDef->GetName());
+						ASC->RemoveActiveGameplayEffect(ActiveGEHandle);
+						bAnyMatch = true;
+					}
+				}
+			}
+
+			if (!bAnyMatch)
+			{
+				OutputDevice.Logf(TEXT("Found no Active Gameplay Effects on %s (%s) that matched '%s'"), *PC->GetName(), *ASC->GetName(), *SearchString);
+			}
+		}), ECVF_Cheat);
+
+FAutoConsoleCommand DebugAbilitySystemEffectApply(TEXT("AbilitySystem.Effect.Apply"), TEXT("<Class/AssetName> [Level]. Apply a Gameplay Effect on the Player.  Substring name matching works for Asset Tags, Asset Paths, or Class Names."),
+	FConsoleCommandWithWorldArgsAndOutputDeviceDelegate::CreateLambda([](const TArray<FString>& Args, UWorld* World, FOutputDevice& OutputDevice)
+	{
+		UAbilitySystemGlobals& AbilitySystemGlobals = UAbilitySystemGlobals::Get();
+
+		APlayerController* PC = World->GetFirstPlayerController();
+		UAbilitySystemComponent* ASC = PC ? AbilitySystemGlobals.GetAbilitySystemComponentFromActor(PC->GetPawn()) : nullptr;
+		if (!ASC)
+		{
+			OutputDevice.Logf(TEXT("Could not find Player (%s) with AbilitySystemComponent in World (%s)"), *GetNameSafe(PC), *GetNameSafe(World));
+			return;
+		}
+
+		// We couldn't find anything the user was searching for, so early out
+		const FString SearchString = Args[0];
+		TSubclassOf<UGameplayEffect> GameplayEffectClass = UE::AbilitySystemGlobals::FuzzyFindClass<UGameplayEffect>(SearchString);
+		if (!GameplayEffectClass)
+		{
+			OutputDevice.Logf(TEXT("Could not find a valid Gameplay Effect based on Search String '%s'"), *SearchString);
+			return;
+		}
+
+		// Create a GameplayEffectSpec that executes the passed-in parameters
+		FGameplayEffectContextHandle GEContextHandle = ASC->MakeEffectContext();
+		const float Level = FMath::Max(FGameplayEffectConstants::INVALID_LEVEL, (Args.Num() > 1) ? FPlatformString::Atof(*Args[1]) : FGameplayEffectConstants::INVALID_LEVEL);
+		FGameplayEffectSpec GESpec{ GameplayEffectClass.GetDefaultObject(), GEContextHandle, Level };
+
+		// We need to create a new valid prediction key (as if we just started activating an ability) so the GE will fire even if we are not the authority.
+		FPredictionKey FakePredictionKey = (PC->GetLocalRole() == ENetRole::ROLE_Authority) ? FPredictionKey() : FPredictionKey::CreateNewPredictionKey(ASC);
+		FActiveGameplayEffectHandle ActiveGEHandle = ASC->ApplyGameplayEffectSpecToSelf(GESpec, FakePredictionKey);
+		if (ActiveGEHandle.IsValid())
+		{
+			OutputDevice.Logf(TEXT("Successfully Applied Gameplay Effect '%s' on Player '%s'."), *GameplayEffectClass->GetName(), *PC->GetName());
+		}
+		else
+		{
+			OutputDevice.Logf(TEXT("Failed to Apply Gameplay Effect '%s' on Player '%s'."), *GameplayEffectClass->GetName(), *PC->GetName());
+		}
+	}), ECVF_Cheat);
diff --git a/Engine/Plugins/Runtime/GameplayAbilities/Source/GameplayAbilities/Private/GameplayAbilitiesExec.cpp b/Engine/Plugins/Runtime/GameplayAbilities/Source/GameplayAbilities/Private/GameplayAbilitiesExec.cpp
deleted file mode 100644
diff --git a/Engine/Plugins/Runtime/GameplayAbilities/Source/GameplayAbilities/Public/AbilitySystemGlobals.h b/Engine/Plugins/Runtime/GameplayAbilities/Source/GameplayAbilities/Public/AbilitySystemGlobals.h
--- a/Engine/Plugins/Runtime/GameplayAbilities/Source/GameplayAbilities/Public/AbilitySystemGlobals.h
+++ b/Engine/Plugins/Runtime/GameplayAbilities/Source/GameplayAbilities/Public/AbilitySystemGlobals.h
@@ -146,12 +146,12 @@ class GAMEPLAYABILITIES_API UAbilitySystemGlobals : public UObject
 	// Cheat functions
 
 	/** Toggles whether we should ignore ability cooldowns. Does nothing in shipping builds */
-	UFUNCTION(exec)
-	virtual void ToggleIgnoreAbilitySystemCooldowns();
+	UE_DEPRECATED(5.3, "Use CVarAbilitySystemIgnoreCooldowns")
+	virtual void ToggleIgnoreAbilitySystemCooldowns() {}
 
 	/** Toggles whether we should ignore ability costs. Does nothing in shipping builds */
-	UFUNCTION(exec)
-	virtual void ToggleIgnoreAbilitySystemCosts();
+	UE_DEPRECATED(5.3, "Use CVarAbilitySystemIgnoreCosts")
+	virtual void ToggleIgnoreAbilitySystemCosts() {}
 
 	/** Returns true if ability cooldowns are ignored, returns false otherwise. Always returns false in shipping builds. */
 	bool ShouldIgnoreCooldowns() const;
@@ -160,17 +160,17 @@ class GAMEPLAYABILITIES_API UAbilitySystemGlobals : public UObject
 	bool ShouldIgnoreCosts() const;
 
 	/** Show all abilities currently assigned to the local player */
-	UFUNCTION(exec)
-	void ListPlayerAbilities();
+	UE_DEPRECATED(5.3, "Use DebugAbilitySystemAbilityListGrantedCommand")
+	void ListPlayerAbilities() {}
 	/** Force server activation of a specific player ability (useful for cheat testing) */
-	UFUNCTION(exec)
-	void ServerActivatePlayerAbility(FString AbilityNameMatch);
+	UE_DEPRECATED(5.3, "Use DebugAbilitySystemAbilityActivateCommand")
+	void ServerActivatePlayerAbility(FString AbilityNameMatch) {}
 	/** Force server deactivation of a specific player ability (useful for cheat testing) */
-	UFUNCTION(exec)
-	void ServerEndPlayerAbility(FString AbilityNameMatch);
+	UE_DEPRECATED(5.3, "Use DebugAbilitySystemAbilityCancelCommand (EndAbility is only for internal usage)")
+	void ServerEndPlayerAbility(FString AbilityNameMatch) {}
 	/** Force server cancellation of a specific player ability (useful for cheat testing) */
-	UFUNCTION(exec)
-	void ServerCancelPlayerAbility(FString AbilityNameMatch);
+	UE_DEPRECATED(5.3, "Use DebugAbilitySystemAbilityCancelCommand")
+	void ServerCancelPlayerAbility(FString AbilityNameMatch) {}
 
 	/** Called when debug strings are available, to write them to the display */
 	DECLARE_MULTICAST_DELEGATE(FOnClientServerDebugAvailable);
@@ -324,9 +324,11 @@ protected:
 	// data used for ability system cheat commands
 
 	/** If we should ignore the cooldowns when activating abilities in the ability system. Set with ToggleIgnoreAbilitySystemCooldowns() */
+	UE_DEPRECATED(5.3, "Use bIgnoreAbilitySystemCooldowns in the AbilitySystemGlobals namespace, controlled by new CVarAbilitySystemIgnoreCooldowns")
 	bool bIgnoreAbilitySystemCooldowns;
 
 	/** If we should ignore the costs when activating abilities in the ability system. Set with ToggleIgnoreAbilitySystemCosts() */
+	UE_DEPRECATED(5.3, "Use bIgnoreAbilitySystemCosts in the AbilitySystemGlobals namespace, controlled by new CVarAbilitySystemIgnoreCosts")
 	bool bIgnoreAbilitySystemCosts;
 #endif // WITH_ABILITY_CHEATS
 
