Fix for ShouldAbilityRespondToEvent not getting called on the instanced GameplayAbility.

diff --git a/Engine/Plugins/Runtime/GameplayAbilities/Source/GameplayAbilities/Private/AbilitySystemComponent_Abilities.cpp b/Engine/Plugins/Runtime/GameplayAbilities/Source/GameplayAbilities/Private/AbilitySystemComponent_Abilities.cpp
--- a/Engine/Plugins/Runtime/GameplayAbilities/Source/GameplayAbilities/Private/AbilitySystemComponent_Abilities.cpp
+++ b/Engine/Plugins/Runtime/GameplayAbilities/Source/GameplayAbilities/Private/AbilitySystemComponent_Abilities.cpp
@@ -1543,18 +1543,19 @@ bool UAbilitySystemComponent::InternalTryActivateAbility(FGameplayAbilitySpecHan
 		return false;
 	}
 
+	// If it's an instanced one, the instanced ability will be set, otherwise it will be null
+	UGameplayAbility* InstancedAbility = Spec->GetPrimaryInstance();
+
 	if (TriggerEventData)
 	{
-		if (!Ability->ShouldAbilityRespondToEvent(ActorInfo, TriggerEventData))
+		UGameplayAbility* AbilitySource = InstancedAbility ? InstancedAbility : Ability;
+		if (!AbilitySource->ShouldAbilityRespondToEvent(ActorInfo, TriggerEventData))
 		{
-			NotifyAbilityFailed(Handle, Ability, InternalTryActivateAbilityFailureTags);
+			NotifyAbilityFailed(Handle, AbilitySource, InternalTryActivateAbilityFailureTags);
 			return false;
 		}
 	}
 
-	// If it's instance once the instanced ability will be set, otherwise it will be null
-	UGameplayAbility* InstancedAbility = Spec->GetPrimaryInstance();
-
 	{
 		const FGameplayTagContainer* SourceTags = TriggerEventData ? &TriggerEventData->InstigatorTags : nullptr;
 		const FGameplayTagContainer* TargetTags = TriggerEventData ? &TriggerEventData->TargetTags : nullptr;
