Reverting FGameplayEffectQuery::SourceTagQuery back to checking only the Source's SpecTags to make sure old content doesn't break.  Introduced SourceAggregateTagQuery which is probably what you want going forward.

diff --git a/Engine/Plugins/Runtime/GameplayAbilities/Source/GameplayAbilities/Private/AbilitySystemComponent.cpp b/Engine/Plugins/Runtime/GameplayAbilities/Source/GameplayAbilities/Private/AbilitySystemComponent.cpp
--- a/Engine/Plugins/Runtime/GameplayAbilities/Source/GameplayAbilities/Private/AbilitySystemComponent.cpp
+++ b/Engine/Plugins/Runtime/GameplayAbilities/Source/GameplayAbilities/Private/AbilitySystemComponent.cpp
@@ -1532,7 +1532,7 @@ int32 UAbilitySystemComponent::RemoveActiveEffectsWithSourceTags(FGameplayTagCon
 {
 	if (IsOwnerActorAuthoritative())
 	{
-		return RemoveActiveEffects(FGameplayEffectQuery::MakeQuery_MatchAnySourceTags(Tags));
+		return RemoveActiveEffects(FGameplayEffectQuery::MakeQuery_MatchAnySourceSpecTags(Tags));
 	}
 	return 0;
 }
diff --git a/Engine/Plugins/Runtime/GameplayAbilities/Source/GameplayAbilities/Private/GameplayEffect.cpp b/Engine/Plugins/Runtime/GameplayAbilities/Source/GameplayAbilities/Private/GameplayEffect.cpp
--- a/Engine/Plugins/Runtime/GameplayAbilities/Source/GameplayAbilities/Private/GameplayEffect.cpp
+++ b/Engine/Plugins/Runtime/GameplayAbilities/Source/GameplayAbilities/Private/GameplayEffect.cpp
@@ -5099,6 +5099,7 @@ FGameplayEffectQuery& FGameplayEffectQuery::operator=(FGameplayEffectQuery&& Oth
 	CustomMatchDelegate_BP = MoveTemp(Other.CustomMatchDelegate_BP);
 	OwningTagQuery = MoveTemp(Other.OwningTagQuery);
 	EffectTagQuery = MoveTemp(Other.EffectTagQuery);
+	SourceAggregateTagQuery = MoveTemp(Other.SourceAggregateTagQuery);
 	SourceTagQuery = MoveTemp(Other.SourceTagQuery);
 	ModifyingAttribute = MoveTemp(Other.ModifyingAttribute);
 	EffectSource = Other.EffectSource;
@@ -5113,6 +5114,7 @@ FGameplayEffectQuery& FGameplayEffectQuery::operator=(const FGameplayEffectQuery
 	CustomMatchDelegate_BP = Other.CustomMatchDelegate_BP;
 	OwningTagQuery = Other.OwningTagQuery;
 	EffectTagQuery = Other.EffectTagQuery;
+	SourceAggregateTagQuery = Other.SourceAggregateTagQuery;
 	SourceTagQuery = Other.SourceTagQuery;
 	ModifyingAttribute = Other.ModifyingAttribute;
 	EffectSource = Other.EffectSource;
@@ -5199,11 +5201,19 @@ bool FGameplayEffectQuery::Matches(const FGameplayEffectSpec& Spec) const
 		}
 	}
 
+	if (SourceAggregateTagQuery.IsEmpty() == false)
+	{
+		FGameplayTagContainer const& SourceAggregateTags = *Spec.CapturedSourceTags.GetAggregatedTags();
+		if (SourceAggregateTagQuery.Matches(SourceAggregateTags) == false)
+		{
+			return false;
+		}
+	}
+
 	if (SourceTagQuery.IsEmpty() == false)
 	{
-		// Prior to UE5.3 this checked against the SpecTags, which seems like an obvious bug.
-		FGameplayTagContainer const& SourceTags = *Spec.CapturedSourceTags.GetAggregatedTags();
-		if (SourceTagQuery.Matches(SourceTags) == false)
+		FGameplayTagContainer const& SourceSpecTags = Spec.CapturedSourceTags.GetSpecTags();
+		if (SourceTagQuery.Matches(SourceSpecTags) == false)
 		{
 			return false;
 		}
@@ -5259,6 +5269,7 @@ bool FGameplayEffectQuery::IsEmpty() const
 	(
 		OwningTagQuery.IsEmpty() &&
 		EffectTagQuery.IsEmpty() &&
+		SourceAggregateTagQuery.IsEmpty() &&
 		SourceTagQuery.IsEmpty() &&
 		!ModifyingAttribute.IsValid() &&
 		!EffectSource &&
@@ -5321,7 +5332,7 @@ FGameplayEffectQuery FGameplayEffectQuery::MakeQuery_MatchNoEffectTags(const FGa
 }
 
 // static
-FGameplayEffectQuery FGameplayEffectQuery::MakeQuery_MatchAnySourceTags(const FGameplayTagContainer& InTags)
+FGameplayEffectQuery FGameplayEffectQuery::MakeQuery_MatchAnySourceSpecTags(const FGameplayTagContainer& InTags)
 {
 	SCOPE_CYCLE_COUNTER(STAT_MakeGameplayEffectQuery);
 	FGameplayEffectQuery OutQuery;
@@ -5330,7 +5341,7 @@ FGameplayEffectQuery FGameplayEffectQuery::MakeQuery_MatchAnySourceTags(const FG
 }
 
 // static
-FGameplayEffectQuery FGameplayEffectQuery::MakeQuery_MatchAllSourceTags(const FGameplayTagContainer& InTags)
+FGameplayEffectQuery FGameplayEffectQuery::MakeQuery_MatchAllSourceSpecTags(const FGameplayTagContainer& InTags)
 {
 	SCOPE_CYCLE_COUNTER(STAT_MakeGameplayEffectQuery);
 	FGameplayEffectQuery OutQuery;
@@ -5339,7 +5350,7 @@ FGameplayEffectQuery FGameplayEffectQuery::MakeQuery_MatchAllSourceTags(const FG
 }
 
 // static
-FGameplayEffectQuery FGameplayEffectQuery::MakeQuery_MatchNoSourceTags(const FGameplayTagContainer& InTags)
+FGameplayEffectQuery FGameplayEffectQuery::MakeQuery_MatchNoSourceSpecTags(const FGameplayTagContainer& InTags)
 {
 	SCOPE_CYCLE_COUNTER(STAT_MakeGameplayEffectQuery);
 	FGameplayEffectQuery OutQuery;
@@ -5352,6 +5363,7 @@ bool FGameplayEffectQuery::operator==(const FGameplayEffectQuery& Other) const
 	return CustomMatchDelegate_BP == Other.CustomMatchDelegate_BP &&
 	OwningTagQuery == Other.OwningTagQuery &&
 	EffectTagQuery == Other.EffectTagQuery &&
+	SourceAggregateTagQuery == Other.SourceAggregateTagQuery &&
 	SourceTagQuery == Other.SourceTagQuery &&
 	ModifyingAttribute == Other.ModifyingAttribute &&
 	EffectSource == Other.EffectSource &&
diff --git a/Engine/Plugins/Runtime/GameplayAbilities/Source/GameplayAbilities/Public/GameplayEffect.h b/Engine/Plugins/Runtime/GameplayAbilities/Source/GameplayAbilities/Public/GameplayEffect.h
--- a/Engine/Plugins/Runtime/GameplayAbilities/Source/GameplayAbilities/Public/GameplayEffect.h
+++ b/Engine/Plugins/Runtime/GameplayAbilities/Source/GameplayAbilities/Public/GameplayEffect.h
@@ -1423,10 +1423,14 @@ public:
 	UPROPERTY(EditAnywhere, BlueprintReadWrite, Category = Query)
 	FGameplayTagQuery EffectTagQuery;
 
-	/** Query that is matched against tags the source of this GE has */
-	UPROPERTY(EditAnywhere, BlueprintReadWrite, Category = Query)
+	/** Query that is matched against spec tags the source of this GE has */
+	UPROPERTY(EditAnywhere, BlueprintReadWrite, Category = Query, DisplayName = SourceSpecTagQuery)
 	FGameplayTagQuery SourceTagQuery;
 
+	/** Query that is matched against all tags the source of this GE has */
+	UPROPERTY(EditAnywhere, BlueprintReadWrite, Category = Query)
+	FGameplayTagQuery SourceAggregateTagQuery;
+
 	/** Matches on GameplayEffects which modify given attribute. */
 	UPROPERTY(EditAnywhere, BlueprintReadWrite, Category = Query)
 	FGameplayAttribute ModifyingAttribute;
@@ -1469,12 +1473,12 @@ public:
 	/** Creates an effect query that will match if there are no common tags between the given tags and an ActiveGameplayEffect's tags */
 	static FGameplayEffectQuery MakeQuery_MatchNoEffectTags(const FGameplayTagContainer& InTags);
 
-	/** Creates an effect query that will match if there are any common tags between the given tags and an ActiveGameplayEffect's source tags */
-	static FGameplayEffectQuery MakeQuery_MatchAnySourceTags(const FGameplayTagContainer& InTags);
-	/** Creates an effect query that will match if all of the given tags are in the ActiveGameplayEffect's source tags */
-	static FGameplayEffectQuery MakeQuery_MatchAllSourceTags(const FGameplayTagContainer& InTags);
-	/** Creates an effect query that will match if there are no common tags between the given tags and an ActiveGameplayEffect's source tags */
-	static FGameplayEffectQuery MakeQuery_MatchNoSourceTags(const FGameplayTagContainer& InTags);
+	/** Creates an effect query that will match if there are any common tags between the given tags and an ActiveGameplayEffect's source spec tags */
+	static FGameplayEffectQuery MakeQuery_MatchAnySourceSpecTags(const FGameplayTagContainer& InTags);
+	/** Creates an effect query that will match if all of the given tags are in the ActiveGameplayEffect's source spec tags */
+	static FGameplayEffectQuery MakeQuery_MatchAllSourceSpecTags(const FGameplayTagContainer& InTags);
+	/** Creates an effect query that will match if there are no common tags between the given tags and an ActiveGameplayEffect's source spec tags */
+	static FGameplayEffectQuery MakeQuery_MatchNoSourceSpecTags(const FGameplayTagContainer& InTags);
 
 	bool operator==(const FGameplayEffectQuery& Other) const;
 	bool operator!=(const FGameplayEffectQuery& Other) const;
