Fix for FindActiveGameplayEffectHandle.  This was changed incorrectly during modularization of GE's.  The errant variable (FGameplayAbilitySpec::GameplayEffectHandle) is now deprecated in favor of this correct implementation.

diff --git a/Engine/Plugins/Runtime/GameplayAbilities/Source/GameplayAbilities/Private/AbilitySystemComponent.cpp b/Engine/Plugins/Runtime/GameplayAbilities/Source/GameplayAbilities/Private/AbilitySystemComponent.cpp
--- a/Engine/Plugins/Runtime/GameplayAbilities/Source/GameplayAbilities/Private/AbilitySystemComponent.cpp
+++ b/Engine/Plugins/Runtime/GameplayAbilities/Source/GameplayAbilities/Private/AbilitySystemComponent.cpp
@@ -1125,12 +1125,16 @@ FString UAbilitySystemComponent::GetActiveGEDebugString(FActiveGameplayEffectHan
 /** Gets the GE Handle of the GE that granted the passed in Ability */
 FActiveGameplayEffectHandle UAbilitySystemComponent::FindActiveGameplayEffectHandle(FGameplayAbilitySpecHandle Handle) const
 {
-	const FGameplayAbilitySpec* Spec = FindAbilitySpecFromHandle(Handle);
-	if (Spec)
+	for (const FActiveGameplayEffect& ActiveGE : &ActiveGameplayEffects)
 	{
-		return Spec->GameplayEffectHandle;
+		for (const FGameplayAbilitySpecDef& AbilitySpecDef : ActiveGE.Spec.GrantedAbilitySpecs)
+		{
+			if (AbilitySpecDef.AssignedHandle == Handle)
+			{
+				return ActiveGE.Handle;
+			}
+		}
 	}
-
 	return FActiveGameplayEffectHandle();
 }
 
diff --git a/Engine/Plugins/Runtime/GameplayAbilities/Source/GameplayAbilities/Private/AbilitySystemComponent_Abilities.cpp b/Engine/Plugins/Runtime/GameplayAbilities/Source/GameplayAbilities/Private/AbilitySystemComponent_Abilities.cpp
--- a/Engine/Plugins/Runtime/GameplayAbilities/Source/GameplayAbilities/Private/AbilitySystemComponent_Abilities.cpp
+++ b/Engine/Plugins/Runtime/GameplayAbilities/Source/GameplayAbilities/Private/AbilitySystemComponent_Abilities.cpp
@@ -769,13 +769,6 @@ FGameplayAbilitySpec* UAbilitySystemComponent::FindAbilitySpecFromHandle(FGamepl
 
 FGameplayAbilitySpec* UAbilitySystemComponent::FindAbilitySpecFromGEHandle(FActiveGameplayEffectHandle Handle) const
 {
-	for (const FGameplayAbilitySpec& Spec : ActivatableAbilities.Items)
-	{
-		if (Spec.GameplayEffectHandle == Handle)
-		{
-			return const_cast<FGameplayAbilitySpec*>(&Spec);
-		}
-	}
 	return nullptr;
 }
 
diff --git a/Engine/Plugins/Runtime/GameplayAbilities/Source/GameplayAbilities/Private/GameplayAbilityTypes.cpp b/Engine/Plugins/Runtime/GameplayAbilities/Source/GameplayAbilities/Private/GameplayAbilityTypes.cpp
--- a/Engine/Plugins/Runtime/GameplayAbilities/Source/GameplayAbilities/Private/GameplayAbilityTypes.cpp
+++ b/Engine/Plugins/Runtime/GameplayAbilities/Source/GameplayAbilities/Private/GameplayAbilityTypes.cpp
@@ -296,7 +296,9 @@ FGameplayAbilitySpec::FGameplayAbilitySpec(FGameplayAbilitySpecDef& InDef, int32
 {
 	Handle.GenerateNewHandle();
 	InDef.AssignedHandle = Handle;
+PRAGMA_DISABLE_DEPRECATION_WARNINGS
 	GameplayEffectHandle = InGameplayEffectHandle;
+PRAGMA_ENABLE_DEPRECATION_WARNINGS
 
 	SetByCallerTagMagnitudes = InDef.SetByCallerTagMagnitudes;
 
diff --git a/Engine/Plugins/Runtime/GameplayAbilities/Source/GameplayAbilities/Public/AbilitySystemComponent.h b/Engine/Plugins/Runtime/GameplayAbilities/Source/GameplayAbilities/Public/AbilitySystemComponent.h
--- a/Engine/Plugins/Runtime/GameplayAbilities/Source/GameplayAbilities/Public/AbilitySystemComponent.h
+++ b/Engine/Plugins/Runtime/GameplayAbilities/Source/GameplayAbilities/Public/AbilitySystemComponent.h
@@ -1100,6 +1100,7 @@ class GAMEPLAYABILITIES_API UAbilitySystemComponent : public UGameplayTasksCompo
 	FGameplayAbilitySpec* FindAbilitySpecFromHandle(FGameplayAbilitySpecHandle Handle) const;
 	
 	/** Returns an ability spec from a GE handle. If modifying call MarkAbilitySpecDirty */
+	UE_DEPRECATED(5.3, "FindAbilitySpecFromGEHandle was never accurate because a GameplayEffect can grant multiple GameplayAbilities. It now returns nullptr.")
 	FGameplayAbilitySpec* FindAbilitySpecFromGEHandle(FActiveGameplayEffectHandle Handle) const;
 
 	/** Returns an ability spec corresponding to given ability class. If modifying call MarkAbilitySpecDirty */
diff --git a/Engine/Plugins/Runtime/GameplayAbilities/Source/GameplayAbilities/Public/ActiveGameplayEffectHandle.h b/Engine/Plugins/Runtime/GameplayAbilities/Source/GameplayAbilities/Public/ActiveGameplayEffectHandle.h
--- a/Engine/Plugins/Runtime/GameplayAbilities/Source/GameplayAbilities/Public/ActiveGameplayEffectHandle.h
+++ b/Engine/Plugins/Runtime/GameplayAbilities/Source/GameplayAbilities/Public/ActiveGameplayEffectHandle.h
@@ -9,7 +9,7 @@ class UAbilitySystemComponent;
 /**
  * This handle is required for things outside of FActiveGameplayEffectsContainer to refer to a specific active GameplayEffect
  *	For example if a skill needs to create an active effect and then destroy that specific effect that it created, it has to do so
- *	through a handle. a pointer or index into the active list is not sufficient.
+ *	through a handle. a pointer or index into the active list is not sufficient. These are not synchronized between clients and server.
  */
 USTRUCT(BlueprintType)
 struct GAMEPLAYABILITIES_API FActiveGameplayEffectHandle
diff --git a/Engine/Plugins/Runtime/GameplayAbilities/Source/GameplayAbilities/Public/GameplayAbilitySpec.h b/Engine/Plugins/Runtime/GameplayAbilities/Source/GameplayAbilities/Public/GameplayAbilitySpec.h
--- a/Engine/Plugins/Runtime/GameplayAbilities/Source/GameplayAbilities/Public/GameplayAbilitySpec.h
+++ b/Engine/Plugins/Runtime/GameplayAbilities/Source/GameplayAbilities/Public/GameplayAbilitySpec.h
@@ -166,6 +166,14 @@ struct GAMEPLAYABILITIES_API FGameplayAbilitySpec : public FFastArraySerializerI
 {
 	GENERATED_USTRUCT_BODY()
 
+PRAGMA_DISABLE_DEPRECATION_WARNINGS
+	FGameplayAbilitySpec(const FGameplayAbilitySpec&) = default;
+	FGameplayAbilitySpec(FGameplayAbilitySpec&&) = default;
+	FGameplayAbilitySpec& operator=(const FGameplayAbilitySpec&) = default;
+	FGameplayAbilitySpec& operator=(FGameplayAbilitySpec&&) = default;
+	~FGameplayAbilitySpec() = default;
+PRAGMA_ENABLE_DEPRECATION_WARNINGS
+
 	FGameplayAbilitySpec()
 		: Ability(nullptr), Level(1), InputID(INDEX_NONE), SourceObject(nullptr), ActiveCount(0), InputPressed(false), RemoveAfterActivation(false), PendingRemove(false), bActivateOnce(false)
 	{ }
@@ -238,7 +246,8 @@ struct GAMEPLAYABILITIES_API FGameplayAbilitySpec : public FFastArraySerializerI
 	UPROPERTY()
 	TArray<TObjectPtr<UGameplayAbility>> ReplicatedInstances;
 
-	/** Handle to GE that granted us (usually invalid) */
+	/** Handle to GE that granted us (usually invalid). FActiveGameplayEffectHandles are not synced across the network. */
+	UE_DEPRECATED(5.3, "GameplayEffectHandle was never replicated across the network. If you need FGameplayAbilitySpec -> FActiveGameplayEffect then use AbilitySystemComponent::FindActiveGameplayEffectHandle")
 	UPROPERTY(NotReplicated)
 	FActiveGameplayEffectHandle	GameplayEffectHandle;
 
