Fixed the AbilitySystem Test for Activation Flow.  We no longer expect a CommitAbility to fire by default.

diff --git a/Engine/Plugins/Runtime/GameplayAbilities/Source/GameplayAbilities/Private/Tests/AbilitySystemComponentTests.cpp b/Engine/Plugins/Runtime/GameplayAbilities/Source/GameplayAbilities/Private/Tests/AbilitySystemComponentTests.cpp
--- a/Engine/Plugins/Runtime/GameplayAbilities/Source/GameplayAbilities/Private/Tests/AbilitySystemComponentTests.cpp
+++ b/Engine/Plugins/Runtime/GameplayAbilities/Source/GameplayAbilities/Private/Tests/AbilitySystemComponentTests.cpp
@@ -161,7 +161,7 @@ public: // the tests
 			Test->TestTrue(TEXT("TryActivateAbility executes successfully (using FGameplayAbilitySpecHandle)"), bLocalActivation);
 			Test->TestTrue(TEXT(" AbilitySpec.IsActive() after TryActivateAbility (using FGameplayAbilitySpecHandle)"), AbilitySpec->IsActive());
 			Test->TestTrue(TEXT(" AbilityActivated after TryActivateAbility (using FGameplayAbilitySpecHandle)"), TestCallbacks.bReceivedAbilityActivated);
-			Test->TestTrue(TEXT(" AbilityCommitted after TryActivateAbility (using FGameplayAbilitySpecHandle)"), TestCallbacks.bReceivedAbilityCommitted);
+			Test->TestFalse(TEXT(" AbilityCommitted after TryActivateAbility (using FGameplayAbilitySpecHandle)"), TestCallbacks.bReceivedAbilityCommitted);
 			Test->TestFalse(TEXT(" AbilityEnded (prematurely) after TryActivateAbility (using FGameplayAbilitySpecHandle)"), TestCallbacks.bReceiveAbilityEnded);
 			Test->TestFalse(TEXT(" AbilityFailed after TryActivateAbility (with an Ability that should succeed)"), TestCallbacks.bReceiveAbilityFailed);
 
