Removed vestigial GameplayAbility::SetMovementSyncPoint

diff --git a/Engine/Plugins/Runtime/GameplayAbilities/Source/GameplayAbilities/Private/Abilities/Tasks/AbilityTask_ApplyRootMotionConstantForce.cpp b/Engine/Plugins/Runtime/GameplayAbilities/Source/GameplayAbilities/Private/Abilities/Tasks/AbilityTask_ApplyRootMotionConstantForce.cpp
--- a/Engine/Plugins/Runtime/GameplayAbilities/Source/GameplayAbilities/Private/Abilities/Tasks/AbilityTask_ApplyRootMotionConstantForce.cpp
+++ b/Engine/Plugins/Runtime/GameplayAbilities/Source/GameplayAbilities/Private/Abilities/Tasks/AbilityTask_ApplyRootMotionConstantForce.cpp
@@ -78,11 +78,6 @@ void UAbilityTask_ApplyRootMotionConstantForce::SharedInitAndApply()
 				ConstantForce->Settings.SetFlag(ERootMotionSourceSettingsFlags::IgnoreZAccumulate);
 			}
 			RootMotionSourceID = MovementComponent->ApplyRootMotionSource(ConstantForce);
-
-			if (Ability)
-			{
-				Ability->SetMovementSyncPoint(ForceName);
-			}
 		}
 	}
 	else
diff --git a/Engine/Plugins/Runtime/GameplayAbilities/Source/GameplayAbilities/Private/Abilities/Tasks/AbilityTask_ApplyRootMotionJumpForce.cpp b/Engine/Plugins/Runtime/GameplayAbilities/Source/GameplayAbilities/Private/Abilities/Tasks/AbilityTask_ApplyRootMotionJumpForce.cpp
--- a/Engine/Plugins/Runtime/GameplayAbilities/Source/GameplayAbilities/Private/Abilities/Tasks/AbilityTask_ApplyRootMotionJumpForce.cpp
+++ b/Engine/Plugins/Runtime/GameplayAbilities/Source/GameplayAbilities/Private/Abilities/Tasks/AbilityTask_ApplyRootMotionJumpForce.cpp
@@ -115,11 +115,6 @@ void UAbilityTask_ApplyRootMotionJumpForce::SharedInitAndApply()
 			JumpForce->FinishVelocityParams.SetVelocity = FinishSetVelocity;
 			JumpForce->FinishVelocityParams.ClampVelocity = FinishClampVelocity;
 			RootMotionSourceID = MovementComponent->ApplyRootMotionSource(JumpForce);
-
-			if (Ability)
-			{
-				Ability->SetMovementSyncPoint(ForceName);
-			}
 		}
 	}
 	else
diff --git a/Engine/Plugins/Runtime/GameplayAbilities/Source/GameplayAbilities/Private/Abilities/Tasks/AbilityTask_ApplyRootMotionMoveToActorForce.cpp b/Engine/Plugins/Runtime/GameplayAbilities/Source/GameplayAbilities/Private/Abilities/Tasks/AbilityTask_ApplyRootMotionMoveToActorForce.cpp
--- a/Engine/Plugins/Runtime/GameplayAbilities/Source/GameplayAbilities/Private/Abilities/Tasks/AbilityTask_ApplyRootMotionMoveToActorForce.cpp
+++ b/Engine/Plugins/Runtime/GameplayAbilities/Source/GameplayAbilities/Private/Abilities/Tasks/AbilityTask_ApplyRootMotionMoveToActorForce.cpp
@@ -144,11 +144,6 @@ void UAbilityTask_ApplyRootMotionMoveToActorForce::SharedInitAndApply()
 			MoveToActorForce->FinishVelocityParams.SetVelocity = FinishSetVelocity;
 			MoveToActorForce->FinishVelocityParams.ClampVelocity = FinishClampVelocity;
 			RootMotionSourceID = MovementComponent->ApplyRootMotionSource(MoveToActorForce);
-
-			if (Ability)
-			{
-				Ability->SetMovementSyncPoint(ForceName);
-			}
 		}
 	}
 	else
diff --git a/Engine/Plugins/Runtime/GameplayAbilities/Source/GameplayAbilities/Private/Abilities/Tasks/AbilityTask_ApplyRootMotionMoveToForce.cpp b/Engine/Plugins/Runtime/GameplayAbilities/Source/GameplayAbilities/Private/Abilities/Tasks/AbilityTask_ApplyRootMotionMoveToForce.cpp
--- a/Engine/Plugins/Runtime/GameplayAbilities/Source/GameplayAbilities/Private/Abilities/Tasks/AbilityTask_ApplyRootMotionMoveToForce.cpp
+++ b/Engine/Plugins/Runtime/GameplayAbilities/Source/GameplayAbilities/Private/Abilities/Tasks/AbilityTask_ApplyRootMotionMoveToForce.cpp
@@ -75,11 +75,6 @@ void UAbilityTask_ApplyRootMotionMoveToForce::SharedInitAndApply()
 			MoveToForce->FinishVelocityParams.SetVelocity = FinishSetVelocity;
 			MoveToForce->FinishVelocityParams.ClampVelocity = FinishClampVelocity;
 			RootMotionSourceID = MovementComponent->ApplyRootMotionSource(MoveToForce);
-
-			if (Ability)
-			{
-				Ability->SetMovementSyncPoint(ForceName);
-			}
 		}
 	}
 	else
diff --git a/Engine/Plugins/Runtime/GameplayAbilities/Source/GameplayAbilities/Private/Abilities/Tasks/AbilityTask_ApplyRootMotionRadialForce.cpp b/Engine/Plugins/Runtime/GameplayAbilities/Source/GameplayAbilities/Private/Abilities/Tasks/AbilityTask_ApplyRootMotionRadialForce.cpp
--- a/Engine/Plugins/Runtime/GameplayAbilities/Source/GameplayAbilities/Private/Abilities/Tasks/AbilityTask_ApplyRootMotionRadialForce.cpp
+++ b/Engine/Plugins/Runtime/GameplayAbilities/Source/GameplayAbilities/Private/Abilities/Tasks/AbilityTask_ApplyRootMotionRadialForce.cpp
@@ -77,11 +77,6 @@ void UAbilityTask_ApplyRootMotionRadialForce::SharedInitAndApply()
 			RadialForce->FinishVelocityParams.SetVelocity = FinishSetVelocity;
 			RadialForce->FinishVelocityParams.ClampVelocity = FinishClampVelocity;
 			RootMotionSourceID = MovementComponent->ApplyRootMotionSource(RadialForce);
-
-			if (Ability)
-			{
-				Ability->SetMovementSyncPoint(ForceName);
-			}
 		}
 	}
 	else
diff --git a/Engine/Plugins/Runtime/GameplayAbilities/Source/GameplayAbilities/Public/Abilities/GameplayAbility.h b/Engine/Plugins/Runtime/GameplayAbilities/Source/GameplayAbilities/Public/Abilities/GameplayAbility.h
--- a/Engine/Plugins/Runtime/GameplayAbilities/Source/GameplayAbilities/Public/Abilities/GameplayAbility.h
+++ b/Engine/Plugins/Runtime/GameplayAbilities/Source/GameplayAbilities/Public/Abilities/GameplayAbility.h
@@ -404,6 +404,7 @@ public:
 	virtual void SetCurrentMontage(class UAnimMontage* InCurrentMontage);
 
 	/** Movement Sync */
+	UE_DEPRECATED(5.3, "This serves no purpose and will be removed in future engine versions")
 	virtual void SetMovementSyncPoint(FName SyncName);
 
 	// ----------------------------------------------------------------------------------------------------------------
