Adding a GameplayTagQuery field to FGameplayTagRequirements to enable more complex requirements to be specified.

diff --git a/Engine/Plugins/Runtime/GameplayAbilities/Source/GameplayAbilities/Private/AbilitySystemComponent.cpp b/Engine/Plugins/Runtime/GameplayAbilities/Source/GameplayAbilities/Private/AbilitySystemComponent.cpp
--- a/Engine/Plugins/Runtime/GameplayAbilities/Source/GameplayAbilities/Private/AbilitySystemComponent.cpp
+++ b/Engine/Plugins/Runtime/GameplayAbilities/Source/GameplayAbilities/Private/AbilitySystemComponent.cpp
@@ -1888,6 +1888,9 @@ float UAbilitySystemComponent::GetFilteredAttributeValue(const FGameplayAttribut
 {
 	float AttributeValue = 0.f;
 
+	UE_CLOG(!SourceTags.TagQuery.IsEmpty(), LogAbilitySystem, Error, TEXT("GetFilteredAttributeValue: SourceTags cannot contain a TagQuery (it is ignored)"));
+	UE_CLOG(!SourceTags.IgnoreTags.IsEmpty(), LogAbilitySystem, Error, TEXT("GetFilteredAttributeValue: SourceTags cannot contain IgnoreTags (they are not used)"));
+
 	if (SourceTags.RequireTags.Num() == 0 && SourceTags.IgnoreTags.Num() == 0 && HandlesToIgnore.Num() == 0)
 	{
 		// No qualifiers so we can just read this attribute normally
diff --git a/Engine/Plugins/Runtime/GameplayAbilities/Source/GameplayAbilities/Private/GameplayEffect.cpp b/Engine/Plugins/Runtime/GameplayAbilities/Source/GameplayAbilities/Private/GameplayEffect.cpp
--- a/Engine/Plugins/Runtime/GameplayAbilities/Source/GameplayAbilities/Private/GameplayEffect.cpp
+++ b/Engine/Plugins/Runtime/GameplayAbilities/Source/GameplayAbilities/Private/GameplayEffect.cpp
@@ -5393,20 +5393,17 @@ bool FGameplayModifierInfo::operator==(const FGameplayModifierInfo& Other) const
 		return false;
 	}
 
-	if (SourceTags.RequireTags.Num() != Other.SourceTags.RequireTags.Num() || !SourceTags.RequireTags.HasAllExact(Other.SourceTags.RequireTags))
-	{
-		return false;
-	}
-	if (SourceTags.IgnoreTags.Num() != Other.SourceTags.IgnoreTags.Num() || !SourceTags.IgnoreTags.HasAllExact(Other.SourceTags.IgnoreTags))
+	if (EvaluationChannelSettings != Other.EvaluationChannelSettings)
 	{
 		return false;
 	}
 
-	if (TargetTags.RequireTags.Num() != Other.TargetTags.RequireTags.Num() || !TargetTags.RequireTags.HasAllExact(Other.TargetTags.RequireTags))
+	if (SourceTags != Other.SourceTags)
 	{
 		return false;
 	}
-	if (TargetTags.IgnoreTags.Num() != Other.TargetTags.IgnoreTags.Num() || !TargetTags.IgnoreTags.HasAllExact(Other.TargetTags.IgnoreTags))
+
+	if (TargetTags != Other.TargetTags)
 	{
 		return false;
 	}
diff --git a/Engine/Plugins/Runtime/GameplayAbilities/Source/GameplayAbilities/Private/GameplayEffectComponents/TargetTagRequirementsGameplayEffectComponent.cpp b/Engine/Plugins/Runtime/GameplayAbilities/Source/GameplayAbilities/Private/GameplayEffectComponents/TargetTagRequirementsGameplayEffectComponent.cpp
--- a/Engine/Plugins/Runtime/GameplayAbilities/Source/GameplayAbilities/Private/GameplayEffectComponents/TargetTagRequirementsGameplayEffectComponent.cpp
+++ b/Engine/Plugins/Runtime/GameplayAbilities/Source/GameplayAbilities/Private/GameplayEffectComponents/TargetTagRequirementsGameplayEffectComponent.cpp
@@ -3,6 +3,7 @@
 #include "GameplayEffectComponents/TargetTagRequirementsGameplayEffectComponent.h"
 #include "AbilitySystemComponent.h"
 #include "AbilitySystemLog.h"
+#include "Algo/Find.h"
 #include "Misc/DataValidation.h"
 
 #define LOCTEXT_NAMESPACE "TargetTagRequirementsGameplayEffectComponent"
@@ -50,24 +51,43 @@ bool UTargetTagRequirementsGameplayEffectComponent::OnActiveGameplayEffectAdded(
 	FActiveGameplayEffectHandle ActiveGEHandle = ActiveGE.Handle;
 	if (FActiveGameplayEffectEvents* EventSet = ASC->GetActiveEffectEventSet(ActiveGEHandle))
 	{
-		// Keep track of everything I'm binding so that I can unbind it OnEffectRemoved
-		TArray<TTuple<FGameplayTag,FDelegateHandle>> AllBoundEvents;
-
-		// Add our tag requirements to the ASC's Callbacks map. This helps filter down the amount of callbacks we'll get due to tag changes
-		// (rather than registering for the one callback whenever any tag changes).  We also need to keep track to remove those registered delegates.
-		auto BindAllTags = [this, ASC, ActiveGEHandle, &AllBoundEvents](const FGameplayTagContainer& TagContainer)
+		// Quick method of appending a TArray to another TArray with no duplicates.
+		auto AppendUnique = [](TArray<FGameplayTag>& Destination, const TArray<FGameplayTag>& Source)
 		{
-			for (const FGameplayTag& Tag : TagContainer)
+			// Make sure the array won't allocate during the loop
+			if (Destination.GetSlack() < Source.Num())
+			{
+				Destination.Reserve(Destination.Num() + Source.Num());
+			}
+			const TConstArrayView<FGameplayTag> PreModifiedDestinationView{ Destination.GetData(), Destination.Num() };
+
+			for (const FGameplayTag& Tag : Source)
 			{
-				FOnGameplayEffectTagCountChanged& OnTagEvent = ASC->RegisterGameplayTagEvent(Tag, EGameplayTagEventType::NewOrRemoved);
-				FDelegateHandle Handle = OnTagEvent.AddUObject(this, &UTargetTagRequirementsGameplayEffectComponent::OnTagChanged, ActiveGEHandle);
-				AllBoundEvents.Emplace(Tag, Handle);
+				if (!Algo::Find(PreModifiedDestinationView, Tag))
+				{
+					Destination.Emplace(Tag);
+				}
 			}
 		};
-		BindAllTags(OngoingTagRequirements.IgnoreTags);
-		BindAllTags(OngoingTagRequirements.RequireTags);
-		BindAllTags(RemovalTagRequirements.IgnoreTags);
-		BindAllTags(RemovalTagRequirements.RequireTags);
+
+		// We should gather a list of tags to listen on events for
+		TArray<FGameplayTag> GameplayTagsToBind;
+		AppendUnique(GameplayTagsToBind, OngoingTagRequirements.IgnoreTags.GetGameplayTagArray());
+		AppendUnique(GameplayTagsToBind, OngoingTagRequirements.RequireTags.GetGameplayTagArray());
+		AppendUnique(GameplayTagsToBind, OngoingTagRequirements.TagQuery.GetGameplayTagArray());
+		AppendUnique(GameplayTagsToBind, RemovalTagRequirements.IgnoreTags.GetGameplayTagArray());
+		AppendUnique(GameplayTagsToBind, RemovalTagRequirements.RequireTags.GetGameplayTagArray());
+		AppendUnique(GameplayTagsToBind, RemovalTagRequirements.TagQuery.GetGameplayTagArray());
+
+		// Add our tag requirements to the ASC's Callbacks map. This helps filter down the amount of callbacks we'll get due to tag changes
+		// (rather than registering for the one callback whenever any tag changes).  We also need to keep track to remove those registered delegates in OnEffectRemoved.
+		TArray<TTuple<FGameplayTag, FDelegateHandle>> AllBoundEvents;
+		for (const FGameplayTag& Tag : GameplayTagsToBind)
+		{
+			FOnGameplayEffectTagCountChanged& OnTagEvent = ASC->RegisterGameplayTagEvent(Tag, EGameplayTagEventType::NewOrRemoved);
+			FDelegateHandle Handle = OnTagEvent.AddUObject(this, &UTargetTagRequirementsGameplayEffectComponent::OnTagChanged, ActiveGEHandle);
+			AllBoundEvents.Emplace(Tag, Handle);
+		}
 
 		// Now when this Effect is removed, we should remove all of our registered callbacks.
 		EventSet->OnEffectRemoved.AddUObject(this, &UTargetTagRequirementsGameplayEffectComponent::OnGameplayEffectRemoved, ASC, MoveTemp(AllBoundEvents));
diff --git a/Engine/Plugins/Runtime/GameplayAbilities/Source/GameplayAbilities/Private/GameplayEffectTypes.cpp b/Engine/Plugins/Runtime/GameplayAbilities/Source/GameplayAbilities/Private/GameplayEffectTypes.cpp
--- a/Engine/Plugins/Runtime/GameplayAbilities/Source/GameplayAbilities/Private/GameplayEffectTypes.cpp
+++ b/Engine/Plugins/Runtime/GameplayAbilities/Source/GameplayAbilities/Private/GameplayEffectTypes.cpp
@@ -90,6 +90,16 @@ void FGameplayModEvaluationChannelSettings::SetEvaluationChannel(EGameplayModEva
 	}
 }
 
+bool FGameplayModEvaluationChannelSettings::operator==(const FGameplayModEvaluationChannelSettings& Other) const
+{
+	return GetEvaluationChannel() == Other.GetEvaluationChannel();
+}
+
+bool FGameplayModEvaluationChannelSettings::operator!=(const FGameplayModEvaluationChannelSettings& Other) const
+{
+	return !(*this == Other);
+}
+
 float GameplayEffectUtilities::GetModifierBiasByModifierOp(EGameplayModOp::Type ModOp)
 {
 	static const float ModifierOpBiases[EGameplayModOp::Max] = {0.f, 1.f, 1.f, 0.f};
@@ -841,15 +851,16 @@ EGameplayTagEventType::Type FGameplayTagBlueprintPropertyMap::GetGameplayTagEven
 
 bool FGameplayTagRequirements::RequirementsMet(const FGameplayTagContainer& Container) const
 {
-	bool HasRequired = Container.HasAll(RequireTags);
-	bool HasIgnored = Container.HasAny(IgnoreTags);
+	const bool bHasRequired = Container.HasAll(RequireTags);
+	const bool bHasIgnored = Container.HasAny(IgnoreTags);
+	const bool bMatchQuery = TagQuery.IsEmpty() || TagQuery.Matches(Container);
 
-	return HasRequired && !HasIgnored;
+	return bHasRequired && !bHasIgnored && bMatchQuery;
 }
 
 bool FGameplayTagRequirements::IsEmpty() const
 {
-	return (RequireTags.Num() == 0 && IgnoreTags.Num() == 0);
+	return (RequireTags.Num() == 0 && IgnoreTags.Num() == 0 && TagQuery.IsEmpty());
 }
 
 FString FGameplayTagRequirements::ToString() const
@@ -864,13 +875,17 @@ FString FGameplayTagRequirements::ToString() const
 	{
 		Str += FString::Printf(TEXT("ignore: %s "), *IgnoreTags.ToStringSimple());
 	}
+	if (!TagQuery.IsEmpty())
+	{
+		Str += TagQuery.GetDescription();
+	}
 
 	return Str;
 }
 
 bool FGameplayTagRequirements::operator==(const FGameplayTagRequirements& Other) const
 {
-	return RequireTags == Other.RequireTags && IgnoreTags == Other.IgnoreTags;
+	return RequireTags == Other.RequireTags && IgnoreTags == Other.IgnoreTags && TagQuery == Other.TagQuery;
 }
 
 bool FGameplayTagRequirements::operator!=(const FGameplayTagRequirements& Other) const
@@ -878,6 +893,38 @@ bool FGameplayTagRequirements::operator!=(const FGameplayTagRequirements& Other)
 	return !(*this == Other);
 }
 
+FGameplayTagQuery FGameplayTagRequirements::ConvertTagFieldsToTagQuery() const
+{
+	const bool bHasRequireTags = !RequireTags.IsEmpty();
+	const bool bHasIgnoreTags = !IgnoreTags.IsEmpty();
+
+	if (!bHasIgnoreTags && !bHasRequireTags)
+	{
+		return FGameplayTagQuery{};
+	}
+
+	// FGameplayTagContainer::RequirementsMet is HasAll(RequireTags) && !HasAny(IgnoreTags);
+	FGameplayTagQueryExpression RequiredTagsQueryExpression = FGameplayTagQueryExpression().AllTagsMatch().AddTags(RequireTags);
+	FGameplayTagQueryExpression IgnoreTagsQueryExpression = FGameplayTagQueryExpression().NoTagsMatch().AddTags(IgnoreTags);
+
+	FGameplayTagQueryExpression RootQueryExpression;
+	if (bHasRequireTags && bHasIgnoreTags)
+	{
+		RootQueryExpression = FGameplayTagQueryExpression().AllExprMatch().AddExpr(RequiredTagsQueryExpression).AddExpr(IgnoreTagsQueryExpression);
+	}
+	else if (bHasRequireTags)
+	{
+		RootQueryExpression = RequiredTagsQueryExpression;
+	}
+	else // bHasIgnoreTags
+	{
+		RootQueryExpression = IgnoreTagsQueryExpression;
+	}
+
+	// Build the expression
+	return FGameplayTagQuery::BuildQuery(RootQueryExpression);
+}
+
 void FActiveGameplayEffectsContainer::PrintAllGameplayEffects() const
 {
 	for (const FActiveGameplayEffect& Effect : this)
diff --git a/Engine/Plugins/Runtime/GameplayAbilities/Source/GameplayAbilities/Private/Tests/GameplayTagQueryTests.cpp b/Engine/Plugins/Runtime/GameplayAbilities/Source/GameplayAbilities/Private/Tests/GameplayTagQueryTests.cpp
new file mode 100644
--- /dev/null
+++ b/Engine/Plugins/Runtime/GameplayAbilities/Source/GameplayAbilities/Private/Tests/GameplayTagQueryTests.cpp
@@ -0,0 +1,132 @@
+// Copyright Epic Games, Inc. All Rights Reserved.
+
+#include "CoreMinimal.h"
+#include "Misc/AutomationTest.h"
+#include "Misc/CommandLine.h"
+#include "Misc/Parse.h"
+
+#include "GameplayAbilitiesModule.h"
+#include "GameplayEffectTypes.h"
+
+IMPLEMENT_SIMPLE_AUTOMATION_TEST(FGameplayTagQueryMatchesTagRequirementsTest, "System.AbilitySystem.GameplayTagRequirements.ConvertTagFieldsToTagQuery", EAutomationTestFlags::ApplicationContextMask | EAutomationTestFlags::EngineFilter)
+
+#include "GameplayTagsManager.h"
+
+// You can run this on the command line like this to get a longer test
+//  UnrealEditor.exe -game -unattended -nullrhi -ExecCmds="Automation RunTest System.AbilitySystem.GameplayTagRequirements.ConvertTagFieldsToTagQuery;Now;Quit" -GameplayTagQueryMatchesTagRequirementsTest_NumTestsPerScenario=65536
+bool FGameplayTagQueryMatchesTagRequirementsTest::RunTest(const FString& Parameters)
+{
+	constexpr int NumDebugTags = 4;
+	constexpr int NumTestScenarios = 1024;
+	const int32 SomeSeed = FMath::GetRandSeed(); // Need to sync two random streams so we get the same results
+
+	int NumTestsPerScenario = 0;
+	FParse::Value(FCommandLine::Get(), TEXT("GameplayTagQueryMatchesTagRequirementsTest_NumTestsPerScenario="), NumTestsPerScenario);
+	if (NumTestsPerScenario <= 0)
+	{
+		NumTestsPerScenario = 1024;
+	}
+
+	// Stats to Collect
+	uint32 NumMetRequirements = 0;
+	uint64 MetRequirementsCycles = 0;
+
+	uint32 NumQueryMatches = 0;
+	uint64 QueryMatchesCycles = 0;
+
+	// Fill up possible tags with just some random tags from our project
+	TArray<FGameplayTag> PossibleTags;
+	{
+		FGameplayTagContainer AllGameplayTags;
+		UGameplayTagsManager::Get().RequestAllGameplayTags(AllGameplayTags, true);
+
+		if (AllGameplayTags.Num() < NumDebugTags)
+		{
+			AddWarning(FString::Printf(TEXT("There are only %d defined tags in the Project.  We need at least %d to run this test."), AllGameplayTags.Num(), NumDebugTags));
+			return true;
+		}
+
+		int Step = FMath::DivideAndRoundDown(AllGameplayTags.Num(), NumDebugTags);
+		for (int Index = 0; Index < NumDebugTags; ++Index)
+		{
+			FGameplayTag RandomGameplayTag = AllGameplayTags.GetByIndex(Index * Step);
+			PossibleTags.Add(RandomGameplayTag);
+		}
+	}
+
+	// Compute all permutations with those tags
+	TArray<FGameplayTagContainer> PossibleTagContainers;
+	for (int Index = 0; Index < NumDebugTags; ++Index)
+	{
+		FGameplayTagContainer Container;
+		Container.AddTag(PossibleTags[Index]);
+
+		for (int InnerIndex = Index; InnerIndex < NumDebugTags; ++InnerIndex)
+		{
+			Container.AddTag(PossibleTags[InnerIndex]);
+			PossibleTagContainers.Add(Container);
+		}
+	}
+	const int NumPossibleContainers = PossibleTagContainers.Num();
+
+	// Build up scenarios and then test them
+	for (int ScenarioNum = 0; ScenarioNum < NumTestScenarios; ++ScenarioNum)
+	{
+		// Assemble the TagRequirements
+		FGameplayTagRequirements TagReqs;
+		{
+			FRandomStream RandomStream{ FMath::GetRandSeed() };
+			for (int TagNum = 0; TagNum < NumDebugTags; ++TagNum)
+			{
+				FGameplayTag GameplayTag = PossibleTags[TagNum];
+				switch (RandomStream.RandHelper(3))
+				{
+				case 0:
+					TagReqs.RequireTags.AddTag(GameplayTag);
+					break;
+				case 1:
+					TagReqs.IgnoreTags.AddTag(GameplayTag);
+					break;
+				case 2:
+				default:
+					break;
+				};
+			}
+		}
+
+		// Do TagRequirements Test
+		{
+			uint64 StartTime = FPlatformTime::Cycles64();
+
+			FRandomStream RandomStream{ SomeSeed };
+			for (int i = 0; i < NumTestsPerScenario; ++i)
+			{
+				FGameplayTagContainer& TestAgainstTags = PossibleTagContainers[RandomStream.RandHelper(NumPossibleContainers)];
+				NumMetRequirements += TagReqs.RequirementsMet(TestAgainstTags) ? 1 : 0;
+			}
+
+			MetRequirementsCycles += FPlatformTime::Cycles64() - StartTime;
+		}
+
+		// Do TagQuery Test
+		FGameplayTagQuery TagQuery = TagReqs.ConvertTagFieldsToTagQuery();
+		{
+			uint64 StartTime = FPlatformTime::Cycles64();
+
+			FRandomStream RandomStream{ SomeSeed };
+			for (int i = 0; i < NumTestsPerScenario; ++i)
+			{
+				FGameplayTagContainer& TestAgainstTags = PossibleTagContainers[RandomStream.RandHelper(NumPossibleContainers)];
+				NumQueryMatches += TagQuery.Matches(TestAgainstTags) ? 1 : 0;
+			}
+
+			QueryMatchesCycles += FPlatformTime::Cycles64() - StartTime;
+		}
+	} // for each scenario
+
+	AddInfo(FString::Printf(TEXT(" Ran %d tests against FGameplayTagRequirements and FGameplayTagQuery in %.4lf ms"), NumTestScenarios * NumTestsPerScenario, FPlatformTime::ToMilliseconds64(MetRequirementsCycles + QueryMatchesCycles)));
+	AddInfo(FString::Printf(TEXT("	FGameplayTagRequirements took %ull cycles (%.4lf ms) and gave %d matches"), MetRequirementsCycles, FPlatformTime::ToMilliseconds64(MetRequirementsCycles), NumMetRequirements));
+	AddInfo(FString::Printf(TEXT("	FGameplayTagQuery took %ull cycles (%.4lf ms) and gave %d matches"), QueryMatchesCycles, FPlatformTime::ToMilliseconds64(QueryMatchesCycles), NumQueryMatches));
+
+	return true;
+}
diff --git a/Engine/Plugins/Runtime/GameplayAbilities/Source/GameplayAbilities/Public/GameplayEffectTypes.h b/Engine/Plugins/Runtime/GameplayAbilities/Source/GameplayAbilities/Public/GameplayEffectTypes.h
--- a/Engine/Plugins/Runtime/GameplayAbilities/Source/GameplayAbilities/Public/GameplayEffectTypes.h
+++ b/Engine/Plugins/Runtime/GameplayAbilities/Source/GameplayAbilities/Public/GameplayEffectTypes.h
@@ -90,6 +90,9 @@ struct GAMEPLAYABILITIES_API FGameplayModEvaluationChannelSettings
 
 	void SetEvaluationChannel(EGameplayModEvaluationChannel NewChannel);
 
+	bool operator==(const FGameplayModEvaluationChannelSettings& Other) const;
+	bool operator!=(const FGameplayModEvaluationChannelSettings& Other) const;
+
 protected:
 
 	/** Channel the settings would prefer to use, if possible/valid */
@@ -1369,6 +1372,10 @@ struct GAMEPLAYABILITIES_API FGameplayTagRequirements
 	UPROPERTY(EditAnywhere, BlueprintReadWrite, Category = GameplayModifier, meta=(DisplayName="Must Not Have Tags"))
 	FGameplayTagContainer IgnoreTags;
 
+	/** Build up a more complex query that can't be expressed with RequireTags/IgnoreTags alone */
+	UPROPERTY(EditAnywhere, BlueprintReadWrite, Category = GameplayModifier, meta = (DisplayName = "Query Must Match"))
+	FGameplayTagQuery TagQuery;
+
 	/** True if all required tags and no ignore tags found */
 	bool	RequirementsMet(const FGameplayTagContainer& Container) const;
 
@@ -1380,6 +1387,9 @@ struct GAMEPLAYABILITIES_API FGameplayTagRequirements
 
 	bool operator==(const FGameplayTagRequirements& Other) const;
 	bool operator!=(const FGameplayTagRequirements& Other) const;
+
+	/** Converts the RequireTags and IgnoreTags fields into an equivalent FGameplayTagQuery */
+	UE_NODISCARD FGameplayTagQuery ConvertTagFieldsToTagQuery() const;
 };
 
 
diff --git a/Engine/Source/Developer/AutomationController/Private/AutomationCommandline.cpp b/Engine/Source/Developer/AutomationController/Private/AutomationCommandline.cpp
--- a/Engine/Source/Developer/AutomationController/Private/AutomationCommandline.cpp
+++ b/Engine/Source/Developer/AutomationController/Private/AutomationCommandline.cpp
@@ -521,18 +521,23 @@ protected:
 				{
 					AutomationCommandQueue.Add(EAutomationCommand::ListAllTests);
 				}
+				else if (FParse::Command(&TempCmd, TEXT("Now")))
+				{
+					DelayTimer = 0.0f;
+				}
 				else if (FParse::Command(&TempCmd, TEXT("RunTests")) || FParse::Command(&TempCmd, TEXT("RunTest")))
 				{
-					//only one of these should be used
-					if (IsRunTestQueued())
+					if (FParse::Command(&TempCmd, TEXT("Now")))
 					{
-						Ar.Logf(TEXT("Automation: A test run is already Queued: %s. Only one run is supported at a time."), *StringCommand);
+						DelayTimer = 0.0f;
 						continue;
 					}
 
-					if ( FParse::Command(&TempCmd, TEXT("Now")) )
+					//only one of these should be used
+					if (IsRunTestQueued())
 					{
-						DelayTimer = 0.0f;
+						Ar.Logf(TEXT("Automation: A test run is already Queued: %s. Only one run is supported at a time."), *StringCommand);
+						continue;
 					}
 
 					StringCommand = TempCmd;
@@ -664,6 +669,7 @@ protected:
 					Ar.Logf(TEXT("\tAutomation SetFilter <filter name>"));
 					Ar.Logf(TEXT("\tAutomation SetMinimumPriority <minimum priority>"));
 					Ar.Logf(TEXT("\tAutomation SetPriority <priority>"));
+					Ar.Logf(TEXT("\tAutomation Now"));
 					Ar.Logf(TEXT("\tAutomation Quit"));
 					Ar.Logf(TEXT("\tAutomation SoftQuit"));
 					bHandled = false;
diff --git a/Engine/Source/Runtime/GameplayTags/Classes/GameplayTagContainer.h b/Engine/Source/Runtime/GameplayTags/Classes/GameplayTagContainer.h
--- a/Engine/Source/Runtime/GameplayTags/Classes/GameplayTagContainer.h
+++ b/Engine/Source/Runtime/GameplayTags/Classes/GameplayTagContainer.h
@@ -552,6 +552,9 @@ struct GAMEPLAYTAGS_API FGameplayTagContainer
 		InOutGameplayTags = GameplayTags;
 	}
 
+	/** Gets the explicit list of gameplay tags */
+	const TArray<FGameplayTag>& GetGameplayTagArray() const;
+
 	/** Creates a const iterator for the contents of this array */
 	TArray<FGameplayTag>::TConstIterator CreateConstIterator() const
 	{
@@ -882,12 +885,15 @@ public:
 	/** Returns description string. */
 	const FString& GetDescription() const { return UserDescription.IsEmpty() ? AutoDescription : UserDescription; };
 
-	//** Gets the explicit list of all unique gameplay tags referenced by the query. */
+	/** Gets the explicit list of all unique gameplay tags referenced by the query. */
 	void GetGameplayTagArray(TArray<FGameplayTag>& OutGameplayTags) const
 	{
 		OutGameplayTags = TagDictionary;
 	}
 
+	/** Gets the explicit list of all unique gameplay tags referenced by the query. */
+	const TArray<FGameplayTag>& GetGameplayTagArray() const;
+
 #if WITH_EDITOR
 	/** Creates this query based on the given EditableQuery object */
 	void BuildFromEditableQuery(class UEditableGameplayTagQuery& EditableQuery); 
diff --git a/Engine/Source/Runtime/GameplayTags/Private/GameplayTagContainer.cpp b/Engine/Source/Runtime/GameplayTags/Private/GameplayTagContainer.cpp
--- a/Engine/Source/Runtime/GameplayTags/Private/GameplayTagContainer.cpp
+++ b/Engine/Source/Runtime/GameplayTags/Private/GameplayTagContainer.cpp
@@ -507,6 +507,10 @@ bool FQueryEvaluator::EvalExpr(FGameplayTagContainer const& Tags, bool bSkip)
 	return false;
 }
 
+const TArray<FGameplayTag>& FGameplayTagContainer::GetGameplayTagArray() const
+{
+	return GameplayTags;
+}
 
 FGameplayTagContainer& FGameplayTagContainer::operator=(FGameplayTagContainer const& Other)
 {
@@ -539,15 +543,7 @@ bool FGameplayTagContainer::operator==(FGameplayTagContainer const& Other) const
 		return false;
 	}
 
-	for (const FGameplayTag& Tag : GameplayTags)
-	{
-		if (!Tag.MatchesAnyExact(Other))
-		{
-			return false;
-		}
-	}
-
-	return true;
+	return HasAllExact(Other);
 }
 
 bool FGameplayTagContainer::operator!=(FGameplayTagContainer const& Other) const
@@ -1413,6 +1409,11 @@ void FGameplayTagQuery::GetQueryExpr(FGameplayTagQueryExpression& OutExpr) const
 	QE.Read(OutExpr);
 }
 
+const TArray<FGameplayTag>& FGameplayTagQuery::GetGameplayTagArray() const
+{
+	return TagDictionary;
+}
+	
 void FGameplayTagQuery::Serialize(FArchive& Ar)
 {
 	Ar << TokenStreamVersion;
