Deprecated UGameplayEffect::GrantedAbilities and FGameplayEffectSpec::GrantedAbilities in favor of AbilitiesGameplayEffectComponent.

diff --git a/Engine/Plugins/Runtime/GameplayAbilities/Source/GameplayAbilities/Private/AbilitySystemComponent.cpp b/Engine/Plugins/Runtime/GameplayAbilities/Source/GameplayAbilities/Private/AbilitySystemComponent.cpp
--- a/Engine/Plugins/Runtime/GameplayAbilities/Source/GameplayAbilities/Private/AbilitySystemComponent.cpp
+++ b/Engine/Plugins/Runtime/GameplayAbilities/Source/GameplayAbilities/Private/AbilitySystemComponent.cpp
@@ -1127,13 +1127,22 @@ FActiveGameplayEffectHandle UAbilitySystemComponent::FindActiveGameplayEffectHan
 {
 	for (const FActiveGameplayEffect& ActiveGE : &ActiveGameplayEffects)
 	{
+		// Old, deprecated way of handling these (before AbilitiesGameplayEffectComponent):
+PRAGMA_DISABLE_DEPRECATION_WARNINGS
 		for (const FGameplayAbilitySpecDef& AbilitySpecDef : ActiveGE.Spec.GrantedAbilitySpecs)
+PRAGMA_ENABLE_DEPRECATION_WARNINGS
 		{
 			if (AbilitySpecDef.AssignedHandle == Handle)
 			{
 				return ActiveGE.Handle;
 			}
 		}
+
+		// Where AbilitiesGameplayEffectComponent stores its data
+		if (ActiveGE.GrantedAbilityHandles.Contains(Handle))
+		{
+			return ActiveGE.Handle;
+		}
 	}
 	return FActiveGameplayEffectHandle();
 }
diff --git a/Engine/Plugins/Runtime/GameplayAbilities/Source/GameplayAbilities/Private/AbilitySystemComponent_Abilities.cpp b/Engine/Plugins/Runtime/GameplayAbilities/Source/GameplayAbilities/Private/AbilitySystemComponent_Abilities.cpp
--- a/Engine/Plugins/Runtime/GameplayAbilities/Source/GameplayAbilities/Private/AbilitySystemComponent_Abilities.cpp
+++ b/Engine/Plugins/Runtime/GameplayAbilities/Source/GameplayAbilities/Private/AbilitySystemComponent_Abilities.cpp
@@ -501,6 +501,22 @@ void UAbilitySystemComponent::OnGiveAbility(FGameplayAbilitySpec& Spec)
 		}
 	}
 
+	// If this Ability Spec specified that it was created from an Active Gameplay Effect, then link the handle to the Active Gameplay Effect.
+	if (Spec.GameplayEffectHandle.IsValid())
+	{
+		UAbilitySystemComponent* SourceASC = Spec.GameplayEffectHandle.GetOwningAbilitySystemComponent();
+		UE_CLOG(!SourceASC, LogAbilitySystem, Error, TEXT("OnGiveAbility Spec '%s' GameplayEffectHandle had invalid Owning Ability System Component"), *Spec.GetDebugString());
+		if (SourceASC)
+		{
+			FActiveGameplayEffect* SourceActiveGE = SourceASC->ActiveGameplayEffects.GetActiveGameplayEffect(Spec.GameplayEffectHandle);
+			UE_CLOG(!SourceActiveGE, LogAbilitySystem, Error, TEXT("OnGiveAbility Spec '%s' GameplayEffectHandle was not active on Owning Ability System Component '%s'"), *Spec.GetDebugString(), *SourceASC->GetName());
+			if (SourceActiveGE)
+			{
+				SourceActiveGE->GrantedAbilityHandles.AddUnique(Spec.Handle);
+			}
+		}
+	}
+
 	for (const FAbilityTriggerData& TriggerData : Spec.Ability->AbilityTriggers)
 	{
 		FGameplayTag EventTag = TriggerData.TriggerTag;
@@ -605,6 +621,22 @@ void UAbilitySystemComponent::OnRemoveAbility(FGameplayAbilitySpec& Spec)
 		Spec.Ability->OnRemoveAbility(AbilityActorInfo.Get(), Spec);
 	}
 
+	// If this Ability Spec specified that it was created from an Active Gameplay Effect, then unlink the handle to the Active Gameplay Effect.
+	// Note: It's possible (maybe even likely) that the ActiveGE is no longer considered active by this point.
+	if (Spec.GameplayEffectHandle.IsValid())
+	{
+		UAbilitySystemComponent* SourceASC = Spec.GameplayEffectHandle.GetOwningAbilitySystemComponent();
+		UE_CLOG(!SourceASC, LogAbilitySystem, Error, TEXT("OnRemoveAbility Spec '%s' GameplayEffectHandle had invalid Owning Ability System Component"), *Spec.GetDebugString());
+		if (SourceASC)
+		{
+			FActiveGameplayEffect* SourceActiveGE = SourceASC->ActiveGameplayEffects.GetActiveGameplayEffect(Spec.GameplayEffectHandle);
+			if (SourceActiveGE)
+			{
+				SourceActiveGE->GrantedAbilityHandles.Remove(Spec.Handle);
+			}
+		}
+	}
+
 	Spec.ReplicatedInstances.Empty();
 	Spec.NonReplicatedInstances.Empty();
 }
@@ -678,7 +710,9 @@ void UAbilitySystemComponent::CheckForClearedAbilities()
 	// Clear any out of date ability spec handles on active gameplay effects
 	for (FActiveGameplayEffect& ActiveGE : &ActiveGameplayEffects)
 	{
+PRAGMA_DISABLE_DEPRECATION_WARNINGS
 		for (FGameplayAbilitySpecDef& AbilitySpec : ActiveGE.Spec.GrantedAbilitySpecs)
+PRAGMA_ENABLE_DEPRECATION_WARNINGS
 		{
 			if (AbilitySpec.AssignedHandle.IsValid() && FindAbilitySpecFromHandle(AbilitySpec.AssignedHandle) == nullptr)
 			{
diff --git a/Engine/Plugins/Runtime/GameplayAbilities/Source/GameplayAbilities/Private/GameplayEffect.cpp b/Engine/Plugins/Runtime/GameplayAbilities/Source/GameplayAbilities/Private/GameplayEffect.cpp
--- a/Engine/Plugins/Runtime/GameplayAbilities/Source/GameplayAbilities/Private/GameplayEffect.cpp
+++ b/Engine/Plugins/Runtime/GameplayAbilities/Source/GameplayAbilities/Private/GameplayEffect.cpp
@@ -22,7 +22,9 @@
 #include "GameplayEffectCustomApplicationRequirement.h"
 #include "GameplayEffectUIData.h"
 #include "Misc/DataValidation.h"
+#include "HAL/IConsoleManager.h"
 
+#include "GameplayEffectComponents/AbilitiesGameplayEffectComponent.h"
 #include "GameplayEffectComponents/AssetTagsGameplayEffectComponent.h"
 #include "GameplayEffectComponents/BlockAbilityTagsGameplayEffectComponent.h"
 #include "GameplayEffectComponents/ChanceToApplyGameplayEffectComponent.h"
@@ -63,6 +65,28 @@ DECLARE_CYCLE_STAT(TEXT("MakeQuery"), STAT_MakeGameplayEffectQuery, STATGROUP_Ab
 	SCALABLEFLOAT_REPORTERROR_WITHPATHNAME(Scalable, PathNameString);
 #endif // WITH_EDITOR
 
+namespace UE::GameplayEffect
+{
+	TAutoConsoleVariable<int32> CVarGameplayEffectMaxVersion(TEXT("AbilitySystem.GameplayEffects.MaxVersion"), (int32)EGameplayEffectVersion::Current, TEXT("Override the Gameplay Effect Current Version (disabling upgrade code paths)"), ECVF_Default);
+
+#if WITH_EDITOR
+	namespace EditorOnly
+	{
+		// Helper function that allows us to guard against upgrading versions (in case a bug is discovered while the product is live)
+		static bool ShouldUpgradeVersion(EGameplayEffectVersion FromVersion, EGameplayEffectVersion ToVersion)
+		{
+			EGameplayEffectVersion MaxVersion = static_cast<EGameplayEffectVersion>(CVarGameplayEffectMaxVersion.GetValueOnGameThread() & 0xff);
+			if (MaxVersion < ToVersion)
+			{
+				return false;
+			}
+
+			return FromVersion < ToVersion;
+		}
+	}
+#endif
+}
+
 // Versioning code
 bool FGameplayEffectVersion::Serialize(FStructuredArchive::FSlot Slot)
 {
@@ -200,7 +224,7 @@ void UGameplayEffect::PostLoad()
 	IsDataValid(DataValidationContext);
 
 	// We're done loading (and therefore upgrading), boost the version.
-	SetVersion(EGameplayEffectVersion::Current);
+	SetVersion( static_cast<EGameplayEffectVersion>(UE::GameplayEffect::CVarGameplayEffectMaxVersion.GetValueOnGameThread()) );
 #endif // WITH_EDITOR
 }
 
@@ -294,6 +318,7 @@ void UGameplayEffect::PostCDOCompiled(const FPostCDOCompiledContext& Context)
 	PostCDOCompiledFixupSubobjects();
 
 	// Move any data from deprecated properties into components, create components as needed
+	ConvertAbilitiesComponent();
 	ConvertAdditionalEffectsComponent();
 	ConvertAssetTagsComponent();
 	ConvertBlockByTagsComponent();
@@ -312,21 +337,44 @@ void UGameplayEffect::PostCDOCompiled(const FPostCDOCompiledContext& Context)
 	}
 }
 
-bool UGameplayEffect::VersionAllowsCreateGEComponentsDuringUpgrade() const
-{
-	return (GetVersion() < EGameplayEffectVersion::Modular53);
-}
-
 // Upgrade Code Section.  Lots of uses of deprecated variables (because we're upgrading them)
 // So just disable it throughout this section.
 PRAGMA_DISABLE_DEPRECATION_WARNINGS
+void UGameplayEffect::ConvertAbilitiesComponent()
+{
+	// Get the archetype to compare to
+	const UGameplayEffect* Archetype = CastChecked<UGameplayEffect>(GetArchetype());
+
+	const bool bChanged = !(GrantedAbilities == Archetype->GrantedAbilities);
+	if (bChanged && UE::GameplayEffect::EditorOnly::ShouldUpgradeVersion(GetVersion(), EGameplayEffectVersion::AbilitiesComponent53))
+	{
+		UAbilitiesGameplayEffectComponent& AbilitiesComponent = FindOrAddComponent<UAbilitiesGameplayEffectComponent>();
+		for (const FGameplayAbilitySpecDef& GASpecDef : GrantedAbilities)
+		{
+			FGameplayAbilitySpecConfig GASpecConfig;
+			GASpecConfig.Ability = GASpecDef.Ability;
+			GASpecConfig.InputID = GASpecDef.InputID;
+			GASpecConfig.LevelScalableFloat = GASpecDef.LevelScalableFloat;
+			GASpecConfig.RemovalPolicy = GASpecDef.RemovalPolicy;
+
+			UE_CLOG(GASpecConfig.LevelScalableFloat.GetValue() <= 0.0f, LogGameplayEffects, Error, TEXT("Converted AbilitiesComponent on %s has scale %.2f"), *GetName(), GASpecConfig.LevelScalableFloat.GetValue());
+
+			// This function is smart enough to not add duplicates (handles the inheritance case)
+			AbilitiesComponent.AddGrantedAbilityConfig(GASpecConfig);
+		}
+
+		// We should empty these because the old code path that executes them is still around.
+		GrantedAbilities.Empty();
+	}
+}
+
 void UGameplayEffect::ConvertCustomCanApplyComponent()
 {
 	// Get the archetype to compare to
 	const UGameplayEffect* Archetype = CastChecked<UGameplayEffect>(GetArchetype());
 
 	const bool bChanged = !(ApplicationRequirements_DEPRECATED == Archetype->ApplicationRequirements_DEPRECATED);
-	if (bChanged && VersionAllowsCreateGEComponentsDuringUpgrade())
+	if (bChanged && UE::GameplayEffect::EditorOnly::ShouldUpgradeVersion(GetVersion(), EGameplayEffectVersion::Modular53))
 	{
 		UCustomCanApplyGameplayEffectComponent& ApplicationComponent = FindOrAddComponent<UCustomCanApplyGameplayEffectComponent>();
 		ApplicationComponent.ApplicationRequirements.Append(ApplicationRequirements_DEPRECATED);
@@ -348,7 +396,7 @@ void UGameplayEffect::ConvertChanceToApplyComponent()
 	bool bChanged = (ChanceToApplyToTarget_DEPRECATED != Archetype->ChanceToApplyToTarget_DEPRECATED);
 	bChanged = bChanged || !(ApplicationRequirements_DEPRECATED == Archetype->ApplicationRequirements_DEPRECATED);
 
-	if (bChanged && VersionAllowsCreateGEComponentsDuringUpgrade())
+	if (bChanged && UE::GameplayEffect::EditorOnly::ShouldUpgradeVersion(GetVersion(), EGameplayEffectVersion::Modular53))
 	{
 		UChanceToApplyGameplayEffectComponent& ChanceToApplyComponent = FindOrAddComponent<UChanceToApplyGameplayEffectComponent>();
 		ChanceToApplyComponent.SetChanceToApplyToTarget(ChanceToApplyToTarget_DEPRECATED);
@@ -368,7 +416,7 @@ void UGameplayEffect::ConvertAssetTagsComponent()
 	const UGameplayEffect* Archetype = Cast<UGameplayEffect>(GetArchetype());
 
 	bool bChanged = InheritableGameplayEffectTags.CombinedTags != Archetype->InheritableGameplayEffectTags.CombinedTags;
-	if (bChanged && VersionAllowsCreateGEComponentsDuringUpgrade())
+	if (bChanged && UE::GameplayEffect::EditorOnly::ShouldUpgradeVersion(GetVersion(), EGameplayEffectVersion::Modular53))
 	{
 		UAssetTagsGameplayEffectComponent& AssetTagsComponent = FindOrAddComponent<UAssetTagsGameplayEffectComponent>();
 		AssetTagsComponent.SetAndApplyAssetTagChanges(InheritableGameplayEffectTags);
@@ -393,7 +441,7 @@ void UGameplayEffect::ConvertAdditionalEffectsComponent()
 	bool bRoutineChanged = (RoutineExpirationEffectClasses != Archetype->RoutineExpirationEffectClasses);
 
 	bool bChanged = bConditionalChanged || bPrematureChanged || bRoutineChanged;
-	if (bChanged && VersionAllowsCreateGEComponentsDuringUpgrade())
+	if (bChanged && UE::GameplayEffect::EditorOnly::ShouldUpgradeVersion(GetVersion(), EGameplayEffectVersion::Modular53))
 	{
 		UAdditionalEffectsGameplayEffectComponent& ConditionalEffectsComponent = FindOrAddComponent<UAdditionalEffectsGameplayEffectComponent>();
 
@@ -429,7 +477,7 @@ void UGameplayEffect::ConvertBlockByTagsComponent()
 	const UGameplayEffect* Archetype = Cast<UGameplayEffect>(GetArchetype());
 
 	bool bChanged = InheritableBlockedAbilityTagsContainer.CombinedTags != Archetype->InheritableBlockedAbilityTagsContainer.CombinedTags;
-	if (bChanged && VersionAllowsCreateGEComponentsDuringUpgrade())
+	if (bChanged && UE::GameplayEffect::EditorOnly::ShouldUpgradeVersion(GetVersion(), EGameplayEffectVersion::Modular53))
 	{
 		UBlockAbilityTagsGameplayEffectComponent& BlockByTagsComponent = FindOrAddComponent<UBlockAbilityTagsGameplayEffectComponent>();
 		BlockByTagsComponent.SetAndApplyBlockedAbilityTagChanges(InheritableBlockedAbilityTagsContainer);
@@ -450,7 +498,7 @@ void UGameplayEffect::ConvertTargetTagsComponent()
 	const UGameplayEffect* Archetype = Cast<UGameplayEffect>(GetArchetype());
 
 	bool bChanged = InheritableOwnedTagsContainer.CombinedTags != Archetype->InheritableOwnedTagsContainer.CombinedTags;
-	if (bChanged && VersionAllowsCreateGEComponentsDuringUpgrade())
+	if (bChanged && UE::GameplayEffect::EditorOnly::ShouldUpgradeVersion(GetVersion(), EGameplayEffectVersion::Modular53))
 	{
 		UTargetTagsGameplayEffectComponent& TargetTagsComponent = FindOrAddComponent<UTargetTagsGameplayEffectComponent>();
 		TargetTagsComponent.SetAndApplyTargetTagChanges(InheritableOwnedTagsContainer);
@@ -474,7 +522,7 @@ void UGameplayEffect::ConvertImmunityComponent()
 	const bool bQueryChanged (GrantedApplicationImmunityQuery != Archetype->GrantedApplicationImmunityQuery);
 
 	const bool bChanged = bTagsChanged || bQueryChanged;
-	if (bChanged && VersionAllowsCreateGEComponentsDuringUpgrade())
+	if (bChanged && UE::GameplayEffect::EditorOnly::ShouldUpgradeVersion(GetVersion(), EGameplayEffectVersion::Modular53))
 	{
 		UImmunityGameplayEffectComponent& ImmunityComponent = FindOrAddComponent<UImmunityGameplayEffectComponent>();
 
@@ -542,7 +590,7 @@ void UGameplayEffect::ConvertRemoveOtherComponent()
 	const bool bQueryChanged = RemoveGameplayEffectQuery != Archetype->RemoveGameplayEffectQuery;
 
 	const bool bChanged = bTagsChanged || bQueryChanged;
-	if (bChanged && VersionAllowsCreateGEComponentsDuringUpgrade())
+	if (bChanged && UE::GameplayEffect::EditorOnly::ShouldUpgradeVersion(GetVersion(), EGameplayEffectVersion::Modular53))
 	{
 		URemoveOtherGameplayEffectComponent& RemoveOtherComponent = FindOrAddComponent<URemoveOtherGameplayEffectComponent>();
 		
@@ -584,7 +632,7 @@ void UGameplayEffect::ConvertTagRequirementsComponent()
 	bool bOngoingChanged = OngoingTagRequirements != Archetype->OngoingTagRequirements;
 
 	const bool bChanged = bApplicationChanged || bRemovalChanged || bOngoingChanged;
-	if (bChanged && VersionAllowsCreateGEComponentsDuringUpgrade())
+	if (bChanged && UE::GameplayEffect::EditorOnly::ShouldUpgradeVersion(GetVersion(), EGameplayEffectVersion::Modular53))
 	{
 		UTargetTagRequirementsGameplayEffectComponent& TagRequirementsComponent = FindOrAddComponent<UTargetTagRequirementsGameplayEffectComponent>();
 
@@ -620,7 +668,7 @@ void UGameplayEffect::ConvertUIComponent()
 	const UGameplayEffect* Archetype = Cast<UGameplayEffect>(GetArchetype());
 
 	const bool bChanged = (UIData != Archetype->UIData);
-	if (bChanged && VersionAllowsCreateGEComponentsDuringUpgrade())
+	if (bChanged && UE::GameplayEffect::EditorOnly::ShouldUpgradeVersion(GetVersion(), EGameplayEffectVersion::Modular53))
 	{
 		const UGameplayEffectUIData* UIComponent = FindComponent<UGameplayEffectUIData>();
 		if (!UIComponent && UIData)
@@ -1376,11 +1424,13 @@ void FGameplayEffectSpec::Initialize(const UGameplayEffect* InDef, const FGamepl
 	// Prepare source tags before accessing them in the Components
 	CaptureDataFromSource();
 
+PRAGMA_DISABLE_DEPRECATION_WARNINGS
 	// ------------------------------------------------
 	//	Granted Abilities
 	// ------------------------------------------------
 
 	// Make Granted AbilitySpecs (caller may modify these specs after creating spec, which is why we dont just reference them from the def)
+	// Note: These are going to be deprecated in the future in favor of immutable GameplayEffectComponents.
 	GrantedAbilitySpecs = InDef->GrantedAbilities;
 
 	// if we're granting abilities and they don't specify a source object use the source of this GE
@@ -1391,6 +1441,7 @@ void FGameplayEffectSpec::Initialize(const UGameplayEffect* InDef, const FGamepl
 			AbilitySpecDef.SourceObject = InEffectContext.GetSourceObject();
 		}
 	}
+PRAGMA_ENABLE_DEPRECATION_WARNINGS
 }
 
 void FGameplayEffectSpec::InitializeFromLinkedSpec(const UGameplayEffect* InDef, const FGameplayEffectSpec& OriginalSpec)
@@ -3547,25 +3598,24 @@ FActiveGameplayEffect* FActiveGameplayEffectsContainer::ApplyGameplayEffectSpec(
 		// re-registered later below, as necessary.
 		ExistingSpec.CapturedRelevantAttributes.UnregisterLinkedAggregatorCallbacks(ExistingStackableGE->Handle);
 
-		// @todo: If dynamically granted tags differ (which they shouldn't), we'll actually have to diff them
+		// @todo: If dynamically granted/asset tags differ (which they shouldn't), we'll actually have to diff them
 		// and cause a removal and add of only the ones that have changed. For now, ensure on this happening and come
 		// back to this later.
 		ensureMsgf(ExistingSpec.DynamicGrantedTags == Spec.DynamicGrantedTags, TEXT("While adding a stack of the gameplay effect: %s, the old stack and the new application had different dynamically granted tags, which is currently not resolved properly!"), *Spec.Def->GetName());
-		
+		ensureMsgf(ExistingSpec.GetDynamicAssetTags() == Spec.GetDynamicAssetTags(), TEXT("While adding a stack of the gameplay effect: %s, the old stack and the new application had different dynamic asset tags, which is currently not resolved properly! Existing: %s. New: %s"), *Spec.Def->GetName(), *ExistingSpec.GetDynamicAssetTags().ToStringSimple(), *Spec.GetDynamicAssetTags().ToStringSimple());
+
+PRAGMA_DISABLE_DEPRECATION_WARNINGS
 		// We only grant abilities on the first apply. So we *dont* want the new spec's GrantedAbilitySpecs list
 		TArray<FGameplayAbilitySpecDef>	GrantedSpecTempArray(MoveTemp(ExistingStackableGE->Spec.GrantedAbilitySpecs));
 
-		// @todo: If dynamic asset tags differ (which they shouldn't), we'll actually have to diff them
-		// and cause a removal and add of only the ones that have changed. For now, ensure on this happening and come
-		// back to this later.
-		ensureMsgf(ExistingSpec.GetDynamicAssetTags() == Spec.GetDynamicAssetTags(), TEXT("While adding a stack of the gameplay effect: %s, the old stack and the new application had different dynamic asset tags, which is currently not resolved properly! Existing: %s. New: %s"), *Spec.Def->GetName(), *ExistingSpec.GetDynamicAssetTags().ToStringSimple(), *Spec.GetDynamicAssetTags().ToStringSimple() );
-
 		ExistingStackableGE->Spec = Spec;
-		ExistingStackableGE->Spec.SetStackCount(NewStackCount);
 
 		// Swap in old granted ability spec
 		ExistingStackableGE->Spec.GrantedAbilitySpecs = MoveTemp(GrantedSpecTempArray);
-		
+PRAGMA_ENABLE_DEPRECATION_WARNINGS
+
+		ExistingStackableGE->Spec.SetStackCount(NewStackCount);
+
 		AppliedActiveGE = ExistingStackableGE;
 
 		const UGameplayEffect* GEDef = ExistingSpec.Def;
@@ -3874,7 +3924,7 @@ void FActiveGameplayEffectsContainer::AddActiveGameplayEffectGrantedTagsAndModif
 		Owner->AddMinimalReplicationGameplayTags(Effect.Spec.DynamicGrantedTags);
 	}
 
-
+PRAGMA_DISABLE_DEPRECATION_WARNINGS
 	// Grant abilities
 	if (IsNetAuthority() && !Owner->bSuppressGrantAbility)
 	{
@@ -3891,6 +3941,7 @@ void FActiveGameplayEffectsContainer::AddActiveGameplayEffectGrantedTagsAndModif
 			}
 		}	
 	}
+PRAGMA_ENABLE_DEPRECATION_WARNINGS
 
 	// Update GameplayCue tags and events
 	if (!Owner->bSuppressGameplayCues)
@@ -4162,6 +4213,7 @@ void FActiveGameplayEffectsContainer::RemoveActiveGameplayEffectGrantedTagsAndMo
 		Owner->RemoveMinimalReplicationGameplayTags(Effect.Spec.DynamicGrantedTags);
 	}
 
+PRAGMA_DISABLE_DEPRECATION_WARNINGS
 	// Cancel/remove granted abilities
 	if (IsNetAuthority())
 	{
@@ -4190,6 +4242,7 @@ void FActiveGameplayEffectsContainer::RemoveActiveGameplayEffectGrantedTagsAndMo
 			}
 		}
 	}
+PRAGMA_ENABLE_DEPRECATION_WARNINGS
 
 	// Update GameplayCue tags and events
 	if (!Owner->bSuppressGameplayCues)
diff --git a/Engine/Plugins/Runtime/GameplayAbilities/Source/GameplayAbilities/Private/GameplayEffectComponents/AbilitiesGameplayEffectComponent.cpp b/Engine/Plugins/Runtime/GameplayAbilities/Source/GameplayAbilities/Private/GameplayEffectComponents/AbilitiesGameplayEffectComponent.cpp
new file mode 100644
--- /dev/null
+++ b/Engine/Plugins/Runtime/GameplayAbilities/Source/GameplayAbilities/Private/GameplayEffectComponents/AbilitiesGameplayEffectComponent.cpp
@@ -0,0 +1,208 @@
+// Copyright Epic Games, Inc. All Rights Reserved.
+
+#include "GameplayEffectComponents/AbilitiesGameplayEffectComponent.h"
+#include "AbilitySystemComponent.h"
+#include "AbilitySystemLog.h"
+#include "Misc/DataValidation.h"
+
+#define LOCTEXT_NAMESPACE "AbilitiesGameplayEffectComponent"
+
+bool operator==(const FGameplayAbilitySpecConfig& Lhs, const FGameplayAbilitySpecConfig& Rhs)
+{
+	return Lhs.Ability == Rhs.Ability &&
+		Lhs.InputID == Rhs.InputID &&
+		Lhs.LevelScalableFloat == Rhs.LevelScalableFloat &&
+		Lhs.RemovalPolicy == Rhs.RemovalPolicy;
+}
+
+UAbilitiesGameplayEffectComponent::UAbilitiesGameplayEffectComponent()
+{
+#if WITH_EDITORONLY_DATA
+	EditorFriendlyName = TEXT("Grant Abilities While Active");
+#endif
+}
+
+void UAbilitiesGameplayEffectComponent::AddGrantedAbilityConfig(const FGameplayAbilitySpecConfig& Config)
+{
+	for (const FGameplayAbilitySpecConfig& Existing : GrantAbilityConfigs)
+	{
+		if (Existing == Config)
+		{
+			return;
+		}
+	}
+
+	GrantAbilityConfigs.Emplace(Config);
+}
+
+
+// Reminder: this is happening on the authority only
+void UAbilitiesGameplayEffectComponent::OnInhibitionChanged(FActiveGameplayEffectHandle ActiveGEHandle, bool bIsInhibited) const
+{
+	if (bIsInhibited)
+	{
+		RemoveAbilities(ActiveGEHandle);
+	}
+	else
+	{
+		GrantAbilities(ActiveGEHandle);
+	}
+}
+
+// Reminder: this is happening on the authority only
+void UAbilitiesGameplayEffectComponent::GrantAbilities(FActiveGameplayEffectHandle ActiveGEHandle) const
+{
+	UAbilitySystemComponent* ASC = ActiveGEHandle.GetOwningAbilitySystemComponent();
+	if (!ensure(ASC))
+	{
+		UE_LOG(LogGameplayEffects, Error, TEXT("%s was passed an ActiveGEHandle %s which did not have a valid associated AbilitySystemComponent"), ANSI_TO_TCHAR(__func__), *ActiveGEHandle.ToString());
+		return;
+	}
+
+	if (ASC->bSuppressGrantAbility)
+	{
+		UE_LOG(LogGameplayEffects, Warning, TEXT("%s suppressed by %s bSuppressGrantAbility"), *GetName(), *ASC->GetName());
+		return;
+	}
+
+	const FActiveGameplayEffect* ActiveGE = ASC->GetActiveGameplayEffect(ActiveGEHandle);
+	if (!ActiveGE)
+	{
+		UE_LOG(LogGameplayEffects, Warning, TEXT("ActiveGEHandle %s did not corresponds to Active Gameplay Effect on %s. This could potentially happen if you remove the GE during the application of other GE's"), *ActiveGEHandle.ToString(), *ASC->GetName());
+		return;
+	}
+	const FGameplayEffectSpec& ActiveGESpec = ActiveGE->Spec;
+
+	const TArray<FGameplayAbilitySpec>& AllAbilities = ASC->GetActivatableAbilities();
+	for (const FGameplayAbilitySpecConfig& AbilityConfig : GrantAbilityConfigs)
+	{
+		// Check that we're configured
+		const UGameplayAbility* AbilityCDO = AbilityConfig.Ability.GetDefaultObject();
+		if (!AbilityCDO)
+		{
+			continue;
+		}
+
+		// Only do this if we haven't assigned the ability yet! This prevents cases where stacking GEs
+		// would regrant the ability every time the stack was applied
+		const bool bAlreadyGrantedAbility = AllAbilities.ContainsByPredicate([ASC, AbilityCDO, &ActiveGEHandle](FGameplayAbilitySpec& Spec) { return Spec.Ability == AbilityCDO && Spec.GameplayEffectHandle == ActiveGEHandle; });
+		if (bAlreadyGrantedAbility)
+		{
+			continue;
+		}
+
+		const FString ContextString = FString::Printf(TEXT("%s for %s from %s"), ANSI_TO_TCHAR(__func__), *AbilityCDO->GetName(), *GetNameSafe(ActiveGESpec.Def));
+		const int32 Level = static_cast<int32>(AbilityConfig.LevelScalableFloat.GetValueAtLevel(ActiveGESpec.GetLevel(), &ContextString));
+
+		// Now grant that ability to the owning actor
+		FGameplayAbilitySpec AbilitySpec{ AbilityConfig.Ability, Level, AbilityConfig.InputID, ActiveGESpec.GetEffectContext().GetSourceObject() };
+		AbilitySpec.SetByCallerTagMagnitudes = ActiveGESpec.SetByCallerTagMagnitudes;
+		AbilitySpec.GameplayEffectHandle = ActiveGEHandle;
+
+		ASC->GiveAbility(AbilitySpec);
+	}
+}
+
+// Reminder: this is happening on the authority only and the ActiveGEHandle can be considered 'inactive' by this time.
+void UAbilitiesGameplayEffectComponent::RemoveAbilities(FActiveGameplayEffectHandle ActiveGEHandle) const
+{
+	UAbilitySystemComponent* ASC = ActiveGEHandle.GetOwningAbilitySystemComponent();
+	if (!ensure(ASC))
+	{
+		UE_LOG(LogGameplayEffects, Error, TEXT("%s was passed an ActiveGEHandle %s which did not have a valid associated AbilitySystemComponent"), ANSI_TO_TCHAR(__func__), *ActiveGEHandle.ToString());
+		return;
+	}
+
+	const TArray<FGameplayAbilitySpec>& AllAbilities = ASC->GetActivatableAbilities();
+	for (const FGameplayAbilitySpecConfig& AbilityConfig : GrantAbilityConfigs)
+	{
+		// Check that we're configured
+		const UGameplayAbility* AbilityCDO = AbilityConfig.Ability.GetDefaultObject();
+		if (!AbilityCDO)
+		{
+			continue;
+		}
+
+		// See if we were granted, and if so we can remove it
+		const FGameplayAbilitySpec* AbilitySpecDef = AllAbilities.FindByPredicate([AbilityCDO, &ActiveGEHandle](FGameplayAbilitySpec& Spec) { return Spec.Ability == AbilityCDO && Spec.GameplayEffectHandle == ActiveGEHandle; });
+		if (!AbilitySpecDef)
+		{
+			continue;
+		}
+
+		switch (AbilityConfig.RemovalPolicy)
+		{
+			case EGameplayEffectGrantedAbilityRemovePolicy::CancelAbilityImmediately:
+			{
+				ASC->ClearAbility(AbilitySpecDef->Handle);
+				break;
+			}
+			case EGameplayEffectGrantedAbilityRemovePolicy::RemoveAbilityOnEnd:
+			{
+				ASC->SetRemoveAbilityOnEnd(AbilitySpecDef->Handle);
+				break;
+			}
+			default:
+			{
+				// Do nothing to granted ability
+				break;
+			}
+		}
+	}
+}
+
+bool UAbilitiesGameplayEffectComponent::OnActiveGameplayEffectAdded(FActiveGameplayEffectsContainer& ActiveGEContainer, FActiveGameplayEffect& ActiveGE) const
+{
+	if (ActiveGEContainer.IsNetAuthority())
+	{
+		ActiveGE.EventSet.OnEffectRemoved.AddUObject(this, &UAbilitiesGameplayEffectComponent::OnActiveGameplayEffectRemoved);
+		ActiveGE.EventSet.OnInhibitionChanged.AddUObject(this, &UAbilitiesGameplayEffectComponent::OnInhibitionChanged);
+	}
+
+	return true;
+}
+
+void UAbilitiesGameplayEffectComponent::OnActiveGameplayEffectRemoved(const FGameplayEffectRemovalInfo& RemovalInfo) const
+{
+	const FActiveGameplayEffect* ActiveGE = RemovalInfo.ActiveEffect;
+	if (!ensure(ActiveGE))
+	{
+		UE_LOG(LogGameplayEffects, Error, TEXT("FGameplayEffectRemovalInfo::ActiveEffect was not populated in %s"), ANSI_TO_TCHAR(__func__));
+		return;
+	}
+
+	RemoveAbilities(ActiveGE->Handle);
+}
+
+#if WITH_EDITOR
+EDataValidationResult UAbilitiesGameplayEffectComponent::IsDataValid(FDataValidationContext& Context) const
+{
+	EDataValidationResult Result = Super::IsDataValid(Context);
+
+	if (GetOwner()->DurationPolicy == EGameplayEffectDurationType::Instant)
+	{
+		if (GrantAbilityConfigs.Num() > 0)
+		{
+			Context.AddError(FText::FormatOrdered(LOCTEXT("InstantDoesNotWorkWithGrantAbilities", "GrantAbilityConfigs does not work with Instant Effects: {0}."), FText::FromString(GetClass()->GetName())));
+			Result = EDataValidationResult::Invalid;
+		}
+	}
+
+	for (int Index = 0; Index < GrantAbilityConfigs.Num(); ++Index)
+	{
+		const TSubclassOf<UGameplayAbility> AbilityClass = GrantAbilityConfigs[Index].Ability;
+		for (int CheckIndex = Index + 1; CheckIndex < GrantAbilityConfigs.Num(); ++CheckIndex)
+		{
+			if (AbilityClass == GrantAbilityConfigs[CheckIndex].Ability)
+			{
+				Context.AddError(FText::FormatOrdered(LOCTEXT("GrantAbilitiesMustBeUnique", "Multiple Abilities of the same type cannot be granted by {0}."), FText::FromString(GetClass()->GetName())));
+				Result = EDataValidationResult::Invalid;
+			}
+		}
+	}
+
+	return Result;
+}
+#endif // WITH_EDITOR
+
+#undef LOCTEXT_NAMESPACE
diff --git a/Engine/Plugins/Runtime/GameplayAbilities/Source/GameplayAbilities/Public/GameplayAbilitySpec.h b/Engine/Plugins/Runtime/GameplayAbilities/Source/GameplayAbilities/Public/GameplayAbilitySpec.h
--- a/Engine/Plugins/Runtime/GameplayAbilities/Source/GameplayAbilities/Public/GameplayAbilitySpec.h
+++ b/Engine/Plugins/Runtime/GameplayAbilities/Source/GameplayAbilities/Public/GameplayAbilitySpec.h
@@ -246,8 +246,10 @@ PRAGMA_ENABLE_DEPRECATION_WARNINGS
 	UPROPERTY()
 	TArray<TObjectPtr<UGameplayAbility>> ReplicatedInstances;
 
-	/** Handle to GE that granted us (usually invalid). FActiveGameplayEffectHandles are not synced across the network. */
-	UE_DEPRECATED(5.3, "GameplayEffectHandle was never replicated across the network. If you need FGameplayAbilitySpec -> FActiveGameplayEffect then use AbilitySystemComponent::FindActiveGameplayEffectHandle")
+	/**
+	 * Handle to GE that granted us (usually invalid). FActiveGameplayEffectHandles are not synced across the network and this is valid only on Authority.
+	 * If you need FGameplayAbilitySpec -> FActiveGameplayEffectHandle, then use AbilitySystemComponent::FindActiveGameplayEffectHandle.
+	 */
 	UPROPERTY(NotReplicated)
 	FActiveGameplayEffectHandle	GameplayEffectHandle;
 
diff --git a/Engine/Plugins/Runtime/GameplayAbilities/Source/GameplayAbilities/Public/GameplayEffect.h b/Engine/Plugins/Runtime/GameplayAbilities/Source/GameplayAbilities/Public/GameplayEffect.h
--- a/Engine/Plugins/Runtime/GameplayAbilities/Source/GameplayAbilities/Public/GameplayEffect.h
+++ b/Engine/Plugins/Runtime/GameplayAbilities/Source/GameplayAbilities/Public/GameplayEffect.h
@@ -92,8 +92,9 @@ enum class EGameplayEffectVersion : uint8
 {
 	Monolithic,	// Legacy version from Pre-UE5.3 (before we were versioning)
 	Modular53,	// New modular version available in UE5.3
+	AbilitiesComponent53,	// Granted Abilities are moved into the Abilities Component
 
-	Current = Modular53
+	Current = AbilitiesComponent53
 };
 
 struct GAMEPLAYABILITIES_API FGameplayEffectConstants
@@ -1196,6 +1197,7 @@ public:
 	uint32 bDurationLocked : 1;
 
 	/** List of abilities granted by this effect */
+	UE_DEPRECATED(5.3, "This variable will be removed in favor of (immutable) GASpecs that live on GameplayEffectComponents (e.g. AbilitiesGameplayEffectComponent)")
 	UPROPERTY()
 	TArray<FGameplayAbilitySpecDef> GrantedAbilitySpecs;
 
@@ -1356,6 +1358,10 @@ struct GAMEPLAYABILITIES_API FActiveGameplayEffect : public FFastArraySerializer
 
 	UPROPERTY()
 	FPredictionKey	PredictionKey;
+	
+	/** Handles of Gameplay Abilities that were granted to the target by this Active Gameplay Effect */
+	UPROPERTY()
+	TArray<FGameplayAbilitySpecHandle> GrantedAbilityHandles;
 
 	/** Server time this started */
 	UPROPERTY()
@@ -2148,10 +2154,8 @@ protected:
 	void SetVersion(EGameplayEffectVersion Version);
 
 private:
-	/** Logic that determines if we can upgrade this Gameplay Effect to use Gameplay Effect Components (or not if it's already been done). Used in the upgrade paths below. */
-	bool VersionAllowsCreateGEComponentsDuringUpgrade() const;
-
 	// Helper functions for converting data to use components
+	void ConvertAbilitiesComponent();
 	void ConvertAdditionalEffectsComponent();
 	void ConvertAssetTagsComponent();
 	void ConvertBlockByTagsComponent();
@@ -2340,7 +2344,8 @@ public:
 	// ----------------------------------------------------------------------
 	
 	/** Policy for what abilities this GE will grant. */
-	UPROPERTY(EditDefaultsOnly, BlueprintReadOnly, Category = "Granted Abilities")
+	UE_DEPRECATED(5.3, "GrantedAbilities are deprecated in favor of AbilitiesGameplayEffectComponent")
+	UPROPERTY(BlueprintReadOnly, Category = "Granted Abilities")
 	TArray<FGameplayAbilitySpecDef>	GrantedAbilities;
 
 	// ----------------------------------------------------------------------
diff --git a/Engine/Plugins/Runtime/GameplayAbilities/Source/GameplayAbilities/Public/GameplayEffectComponents/AbilitiesGameplayEffectComponent.h b/Engine/Plugins/Runtime/GameplayAbilities/Source/GameplayAbilities/Public/GameplayEffectComponents/AbilitiesGameplayEffectComponent.h
new file mode 100644
--- /dev/null
+++ b/Engine/Plugins/Runtime/GameplayAbilities/Source/GameplayAbilities/Public/GameplayEffectComponents/AbilitiesGameplayEffectComponent.h
@@ -0,0 +1,73 @@
+// Copyright Epic Games, Inc. All Rights Reserved.
+
+#pragma once
+
+#include "GameplayEffectComponent.h"
+#include "GameplayAbilitySpec.h"
+#include "AbilitiesGameplayEffectComponent.generated.h"
+
+/** Options on how to configure a GameplayAbilitySpec to grant an Actor */
+USTRUCT()
+struct FGameplayAbilitySpecConfig
+{
+	GENERATED_BODY()
+
+	/** What ability to grant */
+	UPROPERTY(EditDefaultsOnly, Category = "Ability Definition")
+	TSubclassOf<UGameplayAbility> Ability;
+
+	/** What level to grant this ability at */
+	UPROPERTY(EditDefaultsOnly, Category = "Ability Definition", DisplayName = "Level", meta=(UIMin=0.0))
+	FScalableFloat LevelScalableFloat = FScalableFloat{ 1.0f };
+
+	/** Input ID to bind this ability to */
+	UPROPERTY(EditDefaultsOnly, Category = "Ability Definition")
+	int32 InputID = INDEX_NONE;
+
+	/** What will remove this ability later */
+	UPROPERTY(EditDefaultsOnly, Category = "Ability Definition")
+	EGameplayEffectGrantedAbilityRemovePolicy RemovalPolicy = EGameplayEffectGrantedAbilityRemovePolicy::CancelAbilityImmediately;
+};
+
+/**
+ * Grants additional Gameplay Abilities to the Target of a Gameplay Effect while active
+ */
+UCLASS()
+class GAMEPLAYABILITIES_API UAbilitiesGameplayEffectComponent : public UGameplayEffectComponent
+{
+	GENERATED_BODY()
+
+public:
+	/** Constructor */
+	UAbilitiesGameplayEffectComponent();
+
+	/** Register for the appropriate events when we're applied */
+	virtual bool OnActiveGameplayEffectAdded(FActiveGameplayEffectsContainer& ActiveGEContainer, FActiveGameplayEffect& ActiveGE) const override;
+
+	/** Adds an entry for granting Gameplay Abilities */
+	void AddGrantedAbilityConfig(const FGameplayAbilitySpecConfig& Config);
+
+#if WITH_EDITOR
+	/** Warn on misconfigured Gameplay Effect */
+	virtual EDataValidationResult IsDataValid(class FDataValidationContext& Context) const override;
+#endif
+
+protected:
+	/** This allows us to be notified when the owning GameplayEffect has had its inhibition changed (which can happen on the initial application). */
+	virtual void OnInhibitionChanged(FActiveGameplayEffectHandle ActiveGEHandle, bool bIsInhibited) const;
+
+	/** We should grant the configured Gameplay Abilities */
+	virtual void GrantAbilities(FActiveGameplayEffectHandle ActiveGEHandle) const;
+
+	/** We should remove the configured Gameplay Abilities */
+	virtual void RemoveAbilities(FActiveGameplayEffectHandle ActiveGEHandle) const;
+
+private:
+	/** We must undo all effects when removed */
+	void OnActiveGameplayEffectRemoved(const FGameplayEffectRemovalInfo& RemovalInfo) const;
+
+protected:
+	/** Abilities to Grant to the Target while this Gameplay Effect is active */
+	UPROPERTY(EditDefaultsOnly, Category = GrantAbilities)
+	TArray<FGameplayAbilitySpecConfig>	GrantAbilityConfigs;
+};
