Fixing up an issue where the initial addition of an AssetTags wouldn't show its inherited values (if you removed, then re-added the AssetTags).

diff --git a/Engine/Plugins/Runtime/GameplayAbilities/Source/GameplayAbilities/Private/GameplayEffectComponents/AssetTagsGameplayEffectComponent.cpp b/Engine/Plugins/Runtime/GameplayAbilities/Source/GameplayAbilities/Private/GameplayEffectComponents/AssetTagsGameplayEffectComponent.cpp
--- a/Engine/Plugins/Runtime/GameplayAbilities/Source/GameplayAbilities/Private/GameplayEffectComponents/AssetTagsGameplayEffectComponent.cpp
+++ b/Engine/Plugins/Runtime/GameplayAbilities/Source/GameplayAbilities/Private/GameplayEffectComponents/AssetTagsGameplayEffectComponent.cpp
@@ -8,6 +8,11 @@ void UAssetTagsGameplayEffectComponent::PostInitProperties()
 {
 	Super::PostInitProperties();
 
+	// Try to find the parent and update the inherited tags.  We do this 'early' because this is the only function
+	// we can use for creation of the object (and this runs post-constructor).
+	const UAssetTagsGameplayEffectComponent* Parent = FindParentComponent(*this);
+	InheritableAssetTags.UpdateInheritedTagProperties(Parent ? &Parent->InheritableAssetTags : nullptr);
+
 #if WITH_EDITORONLY_DATA
 	if (GetOwner())
 	{
diff --git a/Engine/Plugins/Runtime/GameplayAbilities/Source/GameplayAbilities/Private/GameplayEffectComponents/BlockAbilityTagsGameplayEffectComponent.cpp b/Engine/Plugins/Runtime/GameplayAbilities/Source/GameplayAbilities/Private/GameplayEffectComponents/BlockAbilityTagsGameplayEffectComponent.cpp
--- a/Engine/Plugins/Runtime/GameplayAbilities/Source/GameplayAbilities/Private/GameplayEffectComponents/BlockAbilityTagsGameplayEffectComponent.cpp
+++ b/Engine/Plugins/Runtime/GameplayAbilities/Source/GameplayAbilities/Private/GameplayEffectComponents/BlockAbilityTagsGameplayEffectComponent.cpp
@@ -2,8 +2,15 @@
 
 #include "GameplayEffectComponents/BlockAbilityTagsGameplayEffectComponent.h"
 
-UBlockAbilityTagsGameplayEffectComponent::UBlockAbilityTagsGameplayEffectComponent()
+void UBlockAbilityTagsGameplayEffectComponent::PostInitProperties()
 {
+	Super::PostInitProperties();
+
+	// Try to find the parent and update the inherited tags.  We do this 'early' because this is the only function
+	// we can use for creation of the object (and this runs post-constructor).
+	const UBlockAbilityTagsGameplayEffectComponent* Parent = FindParentComponent(*this);
+	InheritableBlockedAbilityTagsContainer.UpdateInheritedTagProperties(Parent ? &Parent->InheritableBlockedAbilityTagsContainer : nullptr);
+
 #if WITH_EDITORONLY_DATA
 	EditorFriendlyName = TEXT("Block Abilities w/ Tags");
 #endif
diff --git a/Engine/Plugins/Runtime/GameplayAbilities/Source/GameplayAbilities/Private/GameplayEffectComponents/TargetTagsGameplayEffectComponent.cpp b/Engine/Plugins/Runtime/GameplayAbilities/Source/GameplayAbilities/Private/GameplayEffectComponents/TargetTagsGameplayEffectComponent.cpp
--- a/Engine/Plugins/Runtime/GameplayAbilities/Source/GameplayAbilities/Private/GameplayEffectComponents/TargetTagsGameplayEffectComponent.cpp
+++ b/Engine/Plugins/Runtime/GameplayAbilities/Source/GameplayAbilities/Private/GameplayEffectComponents/TargetTagsGameplayEffectComponent.cpp
@@ -2,8 +2,15 @@
 
 #include "GameplayEffectComponents/TargetTagsGameplayEffectComponent.h"
 
-UTargetTagsGameplayEffectComponent::UTargetTagsGameplayEffectComponent()
+void UTargetTagsGameplayEffectComponent::PostInitProperties()
 {
+	Super::PostInitProperties();
+
+	// Try to find the parent and update the inherited tags.  We do this 'early' because this is the only function
+	// we can use for creation of the object (and this runs post-constructor).
+	const UTargetTagsGameplayEffectComponent* Parent = FindParentComponent(*this);
+	InheritableGrantedTagsContainer.UpdateInheritedTagProperties(Parent ? &Parent->InheritableGrantedTagsContainer : nullptr);
+
 #if WITH_EDITORONLY_DATA
 	EditorFriendlyName = TEXT("Target Tags (Granted to Actor)");
 #endif // WITH_EDITORONLY_DATA
diff --git a/Engine/Plugins/Runtime/GameplayAbilities/Source/GameplayAbilities/Public/GameplayEffectComponents/AssetTagsGameplayEffectComponent.h b/Engine/Plugins/Runtime/GameplayAbilities/Source/GameplayAbilities/Public/GameplayEffectComponents/AssetTagsGameplayEffectComponent.h
--- a/Engine/Plugins/Runtime/GameplayAbilities/Source/GameplayAbilities/Public/GameplayEffectComponents/AssetTagsGameplayEffectComponent.h
+++ b/Engine/Plugins/Runtime/GameplayAbilities/Source/GameplayAbilities/Public/GameplayEffectComponents/AssetTagsGameplayEffectComponent.h
@@ -13,7 +13,7 @@ class GAMEPLAYABILITIES_API UAssetTagsGameplayEffectComponent : public UGameplay
 	GENERATED_BODY()
 
 public:
-	/** Setup a better EditorFriendlyName */
+	/** Setup a better EditorFriendlyName and do some initialization */
 	virtual void PostInitProperties() override;
 
 	/** Needed to properly apply FInheritedTagContainer properties */
diff --git a/Engine/Plugins/Runtime/GameplayAbilities/Source/GameplayAbilities/Public/GameplayEffectComponents/BlockAbilityTagsGameplayEffectComponent.h b/Engine/Plugins/Runtime/GameplayAbilities/Source/GameplayAbilities/Public/GameplayEffectComponents/BlockAbilityTagsGameplayEffectComponent.h
--- a/Engine/Plugins/Runtime/GameplayAbilities/Source/GameplayAbilities/Public/GameplayEffectComponents/BlockAbilityTagsGameplayEffectComponent.h
+++ b/Engine/Plugins/Runtime/GameplayAbilities/Source/GameplayAbilities/Public/GameplayEffectComponents/BlockAbilityTagsGameplayEffectComponent.h
@@ -13,8 +13,8 @@ class GAMEPLAYABILITIES_API UBlockAbilityTagsGameplayEffectComponent : public UG
 	GENERATED_BODY()
 	
 public:
-	/** Constructor to set the EditorFriendlyName */
-	UBlockAbilityTagsGameplayEffectComponent();
+	/** Setup an EditorFriendlyName and do some initialization */
+	virtual void PostInitProperties() override;
 
 	/** Needed to properly apply FInheritedTagContainer properties */
 	virtual void OnGameplayEffectChanged() const override;
diff --git a/Engine/Plugins/Runtime/GameplayAbilities/Source/GameplayAbilities/Public/GameplayEffectComponents/TargetTagRequirementsGameplayEffectComponent.h b/Engine/Plugins/Runtime/GameplayAbilities/Source/GameplayAbilities/Public/GameplayEffectComponents/TargetTagRequirementsGameplayEffectComponent.h
--- a/Engine/Plugins/Runtime/GameplayAbilities/Source/GameplayAbilities/Public/GameplayEffectComponents/TargetTagRequirementsGameplayEffectComponent.h
+++ b/Engine/Plugins/Runtime/GameplayAbilities/Source/GameplayAbilities/Public/GameplayEffectComponents/TargetTagRequirementsGameplayEffectComponent.h
@@ -29,7 +29,7 @@ public:
 
 #if WITH_EDITOR
 	/**
-	 * Validate incompatable configurations
+	 * Validate incompatible configurations
 	 */
 	virtual EDataValidationResult IsDataValid(class FDataValidationContext& Context) const override;
 #endif // WITH_EDITOR
diff --git a/Engine/Plugins/Runtime/GameplayAbilities/Source/GameplayAbilities/Public/GameplayEffectComponents/TargetTagsGameplayEffectComponent.h b/Engine/Plugins/Runtime/GameplayAbilities/Source/GameplayAbilities/Public/GameplayEffectComponents/TargetTagsGameplayEffectComponent.h
--- a/Engine/Plugins/Runtime/GameplayAbilities/Source/GameplayAbilities/Public/GameplayEffectComponents/TargetTagsGameplayEffectComponent.h
+++ b/Engine/Plugins/Runtime/GameplayAbilities/Source/GameplayAbilities/Public/GameplayEffectComponents/TargetTagsGameplayEffectComponent.h
@@ -13,8 +13,8 @@ class GAMEPLAYABILITIES_API UTargetTagsGameplayEffectComponent : public UGamepla
 	GENERATED_BODY()
 
 public:
-	/** Constructor that sets EditorFriendlyName */
-	UTargetTagsGameplayEffectComponent();
+	/** Setup an EditorFriendlyName and do some initialization */
+	virtual void PostInitProperties() override;
 
 	/** Needed to properly load FInheritedTagContainer properties */
 	virtual void OnGameplayEffectChanged() const override;
