Fix UAbilityTask::OnDestroy to not crash if called recursively for cases like UAbilityTask_StartAbilityState
This is based on PR #10469 but also skips an ensure in UGameplayTask::OnDestroy and downgrades a check so this cannot crash again

diff --git a/Engine/Plugins/Runtime/GameplayAbilities/Source/GameplayAbilities/Private/Abilities/Tasks/AbilityTask.cpp b/Engine/Plugins/Runtime/GameplayAbilities/Source/GameplayAbilities/Private/Abilities/Tasks/AbilityTask.cpp
--- a/Engine/Plugins/Runtime/GameplayAbilities/Source/GameplayAbilities/Private/Abilities/Tasks/AbilityTask.cpp
+++ b/Engine/Plugins/Runtime/GameplayAbilities/Source/GameplayAbilities/Private/Abilities/Tasks/AbilityTask.cpp
@@ -112,24 +112,32 @@ UAbilityTask::UAbilityTask(const FObjectInitializer& ObjectInitializer)
 
 void UAbilityTask::OnDestroy(bool bInOwnerFinished)
 {
-	if (!HasAnyFlags(RF_ClassDefaultObject))
+	// If we have already been destroyed this is being called recursively so skip the tracking as well as the super call
+	if (!bWasSuccessfullyDestroyed)
 	{
-		ensureMsgf(GlobalAbilityTaskCount > 0, TEXT("Mismatched AbilityTask counting"));
-		--GlobalAbilityTaskCount;
-		SET_DWORD_STAT(STAT_AbilitySystem_TaskCount, GlobalAbilityTaskCount);
-
-		if (AbilityTaskCVars::AbilityTaskRecordingType >= AbilityTaskConstants::DebugMinValueToEnableRecording)
+		if (!HasAnyFlags(RF_ClassDefaultObject))
 		{
-			DebugRecordAbilityTaskDestroyed(this);
+			ensureMsgf(GlobalAbilityTaskCount > 0, TEXT("Mismatched AbilityTask counting"));
+			--GlobalAbilityTaskCount;
+			SET_DWORD_STAT(STAT_AbilitySystem_TaskCount, GlobalAbilityTaskCount);
+
+			if (AbilityTaskCVars::AbilityTaskRecordingType >= AbilityTaskConstants::DebugMinValueToEnableRecording)
+			{
+				DebugRecordAbilityTaskDestroyed(this);
+			}
 		}
-	}
-	
-	bWasSuccessfullyDestroyed = true;
-	
-	// #KillPendingKill Clear ability reference so we don't hold onto it and GC can delete it.
-	Ability = nullptr;
 
-	Super::OnDestroy(bInOwnerFinished);
+		bWasSuccessfullyDestroyed = true;
+
+		// #KillPendingKill Clear ability reference so we don't hold onto it and GC can delete it.
+		Ability = nullptr;
+
+		Super::OnDestroy(bInOwnerFinished);
+	}
+	else
+	{
+		ensureMsgf(TaskState == EGameplayTaskState::Finished, TEXT("OnDestroy called twice on %s with invalid state %i"), *GetName(), TaskState);
+	}
 }
 
 void UAbilityTask::BeginDestroy()
@@ -142,7 +150,7 @@ void UAbilityTask::BeginDestroy()
 		if (!HasAnyFlags(RF_ClassDefaultObject))
 		{
 			// this shouldn't happen, it means that ability was destroyed while being active, but we need to keep GlobalAbilityTaskCount in sync anyway
-			checkf(GlobalAbilityTaskCount > 0, TEXT("Mismatched AbilityTask counting"));
+			ensureMsgf(GlobalAbilityTaskCount > 0, TEXT("Mismatched AbilityTask counting"));
 			--GlobalAbilityTaskCount;
 			SET_DWORD_STAT(STAT_AbilitySystem_TaskCount, GlobalAbilityTaskCount);
 
