Fix an ensure in UE::GameplayEffect::EditorOnly::ShouldUpgradeVersion

diff --git a/Engine/Plugins/Runtime/GameplayAbilities/Source/GameplayAbilities/Private/GameplayEffect.cpp b/Engine/Plugins/Runtime/GameplayAbilities/Source/GameplayAbilities/Private/GameplayEffect.cpp
--- a/Engine/Plugins/Runtime/GameplayAbilities/Source/GameplayAbilities/Private/GameplayEffect.cpp
+++ b/Engine/Plugins/Runtime/GameplayAbilities/Source/GameplayAbilities/Private/GameplayEffect.cpp
@@ -75,7 +75,7 @@ namespace UE::GameplayEffect
 		// Helper function that allows us to guard against upgrading versions (in case a bug is discovered while the product is live)
 		static bool ShouldUpgradeVersion(EGameplayEffectVersion FromVersion, EGameplayEffectVersion ToVersion)
 		{
-			EGameplayEffectVersion MaxVersion = static_cast<EGameplayEffectVersion>(CVarGameplayEffectMaxVersion.GetValueOnGameThread() & 0xff);
+			EGameplayEffectVersion MaxVersion = static_cast<EGameplayEffectVersion>(CVarGameplayEffectMaxVersion.GetValueOnAnyThread() & 0xff);
 			if (MaxVersion < ToVersion)
 			{
 				return false;
