Remove ensure for missing instance in InitAbilityActorInfo, as this can happen both during actor shutdown and replays

diff --git a/Engine/Plugins/Runtime/GameplayAbilities/Source/GameplayAbilities/Private/AbilitySystemComponent_Abilities.cpp b/Engine/Plugins/Runtime/GameplayAbilities/Source/GameplayAbilities/Private/AbilitySystemComponent_Abilities.cpp
--- a/Engine/Plugins/Runtime/GameplayAbilities/Source/GameplayAbilities/Private/AbilitySystemComponent_Abilities.cpp
+++ b/Engine/Plugins/Runtime/GameplayAbilities/Source/GameplayAbilities/Private/AbilitySystemComponent_Abilities.cpp
@@ -162,7 +162,8 @@ void UAbilitySystemComponent::InitAbilityActorInfo(AActor* InOwnerActor, AActor*
 				if (Spec.Ability->GetInstancingPolicy() == EGameplayAbilityInstancingPolicy::InstancedPerActor)
 				{
 					UGameplayAbility* AbilityInstance = Spec.GetPrimaryInstance();
-					if (ensureMsgf(AbilityInstance, TEXT("%s: Could not find a Primary Instance for InstancedPerActor Ability Spec!"), ANSI_TO_TCHAR(__func__)))
+					// If we don't have the ability instance, it was either already destroyed or will get called on creation
+					if (AbilityInstance)
 					{
 						AbilityInstance->OnAvatarSet(AbilityActorInfo.Get(), Spec);
 					}
