Protected a gameplay cue crash when given a NaN transform. Logging/erroring has been added.
Bug# PLAY-7639 - Client crash starting with FWeakObjectPtr::operator= (in GameplayCueManager::GetInstancedCueActor)

diff --git a/Engine/Plugins/Runtime/GameplayAbilities/Source/GameplayAbilities/Private/GameplayCueManager.cpp b/Engine/Plugins/Runtime/GameplayAbilities/Source/GameplayAbilities/Private/GameplayCueManager.cpp
--- a/Engine/Plugins/Runtime/GameplayAbilities/Source/GameplayAbilities/Private/GameplayCueManager.cpp
+++ b/Engine/Plugins/Runtime/GameplayAbilities/Source/GameplayAbilities/Private/GameplayCueManager.cpp
@@ -408,6 +408,12 @@ AGameplayCueNotify_Actor* UGameplayCueManager::GetInstancedCueActor(AActor* Targ
 		return nullptr;
 	}
 
+	if (TargetActor && TargetActor->GetActorTransform().ContainsNaN())
+	{
+		UE_LOG(LogAbilitySystem, Error, TEXT("GetInstancedCueActor called with invalid target actor transform (TargetActor = %s)"), *GetNameSafe(TargetActor));
+		return nullptr;
+	}
+
 	// There used to be special code here to handle the case where the TargetActor was a CDO.  I'm not sure why that would be (or how that's even possible -- perhaps in the default Blueprint Viewport? But I can't trigger it.)
 	// Let's log it in case a user comes across this issue.
 	//	Animtion preview hack. If we are trying to play the GC on a CDO, then don't use actor recycling and don't set the owner (to the CDO, which would cause problems)
@@ -460,13 +466,14 @@ AGameplayCueNotify_Actor* UGameplayCueManager::GetInstancedCueActor(AActor* Targ
 	SpawnParams.Owner = TargetActor;
 	SpawnParams.OverrideLevel = World->PersistentLevel;
 	AGameplayCueNotify_Actor* SpawnedCue = World->SpawnActor<AGameplayCueNotify_Actor>(CueClass, TargetActor->GetActorLocation(), TargetActor->GetActorRotation(), SpawnParams);
-	SpawnedCue->CueInstigator = Parameters.GetInstigator();
-	SpawnedCue->CueSourceObject = Parameters.GetSourceObject();
+	if (ensureMsgf(SpawnedCue != nullptr, TEXT("[%s] - Failed to spawn the cue actor! Check the log for a potential reason spawn actor failed!"), ANSI_TO_TCHAR(__FUNCTION__)))
+	{
+		SpawnedCue->CueInstigator = Parameters.GetInstigator();
+		SpawnedCue->CueSourceObject = Parameters.GetSourceObject();
+	}
 
 	UE_CLOG(LogGameplayCueActorSpawning > 0, LogAbilitySystem, Warning, TEXT("Spawned Gameplay Cue Notify Actor: %s (instance %s)"), *CueClass->GetName(), *GetNameSafe(SpawnedCue));
-
 	return SpawnedCue;
-
 }
 
 AGameplayCueNotify_Actor* UGameplayCueManager::FindRecycledCue(const TSubclassOf<AGameplayCueNotify_Actor>& CueClass, const UWorld& FindInWorld)
