PR #7661:  Mark functions as virtual in AbilityTask_WaitTargetData.

diff --git a/Engine/Plugins/Runtime/GameplayAbilities/Source/GameplayAbilities/Public/Abilities/Tasks/AbilityTask_WaitTargetData.h b/Engine/Plugins/Runtime/GameplayAbilities/Source/GameplayAbilities/Public/Abilities/Tasks/AbilityTask_WaitTargetData.h
--- a/Engine/Plugins/Runtime/GameplayAbilities/Source/GameplayAbilities/Public/Abilities/Tasks/AbilityTask_WaitTargetData.h
+++ b/Engine/Plugins/Runtime/GameplayAbilities/Source/GameplayAbilities/Public/Abilities/Tasks/AbilityTask_WaitTargetData.h
@@ -30,16 +30,16 @@ class GAMEPLAYABILITIES_API UAbilityTask_WaitTargetData: public UAbilityTask
 	FWaitTargetDataDelegate	Cancelled;
 
 	UFUNCTION()
-	void OnTargetDataReplicatedCallback(const FGameplayAbilityTargetDataHandle& Data, FGameplayTag ActivationTag);
+	virtual void OnTargetDataReplicatedCallback(const FGameplayAbilityTargetDataHandle& Data, FGameplayTag ActivationTag);
 
 	UFUNCTION()
-	void OnTargetDataReplicatedCancelledCallback();
+	virtual void OnTargetDataReplicatedCancelledCallback();
 
 	UFUNCTION()
-	void OnTargetDataReadyCallback(const FGameplayAbilityTargetDataHandle& Data);
+	virtual void OnTargetDataReadyCallback(const FGameplayAbilityTargetDataHandle& Data);
 
 	UFUNCTION()
-	void OnTargetDataCancelledCallback(const FGameplayAbilityTargetDataHandle& Data);
+	virtual void OnTargetDataCancelledCallback(const FGameplayAbilityTargetDataHandle& Data);
 
 	/** Spawns target actor and waits for it to return valid data or to be canceled. */
 	UFUNCTION(BlueprintCallable, meta=(HidePin = "OwningAbility", DefaultToSelf = "OwningAbility", BlueprintInternalUseOnly = "true", HideSpawnParms="Instigator"), Category="Ability|Tasks")
@@ -52,10 +52,10 @@ class GAMEPLAYABILITIES_API UAbilityTask_WaitTargetData: public UAbilityTask
 	virtual void Activate() override;
 
 	UFUNCTION(BlueprintCallable, meta = (HidePin = "OwningAbility", DefaultToSelf = "OwningAbility", BlueprintInternalUseOnly = "true"), Category = "Abilities")
-	bool BeginSpawningActor(UGameplayAbility* OwningAbility, TSubclassOf<AGameplayAbilityTargetActor> Class, AGameplayAbilityTargetActor*& SpawnedActor);
+	virtual bool BeginSpawningActor(UGameplayAbility* OwningAbility, TSubclassOf<AGameplayAbilityTargetActor> Class, AGameplayAbilityTargetActor*& SpawnedActor);
 
 	UFUNCTION(BlueprintCallable, meta = (HidePin = "OwningAbility", DefaultToSelf = "OwningAbility", BlueprintInternalUseOnly = "true"), Category = "Abilities")
-	void FinishSpawningActor(UGameplayAbility* OwningAbility, AGameplayAbilityTargetActor* SpawnedActor);
+	virtual void FinishSpawningActor(UGameplayAbility* OwningAbility, AGameplayAbilityTargetActor* SpawnedActor);
 
 	/** Called when the ability is asked to confirm from an outside node. What this means depends on the individual task. By default, this does nothing other than ending if bEndTask is true. */
 	virtual void ExternalConfirm(bool bEndTask) override;
@@ -65,15 +65,15 @@ class GAMEPLAYABILITIES_API UAbilityTask_WaitTargetData: public UAbilityTask
 
 protected:
 
-	bool ShouldSpawnTargetActor() const;
-	void InitializeTargetActor(AGameplayAbilityTargetActor* SpawnedActor) const;
-	void FinalizeTargetActor(AGameplayAbilityTargetActor* SpawnedActor) const;
+	virtual bool ShouldSpawnTargetActor() const;
+	virtual void InitializeTargetActor(AGameplayAbilityTargetActor* SpawnedActor) const;
+	virtual void FinalizeTargetActor(AGameplayAbilityTargetActor* SpawnedActor) const;
 
-	void RegisterTargetDataCallbacks();
+	virtual void RegisterTargetDataCallbacks();
 
 	virtual void OnDestroy(bool AbilityEnded) override;
 
-	bool ShouldReplicateDataToServer() const;
+	virtual bool ShouldReplicateDataToServer() const;
 
 protected:
 
