PR #7531: GAS stops search for PC only if valid PC is found

diff --git a/Engine/Plugins/Runtime/GameplayAbilities/Source/GameplayAbilities/Private/GameplayAbilityTypes.cpp b/Engine/Plugins/Runtime/GameplayAbilities/Source/GameplayAbilities/Private/GameplayAbilityTypes.cpp
--- a/Engine/Plugins/Runtime/GameplayAbilities/Source/GameplayAbilities/Private/GameplayAbilityTypes.cpp
+++ b/Engine/Plugins/Runtime/GameplayAbilities/Source/GameplayAbilities/Private/GameplayAbilityTypes.cpp
@@ -43,7 +43,11 @@ void FGameplayAbilityActorInfo::InitFromActor(AActor *InOwnerActor, AActor *InAv
 		if (APawn * Pawn = Cast<APawn>(TestActor))
 		{
 			PlayerController = Cast<APlayerController>(Pawn->GetController());
-			break;
+
+			if (PlayerController.IsValid())
+			{
+				break;
+			}
 		}
 
 		TestActor = TestActor->GetOwner();
