PR #10474: Update `GetAbilitySystemComponent` to default parameter to Self

diff --git a/Engine/Plugins/Runtime/GameplayAbilities/Source/GameplayAbilities/Public/AbilitySystemBlueprintLibrary.h b/Engine/Plugins/Runtime/GameplayAbilities/Source/GameplayAbilities/Public/AbilitySystemBlueprintLibrary.h
--- a/Engine/Plugins/Runtime/GameplayAbilities/Source/GameplayAbilities/Public/AbilitySystemBlueprintLibrary.h
+++ b/Engine/Plugins/Runtime/GameplayAbilities/Source/GameplayAbilities/Public/AbilitySystemBlueprintLibrary.h
@@ -27,8 +27,8 @@ class GAMEPLAYABILITIES_API UAbilitySystemBlueprintLibrary : public UBlueprintFu
 	GENERATED_UCLASS_BODY()
 
 	/** Tries to find an ability system component on the actor, will use AbilitySystemInterface or fall back to a component search */
-	UFUNCTION(BlueprintPure, Category = Ability)
-	static UAbilitySystemComponent* GetAbilitySystemComponent(AActor *Actor);
+	UFUNCTION(BlueprintPure, Category = Ability, Meta=(DefaultToSelf = "Actor"))
+	static UAbilitySystemComponent* GetAbilitySystemComponent(AActor* Actor);
 
 	/**
 	 * This function can be used to trigger an ability on the actor in question with useful payload data.
