PR #8063: FGameplayEffectContextHandle : Check is data valid before retrieve "Actors"

diff --git a/Engine/Plugins/Runtime/GameplayAbilities/Source/GameplayAbilities/Public/GameplayEffectTypes.h b/Engine/Plugins/Runtime/GameplayAbilities/Source/GameplayAbilities/Public/GameplayEffectTypes.h
--- a/Engine/Plugins/Runtime/GameplayAbilities/Source/GameplayAbilities/Public/GameplayEffectTypes.h
+++ b/Engine/Plugins/Runtime/GameplayAbilities/Source/GameplayAbilities/Public/GameplayEffectTypes.h
@@ -675,7 +675,12 @@ struct GAMEPLAYABILITIES_API FGameplayEffectContextHandle
 	/** Returns actor list, may be empty */
 	const TArray<TWeakObjectPtr<AActor>> GetActors()
 	{
-		return Data->GetActors();
+		if (IsValid())
+		{
+			return Data->GetActors();
+		}
+
+		return {};
 	}
 
 	/** Returns hit result, this can be null */
