PR #7542 : Retain rotation for GAS Target Data LocationInfo

diff --git a/Engine/Plugins/Runtime/GameplayAbilities/Source/GameplayAbilities/Public/Abilities/GameplayAbilityTargetTypes.h b/Engine/Plugins/Runtime/GameplayAbilities/Source/GameplayAbilities/Public/Abilities/GameplayAbilityTargetTypes.h
--- a/Engine/Plugins/Runtime/GameplayAbilities/Source/GameplayAbilities/Public/Abilities/GameplayAbilityTargetTypes.h
+++ b/Engine/Plugins/Runtime/GameplayAbilities/Source/GameplayAbilities/Public/Abilities/GameplayAbilityTargetTypes.h
@@ -403,6 +403,11 @@ struct GAMEPLAYABILITIES_API FGameplayAbilityTargetData_LocationInfo : public FG
 		return TargetLocation.GetTargetingTransform().GetLocation();
 	}
 
+	virtual FTransform GetEndPointTransform() const override
+	{
+		return TargetLocation.GetTargetingTransform();
+	}
+
 	// -------------------------------------
 
 	virtual UScriptStruct* GetScriptStruct() const override
