PR #6354 : Fix ability montages stopping abruptly on simulated proxy

diff --git a/Engine/Plugins/Runtime/GameplayAbilities/Source/GameplayAbilities/Private/AbilitySystemComponent_Abilities.cpp b/Engine/Plugins/Runtime/GameplayAbilities/Source/GameplayAbilities/Private/AbilitySystemComponent_Abilities.cpp
--- a/Engine/Plugins/Runtime/GameplayAbilities/Source/GameplayAbilities/Private/AbilitySystemComponent_Abilities.cpp
+++ b/Engine/Plugins/Runtime/GameplayAbilities/Source/GameplayAbilities/Private/AbilitySystemComponent_Abilities.cpp
@@ -47,6 +47,7 @@ CSV_DECLARE_CATEGORY_MODULE_EXTERN(CORE_API, Basic);
 
 static TAutoConsoleVariable<float> CVarReplayMontageErrorThreshold(TEXT("replay.MontageErrorThreshold"), 0.5f, TEXT("Tolerance level for when montage playback position correction occurs in replays"));
 static TAutoConsoleVariable<bool> CVarAbilitySystemSetActivationInfoMultipleTimes(TEXT("AbilitySystem.SetActivationInfoMultipleTimes"), false, TEXT("Set this to true if some replicated Gameplay Abilities aren't setting their owning actors correctly"));
+static TAutoConsoleVariable<bool> CVarGasFixClientSideMontageBlendOutTime(TEXT("AbilitySystem.Fix.ClientSideMontageBlendOutTime"), true, TEXT("Enable a fix to replicate the Montage BlendOutTime for (recently) stopped Montages"));
 
 void UAbilitySystemComponent::InitializeComponent()
 {
@@ -2873,6 +2874,25 @@ void UAbilitySystemComponent::AnimMontage_UpdateReplicatedData(FGameplayAbilityR
 			// Set this prior to calling UpdateShouldTick, so we start ticking if we are playing a Montage
 			OutRepAnimMontageInfo.IsStopped = bIsStopped;
 
+			if (bIsStopped)
+			{
+				// Use AnyThread because GetValueOnGameThread will fail check() when doing replays
+				constexpr bool bForceGameThreadValue = true;
+				if (CVarGasFixClientSideMontageBlendOutTime.GetValueOnAnyThread(bForceGameThreadValue))
+				{
+					// Replicate blend out time. This requires a manual search since Montage_GetBlendTime will fail
+					// in GetActiveInstanceForMontage for Montages that are stopped.
+					for (const FAnimMontageInstance* MontageInstance : AnimInstance->MontageInstances)
+					{
+						if (MontageInstance->Montage == LocalAnimMontageInfo.AnimMontage)
+						{
+							OutRepAnimMontageInfo.BlendTime = MontageInstance->GetBlendTime();
+							break;
+						}
+					}
+				}
+			}
+
 			// When we start or stop an animation, update the clients right away for the Avatar Actor
 			if (AbilityActorInfo->AvatarActor != nullptr)
 			{
