Improve Gameplay Ability System Debug Commands

diff --git a/Engine/Plugins/Runtime/GameplayAbilities/Source/GameplayAbilities/Private/AbilitySystemGlobals.cpp b/Engine/Plugins/Runtime/GameplayAbilities/Source/GameplayAbilities/Private/AbilitySystemGlobals.cpp
--- a/Engine/Plugins/Runtime/GameplayAbilities/Source/GameplayAbilities/Private/AbilitySystemGlobals.cpp
+++ b/Engine/Plugins/Runtime/GameplayAbilities/Source/GameplayAbilities/Private/AbilitySystemGlobals.cpp
@@ -671,7 +671,7 @@ namespace UE::AbilitySystemGlobals
 // AbilitySystem.Ability Debug Commands
 //
 
-FAutoConsoleCommand DebugAbilitySystemAbilityListGrantedCommand(TEXT("AbilitySystem.Ability.ListGranted"), TEXT("Lists all of the Gameplay Abilities currently granted to the Player"),
+FAutoConsoleCommand DebugAbilitySystemAbilityListGrantedCommand(TEXT("AbilitySystem.Ability.Grant"), TEXT("Include a param [ClassName/AssetName] to Grant an Ability to the Player.  Omit parameters to lists all of the Gameplay Abilities currently granted to the Player."),
 	FConsoleCommandWithWorldArgsAndOutputDeviceDelegate::CreateLambda([](const TArray<FString>& Args, UWorld* World, FOutputDevice& OutputDevice)
 	{
 		UAbilitySystemGlobals& AbilitySystemGlobals = UAbilitySystemGlobals::Get();
@@ -684,20 +684,71 @@ FAutoConsoleCommand DebugAbilitySystemAbilityListGrantedCommand(TEXT("AbilitySys
 			return;
 		}
 
-		const UEnum* ExecutionEnumPtr = FindObject<UEnum>(nullptr, TEXT("/Script/GameplayAbilities.EGameplayAbilityNetExecutionPolicy"), true);
-		check(ExecutionEnumPtr&& TEXT("Couldn't locate EGameplayAbilityNetExecutionPolicy enum!"));
+		if (Args.Num() > 1)
+		{
+			OutputDevice.Logf(TEXT("Expected a single parameter (ClassName or AssetName).  Use no parameters to list all granted abilities."));
+			return;
+		}
+
+		// List already granted
+		if (Args.Num() < 1)
+		{
+			const UEnum* ExecutionEnumPtr = FindObject<UEnum>(nullptr, TEXT("/Script/GameplayAbilities.EGameplayAbilityNetExecutionPolicy"), true);
+			check(ExecutionEnumPtr&& TEXT("Couldn't locate EGameplayAbilityNetExecutionPolicy enum!"));
+
+			const UEnum* SecurityEnumPtr = FindObject<UEnum>(nullptr, TEXT("/Script/GameplayAbilities.EGameplayAbilityNetSecurityPolicy"), true);
+			check(SecurityEnumPtr&& TEXT("Couldn't locate EGameplayAbilityNetSecurityPolicy enum!"));
+
+			OutputDevice.Logf(TEXT("Granted abilities to %s (ASC: '%s'):"), *PC->GetName(), *ASC->GetFullName());
+
+			for (FGameplayAbilitySpec& Activatable : ASC->GetActivatableAbilities())
+			{
+				const TCHAR* ActiveText = Activatable.IsActive() ? TEXT("**ACTIVE**") : TEXT("");
+				OutputDevice.Logf(TEXT("   %s (%s - %s) %s"), *Activatable.Ability->GetName(), *ExecutionEnumPtr->GetDisplayNameTextByIndex(Activatable.Ability->GetNetExecutionPolicy()).ToString(), *SecurityEnumPtr->GetDisplayNameTextByIndex(Activatable.Ability->GetNetSecurityPolicy()).ToString(), ActiveText);
+			}
 
-		const UEnum* SecurityEnumPtr = FindObject<UEnum>(nullptr, TEXT("/Script/GameplayAbilities.EGameplayAbilityNetSecurityPolicy"), true);
-		check(SecurityEnumPtr&& TEXT("Couldn't locate EGameplayAbilityNetSecurityPolicy enum!"));
+			return;
+		}
 
-		OutputDevice.Logf(TEXT("Granted abilities to %s (ASC: '%s'):"), *PC->GetName(), *ASC->GetFullName());
+		FString SearchString = Args[0];
 
-		for (FGameplayAbilitySpec& Activatable : ASC->GetActivatableAbilities())
+		// We couldn't find anything the user was searching for, so early out
+		TSubclassOf<UGameplayAbility> ActivateAbilityClass = UE::AbilitySystemGlobals::FuzzyFindClass<UGameplayAbility>(SearchString);
+		if (!ActivateAbilityClass)
 		{
-			const TCHAR* ActiveText = Activatable.IsActive() ? TEXT("**ACTIVE**") : TEXT("");
-			OutputDevice.Logf(TEXT("   %s (%s - %s) %s"), *Activatable.Ability->GetName(), *ExecutionEnumPtr->GetDisplayNameTextByIndex(Activatable.Ability->GetNetExecutionPolicy()).ToString(), *SecurityEnumPtr->GetDisplayNameTextByIndex(Activatable.Ability->GetNetSecurityPolicy()).ToString(), ActiveText);
+			OutputDevice.Logf(TEXT("Could not find a valid Gameplay Ability based on Search String '%s'"), *SearchString);
+			return;
 		}
-	}), ECVF_Cheat);
+
+		// Check if it's already granted
+		if (const FGameplayAbilitySpec* ExistingSpec = ASC->FindAbilitySpecFromClass(ActivateAbilityClass))
+		{
+			OutputDevice.Logf(TEXT("Existing Ability Spec '%s' on Player '%s' (It is already granted)."), *GetNameSafe(*ActivateAbilityClass), *GetNameSafe(PC));
+			return;
+		}
+
+		// If we're not the authority, we need to send the command to the server because we can't grant locally.
+		if (!ASC->IsOwnerActorAuthoritative())
+		{
+			const FString ServerCommand = FString::Printf(TEXT("AbilitySystem.Ability.Grant %s"), *SearchString);
+			PC->ServerExec(ServerCommand);
+
+			OutputDevice.Logf(TEXT("Sent Command '%s' from Player '%s' to Server (Reason: Cannot Grant if not Authority)."), *ServerCommand, *GetNameSafe(PC));
+			return;
+		}
+
+		// It wasn't granted, let's grant it now.
+		FGameplayAbilitySpec AbilitySpec{ ActivateAbilityClass };
+		FGameplayAbilitySpecHandle SpecHandle = ASC->GiveAbility(AbilitySpec);
+		if (SpecHandle.IsValid())
+		{
+			OutputDevice.Logf(TEXT("Successfully Granted '%s' on Player '%s'."), *GetNameSafe(ActivateAbilityClass.Get()), *GetNameSafe(PC));
+		}
+		else
+		{
+			OutputDevice.Logf(TEXT("Failed to Grant '%s' on Player '%s'."), *GetNameSafe(ActivateAbilityClass.Get()), *GetNameSafe(PC));
+		}
+		}), ECVF_Cheat);
 
 FAutoConsoleCommand DebugAbilitySystemAbilityCancelCommand(TEXT("AbilitySystem.Ability.Cancel"), TEXT("<Name>. Cancels (prematurely Ends) a currently executing Gameplay Ability"),
 	FConsoleCommandWithWorldArgsAndOutputDeviceDelegate::CreateLambda([](const TArray<FString>& Args, UWorld* World, FOutputDevice& OutputDevice)
@@ -777,7 +828,7 @@ FAutoConsoleCommand DebugAbilitySystemAbilityActivateCommand(TEXT("AbilitySystem
 
 		if (Args.Num() < 1)
 		{
-			if (IConsoleObject* GetGrantedCommand = IConsoleManager::Get().FindConsoleObject(TEXT("AbilitySystem.Ability.ListGranted"), false))
+			if (IConsoleObject* GetGrantedCommand = IConsoleManager::Get().FindConsoleObject(TEXT("AbilitySystem.Ability.Grant"), false))
 			{
 				GetGrantedCommand->AsCommand()->Execute(Args, World, OutputDevice);
 			}
@@ -829,6 +880,21 @@ FAutoConsoleCommand DebugAbilitySystemAbilityActivateCommand(TEXT("AbilitySystem
 			return;
 		}
 
+		// If we're not the authority, we should check if we need to send the command to the server
+		if (!ASC->IsOwnerActorAuthoritative())
+		{
+			const UGameplayAbility* Ability = ActivateAbilityClass.GetDefaultObject();
+			const bool bSendToServer = Ability && (Ability->GetNetExecutionPolicy() >= EGameplayAbilityNetExecutionPolicy::ServerInitiated);
+			if (bSendToServer)
+			{
+				const FString ServerCommand = FString::Printf(TEXT("AbilitySystem.Ability.Activate %s"), *SearchString);
+				PC->ServerExec(ServerCommand);
+
+				OutputDevice.Logf(TEXT("Sent Command '%s' from Player '%s' to Server (Reason: Net Execution Policy)."), *ServerCommand, *GetNameSafe(PC), *GetNameSafe(ActivateAbilityClass.Get()));
+				return;
+			}
+		}
+
 		// We found what the user was searching for, so let's try to activate it.  First, let's assume the Ability was already granted.
 		bool bSuccess = ASC->TryActivateAbilityByClass(ActivateAbilityClass);
 		if (bSuccess)
diff --git a/Engine/Plugins/Runtime/GameplayAbilities/Source/GameplayAbilitiesEditor/Private/GameplayAbilityAudit.cpp b/Engine/Plugins/Runtime/GameplayAbilities/Source/GameplayAbilitiesEditor/Private/GameplayAbilityAudit.cpp
--- a/Engine/Plugins/Runtime/GameplayAbilities/Source/GameplayAbilitiesEditor/Private/GameplayAbilityAudit.cpp
+++ b/Engine/Plugins/Runtime/GameplayAbilities/Source/GameplayAbilitiesEditor/Private/GameplayAbilityAudit.cpp
@@ -92,6 +92,7 @@ void FGameplayAbilityAuditRow::FillDataFromGameplayAbility(const UGameplayAbilit
 
 	AuditRow.NetExecutionPolicy = GameplayAbility.GetNetExecutionPolicy();
 	AuditRow.NetSecurityPolicy = GameplayAbility.GetNetSecurityPolicy();
+	AuditRow.ReplicationPolicy = GameplayAbility.GetReplicationPolicy();
 }
 
 /**
diff --git a/Engine/Plugins/Runtime/GameplayAbilities/Source/GameplayAbilitiesEditor/Public/GameplayAbilityAudit.h b/Engine/Plugins/Runtime/GameplayAbilities/Source/GameplayAbilitiesEditor/Public/GameplayAbilityAudit.h
--- a/Engine/Plugins/Runtime/GameplayAbilities/Source/GameplayAbilitiesEditor/Public/GameplayAbilityAudit.h
+++ b/Engine/Plugins/Runtime/GameplayAbilities/Source/GameplayAbilitiesEditor/Public/GameplayAbilityAudit.h
@@ -26,7 +26,7 @@ UENUM()
 enum class EGameplayAbilityActivationPath : uint8
 {
 	Native,		/* Ability has a Native Activate function (and therefore uninspectable by Blueprint analysis) */
-	Blueprint,	/* Ability has an Activate event implemented in Blueprints and is therefore not triggered by a GameplayEvent */
+	Blueprint,	/* Ability has an Activate event implemented in Blueprints and is therefore likely not triggered by a GameplayEvent */
 	FromEvent	/* Ability has an ActivateByEvent event implemented in Blueprints and is therefore activated by a GameplayEvent */
 };
 
@@ -65,6 +65,10 @@ protected:
 	UPROPERTY()
 	TEnumAsByte<EGameplayAbilityNetSecurityPolicy::Type> NetSecurityPolicy = EGameplayAbilityNetSecurityPolicy::Type::ClientOrServer;
 
+	/** The Replication Policy which controls if the instance itself is replicated (and therefore controls variable replication) */
+	UPROPERTY()
+	TEnumAsByte<EGameplayAbilityReplicationPolicy::Type> ReplicationPolicy = EGameplayAbilityReplicationPolicy::Type::ReplicateNo;
+
 	/** If the GameplayAbility has a Cost GameplayEffect, this field will tell us which one it's using. */
 	UPROPERTY()
 	FName CostGE = NAME_None;
