Fix Decal GameplayCues destroying the target actor if attached

diff --git a/Engine/Plugins/Runtime/GameplayAbilities/Source/GameplayAbilities/Private/GameplayCueNotifyTypes.cpp b/Engine/Plugins/Runtime/GameplayAbilities/Source/GameplayAbilities/Private/GameplayCueNotifyTypes.cpp
--- a/Engine/Plugins/Runtime/GameplayAbilities/Source/GameplayAbilities/Private/GameplayCueNotifyTypes.cpp
+++ b/Engine/Plugins/Runtime/GameplayAbilities/Source/GameplayAbilities/Private/GameplayCueNotifyTypes.cpp
@@ -979,7 +979,8 @@ bool FGameplayCueNotify_DecalInfo::SpawnDecal(const FGameplayCueNotify_SpawnCont
 	{
 		if (bOverrideFadeOut)
 		{
-			SpawnedDecalComponent->SetFadeOut(FadeOutStartDelay, FadeOutDuration);
+			constexpr bool bDestroyOwnerAfterFade = false;
+			SpawnedDecalComponent->SetFadeOut(FadeOutStartDelay, FadeOutDuration, bDestroyOwnerAfterFade);
 		}
 	}
 
