Enable AutoRTFM for these modules

diff --git a/Engine/Plugins/Runtime/GameplayAbilities/Source/GameplayAbilities/GameplayAbilities.Build.cs b/Engine/Plugins/Runtime/GameplayAbilities/Source/GameplayAbilities/GameplayAbilities.Build.cs
--- a/Engine/Plugins/Runtime/GameplayAbilities/Source/GameplayAbilities/GameplayAbilities.Build.cs
+++ b/Engine/Plugins/Runtime/GameplayAbilities/Source/GameplayAbilities/GameplayAbilities.Build.cs
@@ -39,6 +39,8 @@ namespace UnrealBuildTool.Rules
 			SetupGameplayDebuggerSupport(Target);
 
 			SetupIrisSupport(Target);
+
+			bAllowAutoRTFMInstrumentation = true;
 		}
 	}
 }
diff --git a/Engine/Source/Runtime/Engine/Engine.Build.cs b/Engine/Source/Runtime/Engine/Engine.Build.cs
--- a/Engine/Source/Runtime/Engine/Engine.Build.cs
+++ b/Engine/Source/Runtime/Engine/Engine.Build.cs
@@ -489,5 +489,7 @@ public class Engine : ModuleRules
 		SetupIrisSupport(Target, bIrisAddAsPublicDepedency);
 
 		PrivateDefinitions.Add("UE_DEPRECATE_LEGACY_MATH_CONSTANT_MACRO_NAMES=1");
+
+		bAllowAutoRTFMInstrumentation = true;
 	}
 }
diff --git a/Engine/Source/Runtime/Net/Core/NetCore.Build.cs b/Engine/Source/Runtime/Net/Core/NetCore.Build.cs
--- a/Engine/Source/Runtime/Net/Core/NetCore.Build.cs
+++ b/Engine/Source/Runtime/Net/Core/NetCore.Build.cs
@@ -19,5 +19,7 @@ public class NetCore : ModuleRules
 		);
 
 		UnsafeTypeCastWarningLevel = WarningLevel.Error;
+
+		bAllowAutoRTFMInstrumentation = true;
 	}
 }
