Allow Gameplay Abilities to use Blueprintable Replicated Variables.

diff --git a/Engine/Plugins/Runtime/GameplayAbilities/Source/GameplayAbilities/Private/Abilities/GameplayAbility.cpp b/Engine/Plugins/Runtime/GameplayAbilities/Source/GameplayAbilities/Private/Abilities/GameplayAbility.cpp
--- a/Engine/Plugins/Runtime/GameplayAbilities/Source/GameplayAbilities/Private/Abilities/GameplayAbility.cpp
+++ b/Engine/Plugins/Runtime/GameplayAbilities/Source/GameplayAbilities/Private/Abilities/GameplayAbility.cpp
@@ -14,8 +14,11 @@
 #include "GameplayCue_Types.h"
 #include "GameFramework/CharacterMovementComponent.h"
 #include "GameFramework/Character.h"
+#include "Misc/DataValidation.h"
 #include "UObject/Package.h"
 
+#define LOCTEXT_NAMESPACE "GameplayAbility"
+
 #if UE_WITH_IRIS
 #include "Iris/ReplicationSystem/ReplicationFragmentUtil.h"
 #endif // UE_WITH_IRIS
@@ -196,6 +199,35 @@ bool UGameplayAbility::IsSupportedForNetworking() const
 	return Supported;
 }
 
+#if WITH_EDITOR
+EDataValidationResult UGameplayAbility::IsDataValid(FDataValidationContext& Context) const
+{
+	EDataValidationResult Result = EDataValidationResult::Valid;
+
+	if (GetReplicationPolicy() == EGameplayAbilityReplicationPolicy::ReplicateNo)
+	{
+		UBlueprintGeneratedClass* BPClass = Cast<UBlueprintGeneratedClass>(GetClass());
+		if (BPClass && BPClass->NumReplicatedProperties > 0)
+		{
+			Context.AddError(LOCTEXT("ReplicatedVariablesNeedReplicationYes", "Gameplay Ability Blueprint has replicated variables but Replication Policy is set to not replicate"));
+			Result = EDataValidationResult::Invalid;
+		}
+	}
+
+	for (TFieldIterator<const UFunction> FuncIter(GetClass(), EFieldIterationFlags::IncludeSuper); FuncIter; ++FuncIter)
+	{
+		if (FuncIter->HasAnyFunctionFlags(EFunctionFlags::FUNC_NetMulticast))
+		{
+			FText ErrorText = FText::Format(LOCTEXT("MulticastFunctionDisallowed", "Gameplay Abilities are not replicated to Simulated Proxies and therefore NetMulticast Function {0} is meaningless"), FText::FromString(FuncIter->GetName()));
+			Context.AddError(ErrorText);
+			Result = EDataValidationResult::Invalid;
+		}
+	}
+
+	return Result;
+}
+#endif
+
 bool UGameplayAbility::DoesAbilitySatisfyTagRequirements(const UAbilitySystemComponent& AbilitySystemComponent, const FGameplayTagContainer* SourceTags, const FGameplayTagContainer* TargetTags, OUT FGameplayTagContainer* OptionalRelevantTags) const
 {
 	bool bBlocked = false;
@@ -2048,6 +2080,19 @@ void UGameplayAbility::NotifyAbilityTaskWaitingOnAvatar(class UAbilityTask* Abil
 	}
 }
 
+void UGameplayAbility::GetLifetimeReplicatedProps(TArray< class FLifetimeProperty >& OutLifetimeProps) const
+{
+	Super::GetLifetimeReplicatedProps(OutLifetimeProps);
+
+	if (GetReplicationPolicy() != EGameplayAbilityReplicationPolicy::ReplicateNo)
+	{
+		if (UBlueprintGeneratedClass* BPClass = Cast<UBlueprintGeneratedClass>(GetClass()))
+		{
+			BPClass->GetLifetimeBlueprintReplicationList(OutLifetimeProps);
+		}
+	}
+}
+
 #if UE_WITH_IRIS
 void UGameplayAbility::RegisterReplicationFragments(UE::Net::FFragmentRegistrationContext& Context, UE::Net::EFragmentRegistrationFlags RegistrationFlags)
 {
@@ -2055,3 +2100,5 @@ void UGameplayAbility::RegisterReplicationFragments(UE::Net::FFragmentRegistrati
 	UE::Net::FReplicationFragmentUtil::CreateAndRegisterFragmentsForObject(this, Context, RegistrationFlags);
 }
 #endif // UE_WITH_IRIS
+
+#undef LOCTEXT_NAMESPACE
\ No newline at end of file
diff --git a/Engine/Plugins/Runtime/GameplayAbilities/Source/GameplayAbilities/Public/Abilities/GameplayAbility.h b/Engine/Plugins/Runtime/GameplayAbilities/Source/GameplayAbilities/Public/Abilities/GameplayAbility.h
--- a/Engine/Plugins/Runtime/GameplayAbilities/Source/GameplayAbilities/Public/Abilities/GameplayAbility.h
+++ b/Engine/Plugins/Runtime/GameplayAbilities/Source/GameplayAbilities/Public/Abilities/GameplayAbility.h
@@ -502,6 +502,13 @@ public:
 	virtual bool CallRemoteFunction(UFunction* Function, void* Parameters, FOutParmRec* OutParms, FFrame* Stack) override;
 	virtual bool IsSupportedForNetworking() const override;
 
+#if WITH_EDITOR
+	virtual EDataValidationResult IsDataValid(FDataValidationContext& Context) const override;
+#endif
+
+	/** Overridden to allow Blueprint replicated properties to work */
+	virtual void GetLifetimeReplicatedProps(TArray< class FLifetimeProperty >& OutLifetimeProps) const;
+
 #if UE_WITH_IRIS
 	/** Register all replication fragments */
 	virtual void RegisterReplicationFragments(UE::Net::FFragmentRegistrationContext& Context, UE::Net::EFragmentRegistrationFlags RegistrationFlags) override;
