Remove call to UI update function that is not safe to call during gameplay effect loading, this was causing a large performance regression and will be replaced by a different call in a future fix

diff --git a/Engine/Plugins/Runtime/GameplayAbilities/Source/GameplayAbilities/Private/GameplayEffect.cpp b/Engine/Plugins/Runtime/GameplayAbilities/Source/GameplayAbilities/Private/GameplayEffect.cpp
--- a/Engine/Plugins/Runtime/GameplayAbilities/Source/GameplayAbilities/Private/GameplayEffect.cpp
+++ b/Engine/Plugins/Runtime/GameplayAbilities/Source/GameplayAbilities/Private/GameplayEffect.cpp
@@ -5498,7 +5498,8 @@ void FInheritedTagContainer::UpdateInheritedTagProperties(const FInheritedTagCon
 
 #if WITH_EDITOR
 	// Notify any open editor windows that we should refresh these properties, specifically for refreshing CombinedTags.
-	UGameplayTagsManager::OnEditorRefreshGameplayTagTree.Broadcast();
+	// TODO: This needs to call a different delegate as this the general editor refresh delegate cannot be called while loading gameplay effects
+	// UGameplayTagsManager::OnEditorRefreshGameplayTagTree.Broadcast();
 #endif
 }
 
