[Backout] - CL26417945
#tests Summoned a defender and watched them wander, pick up ammo, and attack husks
Original CL Desc
-----------------------------------------------------------------
PR #7531: GAS stops search for PC only if valid PC is found

#changelist validated

#virtualized

diff --git a/Engine/Plugins/Runtime/GameplayAbilities/Source/GameplayAbilities/Private/GameplayAbilityTypes.cpp b/Engine/Plugins/Runtime/GameplayAbilities/Source/GameplayAbilities/Private/GameplayAbilityTypes.cpp
--- a/Engine/Plugins/Runtime/GameplayAbilities/Source/GameplayAbilities/Private/GameplayAbilityTypes.cpp
+++ b/Engine/Plugins/Runtime/GameplayAbilities/Source/GameplayAbilities/Private/GameplayAbilityTypes.cpp
@@ -43,11 +43,7 @@ void FGameplayAbilityActorInfo::InitFromActor(AActor *InOwnerActor, AActor *InAv
 		if (APawn * Pawn = Cast<APawn>(TestActor))
 		{
 			PlayerController = Cast<APlayerController>(Pawn->GetController());
-
-			if (PlayerController.IsValid())
-			{
-				break;
-			}
+			break;
 		}
 
 		TestActor = TestActor->GetOwner();
