Don't warn when Gameplay Tag Tests fail to have the prequisites to run the Test.  Instead, use Info-level logs and pass.

diff --git a/Engine/Plugins/Runtime/GameplayAbilities/Source/GameplayAbilities/Private/Tests/GameplayTagQueryTests.cpp b/Engine/Plugins/Runtime/GameplayAbilities/Source/GameplayAbilities/Private/Tests/GameplayTagQueryTests.cpp
--- a/Engine/Plugins/Runtime/GameplayAbilities/Source/GameplayAbilities/Private/Tests/GameplayTagQueryTests.cpp
+++ b/Engine/Plugins/Runtime/GameplayAbilities/Source/GameplayAbilities/Private/Tests/GameplayTagQueryTests.cpp
@@ -42,7 +42,7 @@ bool FGameplayTagQueryMatchesTagRequirementsTest::RunTest(const FString& Paramet
 
 		if (AllGameplayTags.Num() < NumDebugTags)
 		{
-			AddWarning(FString::Printf(TEXT("There are only %d defined tags in the Project.  We need at least %d to run this test."), AllGameplayTags.Num(), NumDebugTags));
+			AddInfo(FString::Printf(TEXT("There are only %d defined tags in the Project.  We need at least %d to run this test (skipping)."), AllGameplayTags.Num(), NumDebugTags));
 			return true;
 		}
 
@@ -128,5 +128,6 @@ bool FGameplayTagQueryMatchesTagRequirementsTest::RunTest(const FString& Paramet
 	AddInfo(FString::Printf(TEXT("	FGameplayTagRequirements took %ull cycles (%.4lf ms) and gave %d matches"), MetRequirementsCycles, FPlatformTime::ToMilliseconds64(MetRequirementsCycles), NumMetRequirements));
 	AddInfo(FString::Printf(TEXT("	FGameplayTagQuery took %ull cycles (%.4lf ms) and gave %d matches"), QueryMatchesCycles, FPlatformTime::ToMilliseconds64(QueryMatchesCycles), NumQueryMatches));
 
-	return true;
+	TestEqual(TEXT("FGameplayTagRequirements has same matching rules as FGameplayTagQuery"), NumQueryMatches, NumMetRequirements);
+	return NumMetRequirements == NumQueryMatches;
 }
diff --git a/Engine/Source/Runtime/GameplayTags/Private/Tests/GameplayTagQueryTests.cpp b/Engine/Source/Runtime/GameplayTags/Private/Tests/GameplayTagQueryTests.cpp
--- a/Engine/Source/Runtime/GameplayTags/Private/Tests/GameplayTagQueryTests.cpp
+++ b/Engine/Source/Runtime/GameplayTags/Private/Tests/GameplayTagQueryTests.cpp
@@ -12,9 +12,10 @@ bool FGameplayTagQueryTest_Empty::RunTest(const FString& Parameters)
 {
 	FGameplayTagContainer AllGameplayTags;
 	UGameplayTagsManager::Get().RequestAllGameplayTags(AllGameplayTags, true);
-	if ( !TestTrue(TEXT("GameplayTags are defined for this Project"), !AllGameplayTags.IsEmpty()) )
+	if ( AllGameplayTags.IsEmpty() )
 	{
-		return false;
+		AddInfo(TEXT("GameplayTags are not defined for this Project. Skipping test."));
+		return true;
 	}
 
 	FGameplayTag RandomTag = AllGameplayTags.First();
