Add WaitGameplayTagChanged async tasks (ability and actor versions).
This task allows reacting to all tag count changes as opposed to just Add and Remove.

diff --git a/Engine/Plugins/Runtime/GameplayAbilities/Source/GameplayAbilities/Private/Abilities/Async/AbilityAsync_WaitGameplayTagCountChanged.cpp b/Engine/Plugins/Runtime/GameplayAbilities/Source/GameplayAbilities/Private/Abilities/Async/AbilityAsync_WaitGameplayTagCountChanged.cpp
new file mode 100644
--- /dev/null
+++ b/Engine/Plugins/Runtime/GameplayAbilities/Source/GameplayAbilities/Private/Abilities/Async/AbilityAsync_WaitGameplayTagCountChanged.cpp
@@ -0,0 +1,53 @@
+// Copyright Epic Games, Inc. All Rights Reserved.
+
+#include "Abilities/Async/AbilityAsync_WaitGameplayTagCountChanged.h"
+#include "AbilitySystemComponent.h"
+#include "AbilitySystemLog.h"
+
+void UAbilityAsync_WaitGameplayTagCountChanged::Activate()
+{
+	Super::Activate();
+
+	UAbilitySystemComponent* ASC = GetAbilitySystemComponent();
+	if (ASC && ShouldBroadcastDelegates())
+	{
+		GameplayTagCountChangedHandle = ASC->RegisterGameplayTagEvent(Tag, EGameplayTagEventType::AnyCountChange).AddUObject(this, &UAbilityAsync_WaitGameplayTagCountChanged::GameplayTagCallback);
+	}
+	else
+	{
+		ABILITY_LOG(Warning, TEXT("%s: AbilitySystemComponent is nullptr! Could not register for gameplay tag event with Tag = %s."), *GetName(), *Tag.GetTagName().ToString());
+		EndAction();
+	}
+}
+
+void UAbilityAsync_WaitGameplayTagCountChanged::EndAction()
+{
+	UAbilitySystemComponent* ASC = GetAbilitySystemComponent();
+	if (ASC && GameplayTagCountChangedHandle.IsValid())
+	{
+		ASC->RegisterGameplayTagEvent(Tag).Remove(GameplayTagCountChangedHandle);
+		GameplayTagCountChangedHandle.Reset();
+	}
+	Super::EndAction();
+}
+
+void UAbilityAsync_WaitGameplayTagCountChanged::GameplayTagCallback(const FGameplayTag InTag, int32 NewCount)
+{
+	if (ShouldBroadcastDelegates())
+	{
+		TagCountChanged.Broadcast(NewCount);
+	}
+	else
+	{
+		EndAction();
+	}
+}
+
+UAbilityAsync_WaitGameplayTagCountChanged* UAbilityAsync_WaitGameplayTagCountChanged::WaitGameplayTagCountChangedOnActor(AActor* TargetActor, FGameplayTag Tag)
+{
+	UAbilityAsync_WaitGameplayTagCountChanged* MyObj = NewObject<UAbilityAsync_WaitGameplayTagCountChanged>();
+	MyObj->SetAbilityActor(TargetActor);
+	MyObj->Tag = Tag;
+
+	return MyObj;
+}
diff --git a/Engine/Plugins/Runtime/GameplayAbilities/Source/GameplayAbilities/Private/Abilities/Tasks/AbilityTask_WaitGameplayTagCountChanged.cpp b/Engine/Plugins/Runtime/GameplayAbilities/Source/GameplayAbilities/Private/Abilities/Tasks/AbilityTask_WaitGameplayTagCountChanged.cpp
new file mode 100644
--- /dev/null
+++ b/Engine/Plugins/Runtime/GameplayAbilities/Source/GameplayAbilities/Private/Abilities/Tasks/AbilityTask_WaitGameplayTagCountChanged.cpp
@@ -0,0 +1,63 @@
+// Copyright Epic Games, Inc. All Rights Reserved.
+
+#include "Abilities/Tasks/AbilityTask_WaitGameplayTagCountChanged.h"
+#include "AbilitySystemComponent.h"
+#include "AbilitySystemGlobals.h"
+#include "AbilitySystemLog.h"
+
+UAbilityTask_WaitGameplayTagCountChanged* UAbilityTask_WaitGameplayTagCountChanged::WaitGameplayTagCountChange(UGameplayAbility* OwningAbility, FGameplayTag Tag, AActor* InOptionalExternalTarget /*= nullptr*/)
+{
+	UAbilityTask_WaitGameplayTagCountChanged* MyObj = NewAbilityTask<UAbilityTask_WaitGameplayTagCountChanged>(OwningAbility);
+	MyObj->Tag = Tag;
+	if (InOptionalExternalTarget)
+	{
+		MyObj->OptionalExternalTarget = UAbilitySystemGlobals::GetAbilitySystemComponentFromActor(InOptionalExternalTarget);
+		if (!MyObj->OptionalExternalTarget)
+		{
+			ABILITY_LOG(Warning, TEXT("%s: Passed OptionalExternalTarget's AbilitySystemComponent is nullptr! Will not be able to register for gameplay tag event with Tag = %s. Will instead register event with the AbilitySystemComponent this task's owner (%s)."), *MyObj->GetName(), *Tag.GetTagName().ToString(), *GetNameSafe(MyObj->GetOwnerActor()));
+		}
+	}
+
+	return MyObj;
+}
+
+void UAbilityTask_WaitGameplayTagCountChanged::Activate()
+{
+	Super::Activate();
+
+	UAbilitySystemComponent* ASC = GetTargetASC();
+	if (ASC)
+	{
+		GameplayTagCountChangedHandle = ASC->RegisterGameplayTagEvent(Tag, EGameplayTagEventType::AnyCountChange).AddUObject(this, &UAbilityTask_WaitGameplayTagCountChanged::GameplayTagCallback);
+	}
+}
+
+void UAbilityTask_WaitGameplayTagCountChanged::GameplayTagCallback(const FGameplayTag InTag, int32 NewCount)
+{
+	if (ShouldBroadcastAbilityTaskDelegates())
+	{
+		TagCountChanged.Broadcast(NewCount);
+	}
+}
+
+UAbilitySystemComponent* UAbilityTask_WaitGameplayTagCountChanged::GetTargetASC()
+{
+	if (OptionalExternalTarget)
+	{
+		return OptionalExternalTarget;
+	}
+
+	return AbilitySystemComponent.Get();
+}
+
+void UAbilityTask_WaitGameplayTagCountChanged::OnDestroy(bool bAbilityIsEnding)
+{
+	UAbilitySystemComponent* ASC = GetTargetASC();
+	if (GameplayTagCountChangedHandle.IsValid() && ASC)
+	{
+		ASC->RegisterGameplayTagEvent(Tag).Remove(GameplayTagCountChangedHandle);
+		GameplayTagCountChangedHandle.Reset();
+	}
+
+	Super::OnDestroy(bAbilityIsEnding);
+}
diff --git a/Engine/Plugins/Runtime/GameplayAbilities/Source/GameplayAbilities/Public/Abilities/Async/AbilityAsync_WaitGameplayTagCountChanged.h b/Engine/Plugins/Runtime/GameplayAbilities/Source/GameplayAbilities/Public/Abilities/Async/AbilityAsync_WaitGameplayTagCountChanged.h
new file mode 100644
--- /dev/null
+++ b/Engine/Plugins/Runtime/GameplayAbilities/Source/GameplayAbilities/Public/Abilities/Async/AbilityAsync_WaitGameplayTagCountChanged.h
@@ -0,0 +1,33 @@
+// Copyright Epic Games, Inc. All Rights Reserved.
+#pragma once
+
+#include "AbilityAsync.h"
+#include "AbilityAsync_WaitGameplayTagCountChanged.generated.h"
+
+UCLASS()
+class GAMEPLAYABILITIES_API UAbilityAsync_WaitGameplayTagCountChanged : public UAbilityAsync
+{
+	GENERATED_BODY()
+protected:
+
+	DECLARE_DYNAMIC_MULTICAST_DELEGATE_OneParam(FAsyncWaitGameplayTagCountDelegate, int32, TagCount);
+
+	virtual void Activate() override;
+	virtual void EndAction() override;
+
+	virtual void GameplayTagCallback(const FGameplayTag Tag, int32 NewCount);
+
+	FGameplayTag Tag;
+	FDelegateHandle GameplayTagCountChangedHandle;
+
+public:
+	UPROPERTY(BlueprintAssignable)
+	FAsyncWaitGameplayTagCountDelegate TagCountChanged;
+
+	/**
+	 * Wait until the specified gameplay tag count changes on Target Actor's ability component
+	 * If used in an ability graph, this async action will wait even after activation ends. It's recommended to use WaitGameplayTagCountChange instead.
+	 */
+	UFUNCTION(BlueprintCallable, Category = "Ability|Tasks", meta = (DefaultToSelf = "TargetActor", BlueprintInternalUseOnly = "TRUE"))
+	static UAbilityAsync_WaitGameplayTagCountChanged* WaitGameplayTagCountChangedOnActor(AActor* TargetActor, FGameplayTag Tag);
+};
\ No newline at end of file
diff --git a/Engine/Plugins/Runtime/GameplayAbilities/Source/GameplayAbilities/Public/Abilities/Tasks/AbilityTask_WaitGameplayTagCountChanged.h b/Engine/Plugins/Runtime/GameplayAbilities/Source/GameplayAbilities/Public/Abilities/Tasks/AbilityTask_WaitGameplayTagCountChanged.h
new file mode 100644
--- /dev/null
+++ b/Engine/Plugins/Runtime/GameplayAbilities/Source/GameplayAbilities/Public/Abilities/Tasks/AbilityTask_WaitGameplayTagCountChanged.h
@@ -0,0 +1,38 @@
+// Copyright Epic Games, Inc. All Rights Reserved.
+#pragma once
+
+#include "Abilities/Tasks/AbilityTask.h"
+#include "GameplayTagContainer.h"
+#include "AbilityTask_WaitGameplayTagCountChanged.generated.h"
+
+DECLARE_DYNAMIC_MULTICAST_DELEGATE_OneParam(FWaitGameplayTagCountDelegate, int32, TagCount);
+
+UCLASS()
+class GAMEPLAYABILITIES_API UAbilityTask_WaitGameplayTagCountChanged : public UAbilityTask
+{
+	GENERATED_BODY()
+
+protected:
+	UPROPERTY(BlueprintAssignable)
+	FWaitGameplayTagCountDelegate TagCountChanged;
+
+	/**
+	 * 	Wait until the specified gameplay tag count has changed. By default this will look at the owner of this ability. OptionalExternalTarget can be set to make this look at another actor's tags for changes. 
+	 */
+	UFUNCTION(BlueprintCallable, Category = "Ability|Tasks", meta = (HidePin = "OwningAbility", DefaultToSelf = "OwningAbility", BlueprintInternalUseOnly = "TRUE"))
+	static UAbilityTask_WaitGameplayTagCountChanged* WaitGameplayTagCountChange(UGameplayAbility* OwningAbility, FGameplayTag Tag, AActor* InOptionalExternalTarget = nullptr);
+
+	virtual void Activate() override;
+	virtual void GameplayTagCallback(const FGameplayTag Tag, int32 NewCount);
+
+	FGameplayTag Tag;
+
+	UAbilitySystemComponent* GetTargetASC();
+	virtual void OnDestroy(bool bAbilityIsEnding) override;
+
+	UPROPERTY()
+	TObjectPtr<UAbilitySystemComponent> OptionalExternalTarget;
+
+	FDelegateHandle GameplayTagCountChangedHandle;
+};
+
