Introducing UGameplayEffectComponent::OnGameplayEffectApplied which corrects the rules for applying stacked GameplayEffects.  The rules are simpler for Applied, but the other functions (Executed, AddedToContainer) still exist due to the complex rules.

diff --git a/Engine/Plugins/Runtime/GameplayAbilities/Source/GameplayAbilities/Private/AbilitySystemComponent.cpp b/Engine/Plugins/Runtime/GameplayAbilities/Source/GameplayAbilities/Private/AbilitySystemComponent.cpp
--- a/Engine/Plugins/Runtime/GameplayAbilities/Source/GameplayAbilities/Private/AbilitySystemComponent.cpp
+++ b/Engine/Plugins/Runtime/GameplayAbilities/Source/GameplayAbilities/Private/AbilitySystemComponent.cpp
@@ -926,7 +926,10 @@ FActiveGameplayEffectHandle UAbilitySystemComponent::ApplyGameplayEffectSpecToSe
 		// This is a non-predicted instant effect (it never gets added to ActiveGameplayEffects)
 		ExecuteGameplayEffect(*OurCopyOfSpec, PredictionKey);
 	}
-	
+
+	// Notify the Gameplay Effect (and its Components) that it has been successfully applied
+	Spec.Def->OnApplied(ActiveGameplayEffects, *OurCopyOfSpec, PredictionKey);
+
 	UAbilitySystemComponent* InstigatorASC = Spec.GetContext().GetInstigatorAbilitySystemComponent();
 
 	// Send ourselves a callback	
diff --git a/Engine/Plugins/Runtime/GameplayAbilities/Source/GameplayAbilities/Private/AbilitySystemGlobals.cpp b/Engine/Plugins/Runtime/GameplayAbilities/Source/GameplayAbilities/Private/AbilitySystemGlobals.cpp
--- a/Engine/Plugins/Runtime/GameplayAbilities/Source/GameplayAbilities/Private/AbilitySystemGlobals.cpp
+++ b/Engine/Plugins/Runtime/GameplayAbilities/Source/GameplayAbilities/Private/AbilitySystemGlobals.cpp
@@ -942,9 +942,9 @@ FAutoConsoleCommand DebugAbilitySystemEffectListActiveCommand(TEXT("AbilitySyste
 			if (const FActiveGameplayEffect* ActiveGE = ASC->GetActiveGameplayEffect(ActiveGEHandle))
 			{
 				const UGameplayEffect* GESpecDef = ActiveGE->Spec.Def;
-				const TCHAR* ActiveState = ActiveGE->bIsInhibited ? TEXT("Inhibited") : TEXT("Active");
+				const FString ActiveState = ActiveGE->bIsInhibited ? FString(TEXT("Inhibited")) : FString::Printf(TEXT("Active (Stack: %d)"), ActiveGE->Spec.GetStackCount());
 
-				OutputDevice.Logf(TEXT(" [%u] %s: %s"), GetTypeHash(ActiveGEHandle), *GetNameSafe(GESpecDef), ActiveState);
+				OutputDevice.Logf(TEXT(" [%u] %s: %s"), GetTypeHash(ActiveGEHandle), *GetNameSafe(GESpecDef), *ActiveState);
 			}
 		}
 	}), ECVF_Cheat);
diff --git a/Engine/Plugins/Runtime/GameplayAbilities/Source/GameplayAbilities/Private/GameplayEffect.cpp b/Engine/Plugins/Runtime/GameplayAbilities/Source/GameplayAbilities/Private/GameplayEffect.cpp
--- a/Engine/Plugins/Runtime/GameplayAbilities/Source/GameplayAbilities/Private/GameplayEffect.cpp
+++ b/Engine/Plugins/Runtime/GameplayAbilities/Source/GameplayAbilities/Private/GameplayEffect.cpp
@@ -319,13 +319,13 @@ void UGameplayEffect::PostCDOCompiled(const FPostCDOCompiledContext& Context)
 
 	// Move any data from deprecated properties into components, create components as needed
 	ConvertAbilitiesComponent();
+	ConvertRemoveOtherComponent(); // This should come before AdditionalEffects for legacy compatibility
 	ConvertAdditionalEffectsComponent();
 	ConvertAssetTagsComponent();
 	ConvertBlockByTagsComponent();
 	ConvertChanceToApplyComponent();
 	ConvertCustomCanApplyComponent();
 	ConvertImmunityComponent();
-	ConvertRemoveOtherComponent();
 	ConvertTagRequirementsComponent();
 	ConvertTargetTagsComponent();
 	ConvertUIComponent();
@@ -752,6 +752,19 @@ void UGameplayEffect::OnExecuted(FActiveGameplayEffectsContainer& ActiveGEContai
 	}
 }
 
+void UGameplayEffect::OnApplied(FActiveGameplayEffectsContainer& ActiveGEContainer, FGameplayEffectSpec& GESpec, FPredictionKey& PredictionKey) const
+{
+	UE_VLOG(ActiveGEContainer.Owner, LogAbilitySystem, Verbose, TEXT("GameplayEffect %s applied"), *GetNameSafe(GESpec.Def));
+
+	for (const UGameplayEffectComponent* GEComponent : GEComponents)
+	{
+		if (GEComponent)
+		{
+			GEComponent->OnGameplayEffectApplied(ActiveGEContainer, GESpec, PredictionKey);
+		}
+	}
+}
+
 int32 UGameplayEffect::GetStackLimitCount() const
 {
 	return StackLimitCount;
diff --git a/Engine/Plugins/Runtime/GameplayAbilities/Source/GameplayAbilities/Private/GameplayEffectComponents/AdditionalEffectsGameplayEffectComponent.cpp b/Engine/Plugins/Runtime/GameplayAbilities/Source/GameplayAbilities/Private/GameplayEffectComponents/AdditionalEffectsGameplayEffectComponent.cpp
--- a/Engine/Plugins/Runtime/GameplayAbilities/Source/GameplayAbilities/Private/GameplayEffectComponents/AdditionalEffectsGameplayEffectComponent.cpp
+++ b/Engine/Plugins/Runtime/GameplayAbilities/Source/GameplayAbilities/Private/GameplayEffectComponents/AdditionalEffectsGameplayEffectComponent.cpp
@@ -9,12 +9,6 @@
 
 bool UAdditionalEffectsGameplayEffectComponent::OnActiveGameplayEffectAdded(FActiveGameplayEffectsContainer& ActiveGEContainer, FActiveGameplayEffect& ActiveGE) const
 {
-	UAbilitySystemComponent* ASC = ActiveGEContainer.Owner;
-	if (ensure(ASC))
-	{
-		OnGameplayEffectApplied(*ASC, ActiveGE.Spec, ActiveGE.PredictionKey);
-	}
-
 	// We don't allow prediction of expiration (on removed) effects
 	if (ActiveGEContainer.IsNetAuthority())
 	{
@@ -25,20 +19,13 @@ bool UAdditionalEffectsGameplayEffectComponent::OnActiveGameplayEffectAdded(FAct
 	return true;
 }
 
-void UAdditionalEffectsGameplayEffectComponent::OnGameplayEffectExecuted(FActiveGameplayEffectsContainer& ActiveGEContainer, FGameplayEffectSpec& GESpec, FPredictionKey& PredictionKey) const
+void UAdditionalEffectsGameplayEffectComponent::OnGameplayEffectApplied(FActiveGameplayEffectsContainer& ActiveGEContainer, FGameplayEffectSpec& GESpec, FPredictionKey& PredictionKey) const
 {
-	if (GetOwner()->DurationPolicy == EGameplayEffectDurationType::Instant)
+	if (!ensureMsgf(ActiveGEContainer.Owner, TEXT("OnGameplayEffectApplied is passed an ActiveGEContainer which lives within an ASC but that ASC was somehow null")))
 	{
-		UAbilitySystemComponent* ASC = ActiveGEContainer.Owner;
-		if (ensure(ASC))
-		{
-			OnGameplayEffectApplied(*ASC, GESpec, PredictionKey);
-		}
+		return;
 	}
-}
 
-void UAdditionalEffectsGameplayEffectComponent::OnGameplayEffectApplied(UAbilitySystemComponent& AppliedToASC, FGameplayEffectSpec& GESpec, FPredictionKey& PredictionKey) const
-{
 	const float GELevel = GESpec.GetLevel();
 	const FGameplayEffectContextHandle& GEContextHandle = GESpec.GetEffectContext();
 
@@ -79,6 +66,7 @@ PRAGMA_DISABLE_DEPRECATION_WARNINGS
 	TargetEffectSpecs.Append(GESpec.TargetEffectSpecs);
 PRAGMA_ENABLE_DEPRECATION_WARNINGS
 
+	UAbilitySystemComponent& AppliedToASC = *ActiveGEContainer.Owner;
 	for (const FGameplayEffectSpecHandle& TargetSpec : TargetEffectSpecs)
 	{
 		if (TargetSpec.IsValid())
diff --git a/Engine/Plugins/Runtime/GameplayAbilities/Source/GameplayAbilities/Private/GameplayEffectComponents/RemoveOtherGameplayEffectComponent.cpp b/Engine/Plugins/Runtime/GameplayAbilities/Source/GameplayAbilities/Private/GameplayEffectComponents/RemoveOtherGameplayEffectComponent.cpp
--- a/Engine/Plugins/Runtime/GameplayAbilities/Source/GameplayAbilities/Private/GameplayEffectComponents/RemoveOtherGameplayEffectComponent.cpp
+++ b/Engine/Plugins/Runtime/GameplayAbilities/Source/GameplayAbilities/Private/GameplayEffectComponents/RemoveOtherGameplayEffectComponent.cpp
@@ -13,27 +13,19 @@ URemoveOtherGameplayEffectComponent::URemoveOtherGameplayEffectComponent()
 #endif
 }
 
-bool URemoveOtherGameplayEffectComponent::OnActiveGameplayEffectAdded(FActiveGameplayEffectsContainer& ActiveGEContainer, FActiveGameplayEffect& ActiveGE) const
+void URemoveOtherGameplayEffectComponent::OnGameplayEffectApplied(FActiveGameplayEffectsContainer& ActiveGEContainer, FGameplayEffectSpec& GESpec, FPredictionKey& PredictionKey) const
 {
-	if (ActiveGEContainer.OwnerIsNetAuthority)
+	if (!ActiveGEContainer.OwnerIsNetAuthority)
 	{
-		OnGameplayEffectApplied(ActiveGEContainer, ActiveGE.Handle);
+		return;
 	}
-	return true;
-}
 
-void URemoveOtherGameplayEffectComponent::OnGameplayEffectExecuted(FActiveGameplayEffectsContainer& ActiveGEContainer, FGameplayEffectSpec& GESpec, FPredictionKey& PredictionKey) const
-{
-	const bool bInstantEffect = (GetOwner()->DurationPolicy == EGameplayEffectDurationType::Instant);
-	if (bInstantEffect && ActiveGEContainer.OwnerIsNetAuthority)
-	{
-		FActiveGameplayEffectHandle InvalidHandle;
-		OnGameplayEffectApplied(ActiveGEContainer, InvalidHandle);
-	}
-}
+	FGameplayEffectQuery FindOwnerQuery;
+	FindOwnerQuery.EffectDefinition = GetOwner() ? GetOwner()->GetClass() : nullptr;
+
+	// We need to keep track to ensure we never remove ourselves
+	TArray<FActiveGameplayEffectHandle> ActiveGEHandles = ActiveGEContainer.GetActiveEffects(FindOwnerQuery);
 
-void URemoveOtherGameplayEffectComponent::OnGameplayEffectApplied(FActiveGameplayEffectsContainer& ActiveGEContainer, const FActiveGameplayEffectHandle& ActiveGEHandle) const
-{
 	constexpr int32 RemoveAllStacks = -1;
 	for (const FGameplayEffectQuery& RemoveQuery : RemoveGameplayEffectQueries)
 	{
@@ -41,7 +33,7 @@ void URemoveOtherGameplayEffectComponent::OnGameplayEffectApplied(FActiveGamepla
 		{
 			// If we have an ActiveGEHandle, make sure we never remove ourselves.
 			// If we don't, there's no need to make a copy.
-			if (!ActiveGEHandle.IsValid())
+			if (ActiveGEHandles.IsEmpty())
 			{
 				// Faster path: No copy needed
 				ActiveGEContainer.RemoveActiveEffects(RemoveQuery, RemoveAllStacks);
@@ -49,7 +41,7 @@ void URemoveOtherGameplayEffectComponent::OnGameplayEffectApplied(FActiveGamepla
 			else
 			{
 				FGameplayEffectQuery MutableRemoveQuery = RemoveQuery;
-				MutableRemoveQuery.IgnoreHandles.Add(ActiveGEHandle);
+				MutableRemoveQuery.IgnoreHandles = MoveTemp(ActiveGEHandles);
 
 				ActiveGEContainer.RemoveActiveEffects(MutableRemoveQuery, RemoveAllStacks);
 			}
diff --git a/Engine/Plugins/Runtime/GameplayAbilities/Source/GameplayAbilities/Private/GameplayEffectComponents/TargetTagRequirementsGameplayEffectComponent.cpp b/Engine/Plugins/Runtime/GameplayAbilities/Source/GameplayAbilities/Private/GameplayEffectComponents/TargetTagRequirementsGameplayEffectComponent.cpp
--- a/Engine/Plugins/Runtime/GameplayAbilities/Source/GameplayAbilities/Private/GameplayEffectComponents/TargetTagRequirementsGameplayEffectComponent.cpp
+++ b/Engine/Plugins/Runtime/GameplayAbilities/Source/GameplayAbilities/Private/GameplayEffectComponents/TargetTagRequirementsGameplayEffectComponent.cpp
@@ -33,12 +33,6 @@ bool UTargetTagRequirementsGameplayEffectComponent::CanGameplayEffectApply(const
 	return true;
 }
 
-void UTargetTagRequirementsGameplayEffectComponent::OnGameplayEffectExecuted(FActiveGameplayEffectsContainer& ActiveGEContainer, FGameplayEffectSpec& GESpec, FPredictionKey& PredictionKey) const
-{
-	const bool bInstantEffect = (GetOwner()->DurationPolicy == EGameplayEffectDurationType::Instant);
-	UE_CLOG(bInstantEffect && !OngoingTagRequirements.IsEmpty(), LogGameplayEffects, Warning, TEXT("%s is instant but has tag requirements. Tag requirements can only be used with gameplay effects that have a duration. Tag requirements will be ignored."), *GetOwner()->GetPathName());
-}
-
 // UTargetTagRequirementsGameplayEffectComponent lives on an asset.  This doesn't get instanced at runtime, so this is NOT A UNIQUE INSTANCE (it is a shared instance for any GEContainer/ActiveGE that wants to use it).
 bool UTargetTagRequirementsGameplayEffectComponent::OnActiveGameplayEffectAdded(FActiveGameplayEffectsContainer& GEContainer, FActiveGameplayEffect& ActiveGE) const
 {
diff --git a/Engine/Plugins/Runtime/GameplayAbilities/Source/GameplayAbilities/Public/GameplayEffect.h b/Engine/Plugins/Runtime/GameplayAbilities/Source/GameplayAbilities/Public/GameplayEffect.h
--- a/Engine/Plugins/Runtime/GameplayAbilities/Source/GameplayAbilities/Public/GameplayEffect.h
+++ b/Engine/Plugins/Runtime/GameplayAbilities/Source/GameplayAbilities/Public/GameplayEffect.h
@@ -2093,6 +2093,12 @@ public:
      */
 	void OnExecuted(FActiveGameplayEffectsContainer& ActiveGEContainer, FGameplayEffectSpec& GESpec, FPredictionKey& PredictionKey) const;
 
+	/**
+	 * Receive a notify that this GameplayEffect has been applied (this GE is either previously added to the container or executed in such cases).
+	 * However, this also encompasses cases where a GE was added to a container previously and then applied again to 'stack'.  It does not happen for periodic executions of a duration GE.
+	 */
+	void OnApplied(FActiveGameplayEffectsContainer& ActiveGEContainer, FGameplayEffectSpec& GESpec, FPredictionKey& PredictionKey) const;
+
 	/** Returns all tags that this GE *has* and *does not* grant to any Actor. */
 	const FGameplayTagContainer& GetAssetTags() const { return CachedAssetTags; }
 
diff --git a/Engine/Plugins/Runtime/GameplayAbilities/Source/GameplayAbilities/Public/GameplayEffectComponent.h b/Engine/Plugins/Runtime/GameplayAbilities/Source/GameplayAbilities/Public/GameplayEffectComponent.h
--- a/Engine/Plugins/Runtime/GameplayAbilities/Source/GameplayAbilities/Public/GameplayEffectComponent.h
+++ b/Engine/Plugins/Runtime/GameplayAbilities/Source/GameplayAbilities/Public/GameplayEffectComponent.h
@@ -45,17 +45,24 @@ public:
 	virtual bool CanGameplayEffectApply(const FActiveGameplayEffectsContainer& ActiveGEContainer, const FGameplayEffectSpec& GESpec) const { return true; }
 
 	/**
-     * Called when a Gameplay Effect is Added to the ActiveGameplayEffectsContainer.  GE's are added to that container when they have duration (or are predicting locally)
-     * Return if the effect should remain active, or false to inhibit.  Note: Inhibit does not remove the effect (it remains added but dormant, waiting to unhibit).
+     * Called when a Gameplay Effect is Added to the ActiveGameplayEffectsContainer.  GE's are added to that container when they have duration (or are predicting locally).
+	 * Note: This also occurs as a result of replication (e.g. when the server replicates a GE to the owning client -- including the 'duplicate' GE after a prediction).
+     * Return if the effect should remain active, or false to inhibit.  Note: Inhibit does not remove the effect (it remains added but dormant, waiting to uninhibit).
      */
 	virtual bool OnActiveGameplayEffectAdded(FActiveGameplayEffectsContainer& ActiveGEContainer, FActiveGameplayEffect& ActiveGE) const { return true; }
 
 	/** 
 	 * Called when a Gameplay Effect is executed.  GE's can only Execute on ROLE_Authority.  GE's only Execute when they're applying an instant effect (otherwise they're added to the ActiveGameplayEffectsContainer).
-	 * Note:  Periodic effects Execute every period (and are also added to ActiveGameplayEffectsContainer).
+	 * Note: Periodic effects Execute every period (and are also added to ActiveGameplayEffectsContainer).  One may think of this as periodically executing an instant effect (and thus can only happen on the server).
 	 */
 	virtual void OnGameplayEffectExecuted(FActiveGameplayEffectsContainer& ActiveGEContainer, FGameplayEffectSpec& GESpec, FPredictionKey& PredictionKey) const {}
 
+	/**
+	 * Called when a Gameplay Effect is initially applied, or stacked.  GE's are 'applied' in both cases of duration or instant execution.  This call does not happen periodically, nor through replication.
+	 * One should favor this function over OnActiveGameplayEffectAdded & OnGameplayEffectExecuted (but all multiple may be used depending on the case).
+	 */
+	virtual void OnGameplayEffectApplied(FActiveGameplayEffectsContainer& ActiveGEContainer, FGameplayEffectSpec& GESpec, FPredictionKey& PredictionKey) const {}
+
 	/**
 	 * Let us know that the owning GameplayEffect has been modified, thus apply an asset-related changes to the owning GameplayEffect (e.g. any of its fields)
 	 */
diff --git a/Engine/Plugins/Runtime/GameplayAbilities/Source/GameplayAbilities/Public/GameplayEffectComponents/AdditionalEffectsGameplayEffectComponent.h b/Engine/Plugins/Runtime/GameplayAbilities/Source/GameplayAbilities/Public/GameplayEffectComponents/AdditionalEffectsGameplayEffectComponent.h
--- a/Engine/Plugins/Runtime/GameplayAbilities/Source/GameplayAbilities/Public/GameplayEffectComponents/AdditionalEffectsGameplayEffectComponent.h
+++ b/Engine/Plugins/Runtime/GameplayAbilities/Source/GameplayAbilities/Public/GameplayEffectComponents/AdditionalEffectsGameplayEffectComponent.h
@@ -14,15 +14,14 @@ class GAMEPLAYABILITIES_API UAdditionalEffectsGameplayEffectComponent : public U
 
 public:
 	/**
-     * Called when a Gameplay Effect is Added to the ActiveGameplayEffectsContainer.  GE's are added to that container when they have duration (or are predicting locally)
-     * Return if the effect should remain active, or false to inhibit.  Note: Inhibit does not remove the effect (it remains added but dormant, waiting to unhibit).
+     * Called when a Gameplay Effect is Added to the ActiveGameplayEffectsContainer.  We register a callback to execute the OnComplete Gameplay Effects.
      */
 	virtual bool OnActiveGameplayEffectAdded(FActiveGameplayEffectsContainer& GEContainer, FActiveGameplayEffect& ActiveGE) const override;
 
 	/**
-	 * Called when a Gameplay Effect is executed.  GE's can only Execute on ROLE_Authority.  GE's only Execute when they're instant (otherwise they're added to the ActiveGameplayEffectsContainer).
+	 * Called when a Gameplay Effect is applied.  This executes the OnApplication Gameplay Effects.
 	 */
-	virtual void OnGameplayEffectExecuted(FActiveGameplayEffectsContainer& ActiveGEContainer, FGameplayEffectSpec& GESpec, FPredictionKey& PredictionKey) const override;
+	virtual void OnGameplayEffectApplied(FActiveGameplayEffectsContainer& ActiveGEContainer, FGameplayEffectSpec& GESpec, FPredictionKey& PredictionKey) const override;
 
 #if WITH_EDITOR
 	/**
@@ -33,11 +32,6 @@ public:
 
 
 protected:
-	/** 
-	 * Define our own Applied which gets called in both cases (Added or Executed).
-	 */
-	void OnGameplayEffectApplied(UAbilitySystemComponent& AppliedToASC, FGameplayEffectSpec& GESpec, FPredictionKey& PredictionKey) const;
-
 	/**
 	 * Whenever the ActiveGE gets removed, we want to apply the configured OnComplete GameplayEffects
 	 */
diff --git a/Engine/Plugins/Runtime/GameplayAbilities/Source/GameplayAbilities/Public/GameplayEffectComponents/RemoveOtherGameplayEffectComponent.h b/Engine/Plugins/Runtime/GameplayAbilities/Source/GameplayAbilities/Public/GameplayEffectComponents/RemoveOtherGameplayEffectComponent.h
--- a/Engine/Plugins/Runtime/GameplayAbilities/Source/GameplayAbilities/Public/GameplayEffectComponents/RemoveOtherGameplayEffectComponent.h
+++ b/Engine/Plugins/Runtime/GameplayAbilities/Source/GameplayAbilities/Public/GameplayEffectComponents/RemoveOtherGameplayEffectComponent.h
@@ -18,11 +18,10 @@ public:
 	/** Constructor to set EditorFriendlyName */
 	URemoveOtherGameplayEffectComponent();
 
-	/** Once we've applied, we need to register for ongoing requirements */
-	virtual bool OnActiveGameplayEffectAdded(FActiveGameplayEffectsContainer& GEContainer, FActiveGameplayEffect& ActiveGE) const override;
-
-	/** If we're only executed, it's an indication something has gone wrong and we should log it */
-	virtual void OnGameplayEffectExecuted(FActiveGameplayEffectsContainer& ActiveGEContainer, FGameplayEffectSpec& GESpec, FPredictionKey& PredictionKey) const override;
+	/**
+	 * We will re-run RemoveGameplayEffectQueries every time the owning Gameplay Effect is applied.
+	 */
+	virtual void OnGameplayEffectApplied(FActiveGameplayEffectsContainer& ActiveGEContainer, FGameplayEffectSpec& GESpec, FPredictionKey& PredictionKey) const override;
 
 #if WITH_EDITOR
 	/**
@@ -31,9 +30,6 @@ public:
 	virtual EDataValidationResult IsDataValid(class FDataValidationContext& Context) const override;
 #endif // WITH_EDITOR
 
-private:
-	void OnGameplayEffectApplied(FActiveGameplayEffectsContainer& ActiveGEContainer, const FActiveGameplayEffectHandle& ActiveGEHandle) const;
-
 public:
 	/** On Application of the owning Gameplay Effect, any Active GameplayEffects that *match* these queries will be removed. */
 	UPROPERTY(EditDefaultsOnly, Category = None)
diff --git a/Engine/Plugins/Runtime/GameplayAbilities/Source/GameplayAbilities/Public/GameplayEffectComponents/TargetTagRequirementsGameplayEffectComponent.h b/Engine/Plugins/Runtime/GameplayAbilities/Source/GameplayAbilities/Public/GameplayEffectComponents/TargetTagRequirementsGameplayEffectComponent.h
--- a/Engine/Plugins/Runtime/GameplayAbilities/Source/GameplayAbilities/Public/GameplayEffectComponents/TargetTagRequirementsGameplayEffectComponent.h
+++ b/Engine/Plugins/Runtime/GameplayAbilities/Source/GameplayAbilities/Public/GameplayEffectComponents/TargetTagRequirementsGameplayEffectComponent.h
@@ -24,9 +24,6 @@ public:
 	/** Once we've applied, we need to register for ongoing requirements */
 	virtual bool OnActiveGameplayEffectAdded(FActiveGameplayEffectsContainer& GEContainer, FActiveGameplayEffect& ActiveGE) const override;
 
-	/** If we're only executed, it's an indication something has gone wrong and we should log it */
-	virtual void OnGameplayEffectExecuted(FActiveGameplayEffectsContainer& ActiveGEContainer, FGameplayEffectSpec& GESpec, FPredictionKey& PredictionKey) const override;
-
 #if WITH_EDITOR
 	/**
 	 * Validate incompatible configurations
