Fixing the AbilitySystemGlobals resetting the FActiveGameplayEffect Handle Cache too early (during a map change).

diff --git a/Engine/Plugins/Runtime/GameplayAbilities/Source/GameplayAbilities/Private/AbilitySystemGlobals.cpp b/Engine/Plugins/Runtime/GameplayAbilities/Source/GameplayAbilities/Private/AbilitySystemGlobals.cpp
--- a/Engine/Plugins/Runtime/GameplayAbilities/Source/GameplayAbilities/Private/AbilitySystemGlobals.cpp
+++ b/Engine/Plugins/Runtime/GameplayAbilities/Source/GameplayAbilities/Private/AbilitySystemGlobals.cpp
@@ -522,6 +522,21 @@ void UAbilitySystemGlobals::HandlePreLoadMap(const FWorldContext& WorldContext,
 		return;
 	}
 
+	// If we are preloading a map but coming from an existing map, then we should wait until the previous map is cleaned up,
+	// otherwise we'll end up stomping FActiveGameplayEffectHandle map.
+	if (const UWorld* InWorld = WorldContext.World())
+	{
+		FWorldDelegates::OnPostWorldCleanup.AddWeakLambda(InWorld, [InWorld](UWorld* WorldParam, bool bSessionEnded, bool bCleanupResources)
+			{
+				if (WorldParam == InWorld)
+				{
+					ResetCachedData();
+				}
+			});
+
+		return;
+	}
+
 	ResetCachedData();
 }
 
diff --git a/Engine/Plugins/Runtime/GameplayAbilities/Source/GameplayAbilities/Public/AbilitySystemGlobals.h b/Engine/Plugins/Runtime/GameplayAbilities/Source/GameplayAbilities/Public/AbilitySystemGlobals.h
--- a/Engine/Plugins/Runtime/GameplayAbilities/Source/GameplayAbilities/Public/AbilitySystemGlobals.h
+++ b/Engine/Plugins/Runtime/GameplayAbilities/Source/GameplayAbilities/Public/AbilitySystemGlobals.h
@@ -414,7 +414,7 @@ protected:
 	void OnPreBeginPIE(const bool bIsSimulatingInEditor);
 #endif
 
-	void ResetCachedData();
+	static void ResetCachedData();
 	void HandlePreLoadMap(const FWorldContext& WorldContext, const FString& MapName);
 
 #if WITH_EDITORONLY_DATA
