Added the ability for AbilitySystem.Effect commands to be sent to the server rather than run locally.  This helps apply those Effects on the server and thus allows you to test GE stacking.

diff --git a/Engine/Plugins/Runtime/GameplayAbilities/Source/GameplayAbilities/Private/AbilitySystemGlobals.cpp b/Engine/Plugins/Runtime/GameplayAbilities/Source/GameplayAbilities/Private/AbilitySystemGlobals.cpp
--- a/Engine/Plugins/Runtime/GameplayAbilities/Source/GameplayAbilities/Private/AbilitySystemGlobals.cpp
+++ b/Engine/Plugins/Runtime/GameplayAbilities/Source/GameplayAbilities/Private/AbilitySystemGlobals.cpp
@@ -964,9 +964,12 @@ FAutoConsoleCommand DebugAbilitySystemEffectListActiveCommand(TEXT("AbilitySyste
 		}
 	}), ECVF_Cheat);
 
-FAutoConsoleCommand DebugAbilitySystemEffectRemoveCommand(TEXT("AbilitySystem.Effect.Remove"), TEXT("<Handle/Name>. Remove a Gameplay Effect that is currently active on the Player"),
+FAutoConsoleCommand DebugAbilitySystemEffectRemoveCommand(TEXT("AbilitySystem.Effect.Remove"), TEXT("[-Server] <Handle/Name>. Remove a Gameplay Effect that is currently active on the Player"),
 	FConsoleCommandWithWorldArgsAndOutputDeviceDelegate::CreateLambda([](const TArray<FString>& Args, UWorld* World, FOutputDevice& OutputDevice)
 		{
+			TArray<FString> MutableArgs = Args;
+			const bool bExecuteOnServer = MutableArgs.RemoveSingle(TEXT("-server")) > 0;
+
 			UAbilitySystemGlobals& AbilitySystemGlobals = UAbilitySystemGlobals::Get();
 
 			APlayerController* PC = World->GetFirstPlayerController();
@@ -977,10 +980,15 @@ FAutoConsoleCommand DebugAbilitySystemEffectRemoveCommand(TEXT("AbilitySystem.Ef
 				return;
 			}
 
-			if (Args.Num() < 1)
+			if (MutableArgs.Num() < 1)
 			{
+				OutputDevice.Logf(TEXT("Missing Arguments: [-Server] <Handle/Name>"));
+				OutputDevice.Logf(TEXT(" -Server: Indicates this command should be sent to the server, rather than predicted locally (NOTE: Handles are not the same on the server)."));
+				OutputDevice.Logf(TEXT(" Handle/Name: Supply a substring of a class name to match.  Handles can also be used if executed locally (useful in the case of multiple instances of the same GE)."));
+
 				if (IConsoleObject* ListActiveCommand = IConsoleManager::Get().FindConsoleObject(TEXT("AbilitySystem.Effect.ListActive"), false))
 				{
+					OutputDevice.Logf(TEXT("\nPerforming AbilitySystem.Effect.ListActive:"));
 					ListActiveCommand->AsCommand()->Execute(Args, World, OutputDevice);
 				}
 
@@ -988,7 +996,7 @@ FAutoConsoleCommand DebugAbilitySystemEffectRemoveCommand(TEXT("AbilitySystem.Ef
 			}
 
 			uint32 SearchHandleValue = 0;
-			FString SearchString = Args[0];
+			FString SearchString = MutableArgs[0];
 			if (SearchString.IsNumeric())
 			{
 				int64 SearchHash64 = FCString::Atoi64(*SearchString);
@@ -998,6 +1006,24 @@ FAutoConsoleCommand DebugAbilitySystemEffectRemoveCommand(TEXT("AbilitySystem.Ef
 				}
 			}
 
+			if (!ASC->IsOwnerActorAuthoritative() && bExecuteOnServer)
+			{
+				if (SearchHandleValue > 0)
+				{
+					OutputDevice.Logf(TEXT("Error: Search by Handle Value is not permitted with -server (because ActiveGE Handles are not replicated to the same handle value)"));
+					return;
+				}
+				else
+				{
+					const FString AllArgs = FString::Join(MutableArgs, TEXT(" "));
+					const FString ServerCommand = FString::Printf(TEXT("AbilitySystem.Effect.Remove %s"), *AllArgs);
+					PC->ServerExec(ServerCommand);
+
+					OutputDevice.Logf(TEXT("Sent Command '%s' from Player '%s' to Server (as requested)."), *ServerCommand, *GetNameSafe(PC));
+					return;
+				}
+			}
+
 			bool bAnyMatch = false;
 			const FActiveGameplayEffectsContainer& ActiveGEContainer = ASC->GetActiveGameplayEffects();
 			for (const FActiveGameplayEffectHandle& ActiveGEHandle : ActiveGEContainer.GetAllActiveEffectHandles())
@@ -1023,13 +1049,26 @@ FAutoConsoleCommand DebugAbilitySystemEffectRemoveCommand(TEXT("AbilitySystem.Ef
 			}
 		}), ECVF_Cheat);
 
-FAutoConsoleCommand DebugAbilitySystemEffectApply(TEXT("AbilitySystem.Effect.Apply"), TEXT("<Class/AssetName> [Level]. Apply a Gameplay Effect on the Player.  Substring name matching works for Asset Tags, Asset Paths, or Class Names."),
+FAutoConsoleCommand DebugAbilitySystemEffectApply(TEXT("AbilitySystem.Effect.Apply"), TEXT("[-Server] <Class/AssetName> [Level]. Apply a Gameplay Effect on the Player.  Substring name matching works for Asset Tags, Asset Paths, or Class Names.  Use -Server to send to the server (default is apply locally)."),
 	FConsoleCommandWithWorldArgsAndOutputDeviceDelegate::CreateLambda([](const TArray<FString>& Args, UWorld* World, FOutputDevice& OutputDevice)
 	{
+		TArray<FString> MutableArgs = Args;
+		const bool bExecuteOnServer = MutableArgs.RemoveSingle(TEXT("-server")) > 0;
+
+		if (MutableArgs.Num() < 1)
+		{
+			OutputDevice.Logf(TEXT("Missing Arguments: [-Server] <MatchString> [Level]"));
+			OutputDevice.Logf(TEXT(" -Server: Indicates this command should be sent to the server, rather than predicted locally (search for assets is still verified locally)."));
+			OutputDevice.Logf(TEXT(" MatchString: Supply a substring of a class or asset name to match. Only loaded GameplayEffects are searched."));
+			OutputDevice.Logf(TEXT(" Level: Optionally supply a Level such as 4.  Can be omitted (in which case the GE is level-less)."));
+			return;
+		}
+
 		UAbilitySystemGlobals& AbilitySystemGlobals = UAbilitySystemGlobals::Get();
 
 		APlayerController* PC = World->GetFirstPlayerController();
-		UAbilitySystemComponent* ASC = PC ? AbilitySystemGlobals.GetAbilitySystemComponentFromActor(PC->GetPawn()) : nullptr;
+		APawn* Pawn = PC->GetPawn();
+		UAbilitySystemComponent* ASC = PC ? AbilitySystemGlobals.GetAbilitySystemComponentFromActor(Pawn) : nullptr;
 		if (!ASC)
 		{
 			OutputDevice.Logf(TEXT("Could not find Player (%s) with AbilitySystemComponent in World (%s)"), *GetNameSafe(PC), *GetNameSafe(World));
@@ -1037,7 +1076,7 @@ FAutoConsoleCommand DebugAbilitySystemEffectApply(TEXT("AbilitySystem.Effect.App
 		}
 
 		// We couldn't find anything the user was searching for, so early out
-		const FString SearchString = Args[0];
+		const FString SearchString = MutableArgs[0];
 		TSubclassOf<UGameplayEffect> GameplayEffectClass = UE::AbilitySystemGlobals::FuzzyFindClass<UGameplayEffect>(SearchString);
 		if (!GameplayEffectClass)
 		{
@@ -1045,9 +1084,22 @@ FAutoConsoleCommand DebugAbilitySystemEffectApply(TEXT("AbilitySystem.Effect.App
 			return;
 		}
 
+		// If we're not the authority, we need to send the command to the server because we can't grant locally.
+		if (!ASC->IsOwnerActorAuthoritative() && bExecuteOnServer)
+		{
+			const FString AllArgs = FString::Join(MutableArgs, TEXT(" "));
+			const FString ServerCommand = FString::Printf(TEXT("AbilitySystem.Effect.Apply %s"), *AllArgs);
+			PC->ServerExec(ServerCommand);
+
+			OutputDevice.Logf(TEXT("Sent Command '%s' from Player '%s' to Server (as requested)."), *ServerCommand, *GetNameSafe(PC));
+			return;
+		}
+
 		// Create a GameplayEffectSpec that executes the passed-in parameters
 		FGameplayEffectContextHandle GEContextHandle = ASC->MakeEffectContext();
-		const float Level = FMath::Max(FGameplayEffectConstants::INVALID_LEVEL, (Args.Num() > 1) ? FPlatformString::Atof(*Args[1]) : FGameplayEffectConstants::INVALID_LEVEL);
+		GEContextHandle.AddInstigator(Pawn, Pawn);
+		GEContextHandle.AddOrigin(Pawn->GetActorLocation());
+		const float Level = FMath::Max(FGameplayEffectConstants::INVALID_LEVEL, (Args.Num() > 1) ? FPlatformString::Atof(*MutableArgs[1]) : FGameplayEffectConstants::INVALID_LEVEL);
 		FGameplayEffectSpec GESpec{ GameplayEffectClass.GetDefaultObject(), GEContextHandle, Level };
 
 		// We need to create a new valid prediction key (as if we just started activating an ability) so the GE will fire even if we are not the authority.
