The GAS debug widget enabled via AbilitySystem.DebugAttribute <Attribute> causes a client-side crash when the server has applied runtime constructed GEs, since their Spec.Def are non-replicated objects. Addressed by doing a Spec.Def null-check, as is done in other places.

diff --git a/Engine/Plugins/Runtime/GameplayAbilities/Source/GameplayAbilities/Private/GameplayEffect.cpp b/Engine/Plugins/Runtime/GameplayAbilities/Source/GameplayAbilities/Private/GameplayEffect.cpp
--- a/Engine/Plugins/Runtime/GameplayAbilities/Source/GameplayAbilities/Private/GameplayEffect.cpp
+++ b/Engine/Plugins/Runtime/GameplayAbilities/Source/GameplayAbilities/Private/GameplayEffect.cpp
@@ -4949,7 +4949,7 @@ void FActiveGameplayEffectsContainer::GetActiveGameplayEffectDataByAttribute(TMu
 	// Add all of the active gameplay effects
 	for (const FActiveGameplayEffect& Effect : this)
 	{
-		if (Effect.Spec.Modifiers.Num() == Effect.Spec.Def->Modifiers.Num())
+		if (Effect.Spec.Def && Effect.Spec.Modifiers.Num() == Effect.Spec.Def->Modifiers.Num())
 		{
 			for (int32 Idx = 0; Idx < Effect.Spec.Modifiers.Num(); ++Idx)
 			{
