Moving the GameplayEffects logs from LogAbilitySystem to LogGameplayEffect.  Enhancing the locations of the logs to provide better context to what is happening.

diff --git a/Engine/Plugins/Runtime/GameplayAbilities/Source/GameplayAbilities/Private/GameplayEffect.cpp b/Engine/Plugins/Runtime/GameplayAbilities/Source/GameplayAbilities/Private/GameplayEffect.cpp
--- a/Engine/Plugins/Runtime/GameplayAbilities/Source/GameplayAbilities/Private/GameplayEffect.cpp
+++ b/Engine/Plugins/Runtime/GameplayAbilities/Source/GameplayAbilities/Private/GameplayEffect.cpp
@@ -715,7 +715,7 @@ bool UGameplayEffect::CanApply(const FActiveGameplayEffectsContainer& ActiveGECo
 	{
 		if (GEComponent && !GEComponent->CanGameplayEffectApply(ActiveGEContainer, GESpec))
 		{
-			UE_VLOG(ActiveGEContainer.Owner, LogAbilitySystem, Verbose, TEXT("%s could not apply. Blocked by %s"), *GetNameSafe(GESpec.Def), *GetNameSafe(GEComponent));
+			UE_VLOG_UELOG(ActiveGEContainer.Owner, LogGameplayEffects, Verbose, TEXT("%s could not apply. Blocked by %s"), *GetNameSafe(GESpec.Def), *GetNameSafe(GEComponent));
 			return false;
 		}
 	}
@@ -725,8 +725,6 @@ bool UGameplayEffect::CanApply(const FActiveGameplayEffectsContainer& ActiveGECo
 
 bool UGameplayEffect::OnAddedToActiveContainer(FActiveGameplayEffectsContainer& ActiveGEContainer, FActiveGameplayEffect& ActiveGE) const
 {
-	UE_VLOG(ActiveGEContainer.Owner, LogAbilitySystem, Verbose, TEXT("ActiveGameplayEffect %s added"), *ActiveGE.GetDebugString());
-
 	bool bShouldBeActive = true;
 	for (const UGameplayEffectComponent* GEComponent : GEComponents)
 	{
@@ -736,13 +734,13 @@ bool UGameplayEffect::OnAddedToActiveContainer(FActiveGameplayEffectsContainer&
 		}
 	}
 
+	UE_VLOG_UELOG(ActiveGEContainer.Owner->GetOwnerActor(), LogGameplayEffects, VeryVerbose, TEXT("%s: %s. ShouldBeActive: %d"), ANSI_TO_TCHAR(__func__), *ActiveGE.GetDebugString(), bShouldBeActive);
+
 	return bShouldBeActive;
 }
 
 void UGameplayEffect::OnExecuted(FActiveGameplayEffectsContainer& ActiveGEContainer, FGameplayEffectSpec& GESpec, FPredictionKey& PredictionKey) const
 {
-	UE_VLOG(ActiveGEContainer.Owner, LogAbilitySystem, Verbose, TEXT("GameplayEffect %s executed"), *GetNameSafe(GESpec.Def));
-
 	for (const UGameplayEffectComponent* GEComponent : GEComponents)
 	{
 		if (GEComponent)
@@ -750,12 +748,12 @@ void UGameplayEffect::OnExecuted(FActiveGameplayEffectsContainer& ActiveGEContai
 			GEComponent->OnGameplayEffectExecuted(ActiveGEContainer, GESpec, PredictionKey);
 		}
 	}
+
+	UE_VLOG_UELOG(ActiveGEContainer.Owner->GetOwnerActor(), LogGameplayEffects, VeryVerbose, TEXT("%s: %s"), ANSI_TO_TCHAR(__func__), *GetNameSafe(GESpec.Def));
 }
 
 void UGameplayEffect::OnApplied(FActiveGameplayEffectsContainer& ActiveGEContainer, FGameplayEffectSpec& GESpec, FPredictionKey& PredictionKey) const
 {
-	UE_VLOG(ActiveGEContainer.Owner, LogAbilitySystem, Verbose, TEXT("GameplayEffect %s applied"), *GetNameSafe(GESpec.Def));
-
 	for (const UGameplayEffectComponent* GEComponent : GEComponents)
 	{
 		if (GEComponent)
@@ -763,6 +761,8 @@ void UGameplayEffect::OnApplied(FActiveGameplayEffectsContainer& ActiveGEContain
 			GEComponent->OnGameplayEffectApplied(ActiveGEContainer, GESpec, PredictionKey);
 		}
 	}
+
+	UE_VLOG_UELOG(ActiveGEContainer.Owner->GetOwnerActor(), LogGameplayEffects, VeryVerbose, TEXT("%s: %s"), ANSI_TO_TCHAR(__func__), *GetNameSafe(GESpec.Def));
 }
 
 int32 UGameplayEffect::GetStackLimitCount() const
@@ -987,7 +987,7 @@ bool FGameplayEffectModifierMagnitude::AttemptCalculateMagnitude(const FGameplay
 			break;
 
 			default:
-				ABILITY_LOG(Error, TEXT("Unknown MagnitudeCalculationType %d in AttemptCalculateMagnitude"), (int32)MagnitudeCalculationType);
+				UE_LOG(LogGameplayEffects, Error, TEXT("Unknown MagnitudeCalculationType %d in AttemptCalculateMagnitude"), (int32)MagnitudeCalculationType);
 				OutCalculatedMagnitude = 0.f;
 				break;
 		}
@@ -1732,7 +1732,7 @@ void FGameplayEffectSpec::CalculateModifierMagnitudes()
 		if (ModDef.ModifierMagnitude.AttemptCalculateMagnitude(*this, ModSpec.EvaluatedMagnitude) == false)
 		{
 			ModSpec.EvaluatedMagnitude = 0.f;
-			ABILITY_LOG(Warning, TEXT("Modifier on spec: %s was asked to CalculateMagnitude and failed, falling back to 0."), *ToSimpleString());
+			UE_LOG(LogGameplayEffects, Warning, TEXT("Modifier on spec: %s was asked to CalculateMagnitude and failed, falling back to 0."), *ToSimpleString());
 		}
 	}
 }
@@ -1877,7 +1877,7 @@ float FGameplayEffectSpec::GetSetByCallerMagnitude(FName DataName, bool WarnIfNo
 	}
 	else if (WarnIfNotFound)
 	{
-		ABILITY_LOG(Error, TEXT("FGameplayEffectSpec::GetMagnitude called for Data %s on Def %s when magnitude had not yet been set by caller."), *DataName.ToString(), *Def->GetName());
+		UE_LOG(LogGameplayEffects, Error, TEXT("FGameplayEffectSpec::GetMagnitude called for Data %s on Def %s when magnitude had not yet been set by caller."), *DataName.ToString(), *Def->GetName());
 	}
 
 	return Magnitude;
@@ -1899,7 +1899,7 @@ float FGameplayEffectSpec::GetSetByCallerMagnitude(FGameplayTag DataTag, bool Wa
 	}
 	else if (WarnIfNotFound)
 	{
-		ABILITY_LOG(Error, TEXT("FGameplayEffectSpec::GetMagnitude called for Data %s on Def %s when magnitude had not yet been set by caller."), *DataTag.ToString(), *Def->GetName());
+		UE_LOG(LogGameplayEffects, Error, TEXT("FGameplayEffectSpec::GetMagnitude called for Data %s on Def %s when magnitude had not yet been set by caller."), *DataTag.ToString(), *Def->GetName());
 	}
 
 	return Magnitude;
@@ -2349,12 +2349,10 @@ void FActiveGameplayEffect::PreReplicatedRemove(const struct FActiveGameplayEffe
 {
 	if (Spec.Def == nullptr)
 	{
-		ABILITY_LOG(Error, TEXT("Received PreReplicatedRemove with no UGameplayEffect def."));
+		UE_LOG(LogGameplayEffects, Error, TEXT("Received PreReplicatedRemove with no UGameplayEffect def."));
 		return;
 	}
 
-	ABILITY_LOG(Verbose, TEXT("PreReplicatedRemove: %s %s Marked as Pending Remove: %s"), *Handle.ToString(), *Spec.Def->GetName(), IsPendingRemove ? TEXT("TRUE") : TEXT("FALSE"));
-
 	FGameplayEffectRemovalInfo GameplayEffectRemovalInfo;
 	GameplayEffectRemovalInfo.ActiveEffect = this;
 	GameplayEffectRemovalInfo.StackCount = ClientCachedStackCount;
@@ -2370,6 +2368,8 @@ void FActiveGameplayEffect::PreReplicatedRemove(const struct FActiveGameplayEffe
 	}
 	GameplayEffectRemovalInfo.EffectContext = Spec.GetEffectContext();
 
+	UE_VLOG_UELOG(InArray.Owner->GetOwnerActor(), LogGameplayEffects, Verbose, TEXT("%s (Non-Auth): %s. Premature: %d Inhibited: %d. Pending( Remove: %d OnActive: %d WhileActive: %d )"), ANSI_TO_TCHAR(__func__), *GetDebugString(), GameplayEffectRemovalInfo.bPrematureRemoval, bIsInhibited, IsPendingRemove, bPendingRepOnActiveGC, bPendingRepWhileActiveGC);
+
 	const_cast<FActiveGameplayEffectsContainer&>(InArray).InternalOnActiveGameplayEffectRemoved(*this, !bIsInhibited, GameplayEffectRemovalInfo);	// Const cast is ok. It is there to prevent mutation of the GameplayEffects array, which this wont do.
 }
 
@@ -2377,14 +2377,14 @@ void FActiveGameplayEffect::PostReplicatedAdd(const struct FActiveGameplayEffect
 {
 	if (Spec.Def == nullptr)
 	{
-		ABILITY_LOG(Error, TEXT("FActiveGameplayEffect::PostReplicatedAdd Received ReplicatedGameplayEffect with no UGameplayEffect def. (%s)"), *Spec.GetEffectContext().ToString());
+		UE_LOG(LogGameplayEffects, Error, TEXT("FActiveGameplayEffect::PostReplicatedAdd Received ReplicatedGameplayEffect with no UGameplayEffect def. (%s)"), *Spec.GetEffectContext().ToString());
 		return;
 	}
 
 	if (Spec.Modifiers.Num() != Spec.Def->Modifiers.Num())
 	{
 		// This can happen with older replays, where the replicated Spec.Modifiers size changed in the newer Spec.Def
-		ABILITY_LOG(Error, TEXT("FActiveGameplayEffect::PostReplicatedAdd: Spec.Modifiers.Num() != Spec.Def->Modifiers.Num(). Spec: %s"), *Spec.ToSimpleString());
+		UE_LOG(LogGameplayEffects, Error, TEXT("FActiveGameplayEffect::PostReplicatedAdd: Spec.Modifiers.Num() != Spec.Def->Modifiers.Num(). Spec: %s"), *Spec.ToSimpleString());
 		Spec.Modifiers.Empty();
 		return;
 	}
@@ -2437,6 +2437,8 @@ void FActiveGameplayEffect::PostReplicatedAdd(const struct FActiveGameplayEffect
 	// Handles are not replicated, so create a new one.
 	Handle = FActiveGameplayEffectHandle::GenerateNewHandle(InArray.Owner);
 
+	UE_VLOG_UELOG(InArray.Owner->GetOwnerActor(), LogGameplayEffects, Verbose, TEXT("%s (Non-Auth): %s. Pending( OnActive: %d WhileActive: %d )"), ANSI_TO_TCHAR(__func__), *GetDebugString(), bPendingRepOnActiveGC, bPendingRepWhileActiveGC);
+
 	// Do stuff for adding GEs (add mods, tags, *invoke callbacks*
 	const_cast<FActiveGameplayEffectsContainer&>(InArray).InternalOnActiveGameplayEffectAdded(*this);	// Const cast is ok. It is there to prevent mutation of the GameplayEffects array, which this wont do.
 	
@@ -2446,7 +2448,7 @@ void FActiveGameplayEffect::PostReplicatedChange(const struct FActiveGameplayEff
 {
 	if (Spec.Def == nullptr)
 	{
-		ABILITY_LOG(Error, TEXT("FActiveGameplayEffect::PostReplicatedChange Received ReplicatedGameplayEffect with no UGameplayEffect def. (%s)"), *Spec.GetEffectContext().ToString());
+		UE_LOG(LogGameplayEffects, Error, TEXT("FActiveGameplayEffect::PostReplicatedChange Received ReplicatedGameplayEffect with no UGameplayEffect def. (%s)"), *Spec.GetEffectContext().ToString());
 		return;
 	}
 
@@ -2480,11 +2482,13 @@ void FActiveGameplayEffect::PostReplicatedChange(const struct FActiveGameplayEff
 		// Const cast is ok. It is there to prevent mutation of the GameplayEffects array, which this wont do.
 		const_cast<FActiveGameplayEffectsContainer&>(InArray).UpdateAllAggregatorModMagnitudes(*this);
 	}
+
+	UE_VLOG_UELOG(InArray.Owner->GetOwnerActor(), LogGameplayEffects, Verbose, TEXT("%s (Non-Auth): %s. Pending( OnActive: %d. WhileActive: %d )"), ANSI_TO_TCHAR(__func__), *GetDebugString(), bPendingRepOnActiveGC, bPendingRepWhileActiveGC);
 }
 
 FString FActiveGameplayEffect::GetDebugString()
 {
-	return FString::Printf(TEXT("(Def: %s. PredictionKey: %s)"), *GetNameSafe(Spec.Def), *PredictionKey.ToString());
+	return FString::Printf(TEXT("Def: %s. Handle: %s. PredictionKey: %s"), *GetNameSafe(Spec.Def), *Handle.ToString(), *PredictionKey.ToString());
 }
 
 void FActiveGameplayEffect::RecomputeStartWorldTime(const FActiveGameplayEffectsContainer& InArray)
@@ -2654,7 +2658,7 @@ void FActiveGameplayEffectsContainer::PredictivelyExecuteEffectSpec(FGameplayEff
 		{
 			// TODO: check replication policy. Right now we will replicate every execute via a multicast RPC
 
-			ABILITY_LOG(Log, TEXT("Invoking Execute GameplayCue for %s"), *SpecToUse.ToSimpleString());
+			UE_LOG(LogGameplayEffects, Log, TEXT("Invoking Execute GameplayCue for %s"), *SpecToUse.ToSimpleString());
 
 			UAbilitySystemGlobals::Get().GetGameplayCueManager()->InvokeGameplayCueExecuted_FromSpec(Owner, SpecToUse, PredictionKey);
 		}
@@ -2790,7 +2794,7 @@ void FActiveGameplayEffectsContainer::ExecuteActiveEffectsFrom(FGameplayEffectSp
 	{
 		// TODO: check replication policy. Right now we will replicate every execute via a multicast RPC
 
-		ABILITY_LOG(Log, TEXT("Invoking Execute GameplayCue for %s"), *SpecToUse.ToSimpleString());
+		UE_LOG(LogGameplayEffects, Log, TEXT("Invoking Execute GameplayCue for %s"), *SpecToUse.ToSimpleString());
 
 		UAbilitySystemGlobals::Get().GetGameplayCueManager()->InvokeGameplayCueExecuted_FromSpec(Owner, SpecToUse, PredictionKey);
 	}
@@ -2854,7 +2858,7 @@ FAggregatorRef& FActiveGameplayEffectsContainer::FindOrCreateAttributeAggregator
 
 	// Create a new aggregator for this attribute.
 	float CurrentBaseValueOfProperty = Owner->GetNumericAttributeBase(Attribute);
-	ABILITY_LOG(Log, TEXT("Creating new entry in AttributeAggregatorMap for %s. CurrentValue: %.2f"), *Attribute.GetName(), CurrentBaseValueOfProperty);
+	UE_LOG(LogGameplayEffects, Log, TEXT("Creating new entry in AttributeAggregatorMap for %s. CurrentValue: %.2f"), *Attribute.GetName(), CurrentBaseValueOfProperty);
 
 	FAggregator* NewAttributeAggregator = new FAggregator(CurrentBaseValueOfProperty);
 	
@@ -2911,7 +2915,7 @@ void FActiveGameplayEffectsContainer::OnAttributeAggregatorDirty(FAggregator* Ag
 				// Legacy float attribute case requires the base value to be deduced from the final value, as it is not replicated
 				const float FinalValue = Owner->GetNumericAttribute(Attribute);
 				const float BaseValue = Aggregator->ReverseEvaluate(FinalValue, EvaluationParameters);
-				ABILITY_LOG(Log, TEXT("Reverse Evaluated %s. FinalValue: %.2f  BaseValue: %2.f.  Setting BaseValue.  (Role: %s)"), *Attribute.GetName(), FinalValue, BaseValue, *UEnum::GetValueAsString(Owner->GetOwnerRole()));
+				UE_LOG(LogGameplayEffects, Log, TEXT("Reverse Evaluated %s. FinalValue: %.2f  BaseValue: %2.f.  Setting BaseValue.  (Role: %s)"), *Attribute.GetName(), FinalValue, BaseValue, *UEnum::GetValueAsString(Owner->GetOwnerRole()));
 
 				Aggregator->SetBaseValue(BaseValue, false);
 			}
@@ -2926,7 +2930,7 @@ void FActiveGameplayEffectsContainer::OnAttributeAggregatorDirty(FAggregator* Ag
 	if (EvaluationParameters.IncludePredictiveMods)
 	{
 		const float OldValue = Owner->GetNumericAttribute(Attribute);
-		ABILITY_LOG(Log, TEXT("[%s] Aggregator Evaluated %s. OldValue: %.2f  NewValue: %.2f"), *UEnum::GetValueAsString(Owner->GetOwnerRole()), *Attribute.GetName(), OldValue, NewValue);
+		UE_LOG(LogGameplayEffects, Log, TEXT("[%s] Aggregator Evaluated %s. OldValue: %.2f  NewValue: %.2f"), *UEnum::GetValueAsString(Owner->GetOwnerRole()), *Attribute.GetName(), OldValue, NewValue);
 	}
 
 	InternalUpdateNumericalAttribute(Attribute, NewValue, nullptr, bFromRecursiveCall);
@@ -2993,6 +2997,8 @@ void FActiveGameplayEffectsContainer::OnMagnitudeDependencyChange(FActiveGamepla
 
 void FActiveGameplayEffectsContainer::OnStackCountChange(FActiveGameplayEffect& ActiveEffect, int32 OldStackCount, int32 NewStackCount)
 {
+	UE_VLOG_UELOG(Owner->GetOwnerActor(), LogGameplayEffects, Verbose, TEXT("OnStackCountChange: %s. OldStackCount: %d. NewStackCount: %d"), *ActiveEffect.GetDebugString(), OldStackCount, NewStackCount);
+
 	MarkItemDirty(ActiveEffect);
 	if (OldStackCount != NewStackCount)
 	{
@@ -3035,7 +3041,7 @@ void FActiveGameplayEffectsContainer::UpdateAllAggregatorModMagnitudes(FActiveGa
 
 	if (Spec.Def == nullptr)
 	{
-		ABILITY_LOG(Error, TEXT("UpdateAllAggregatorModMagnitudes called with no UGameplayEffect def."));
+		UE_LOG(LogGameplayEffects, Error, TEXT("UpdateAllAggregatorModMagnitudes called with no UGameplayEffect def."));
 		return;
 	}
 
@@ -3143,7 +3149,7 @@ void FActiveGameplayEffectsContainer::SetBaseAttributeValueFromReplication(const
 
 			// Now set the new value and go through all of the aggregations...
 			Aggregator->SetBaseValue(ServerBaseValue, bDoNotExecuteCallbacksValue);
-			ABILITY_LOG(Log, TEXT("SetBaseAttributeValueFromReplication [%s]: %s rewound to state NewBaseValue: %.2f  OldCurrentValue: %.2f"), OwnerIsNetAuthority ? TEXT("Authority") : TEXT("Client"), *Attribute.AttributeName, ServerBaseValue, OldEvaluatedValue);
+			UE_LOG(LogGameplayEffects, Log, TEXT("SetBaseAttributeValueFromReplication [%s]: %s rewound to state NewBaseValue: %.2f  OldCurrentValue: %.2f"), OwnerIsNetAuthority ? TEXT("Authority") : TEXT("Client"), *Attribute.AttributeName, ServerBaseValue, OldEvaluatedValue);
 		}
 
 		FScopedAggregatorOnDirtyBatch::GlobalFromNetworkUpdate = true;
@@ -3199,7 +3205,7 @@ void FActiveGameplayEffectsContainer::GetGameplayEffectStartTimeAndDuration(FAct
 		}
 	}
 
-	ABILITY_LOG(Warning, TEXT("GetGameplayEffectStartTimeAndDuration called with invalid Handle: %s"), *Handle.ToString());
+	UE_LOG(LogGameplayEffects, Warning, TEXT("GetGameplayEffectStartTimeAndDuration called with invalid Handle: %s"), *Handle.ToString());
 }
 
 void FActiveGameplayEffectsContainer::RecomputeStartWorldTimes(const float WorldTime, const float ServerWorldTime)
@@ -3229,7 +3235,7 @@ float FActiveGameplayEffectsContainer::GetGameplayEffectMagnitude(FActiveGamepla
 		}
 	}
 
-	ABILITY_LOG(Warning, TEXT("GetGameplayEffectMagnitude called with invalid Handle: %s"), *Handle.ToString());
+	UE_LOG(LogGameplayEffects, Warning, TEXT("GetGameplayEffectMagnitude called with invalid Handle: %s"), *Handle.ToString());
 	return -1.f;
 }
 
@@ -3322,7 +3328,7 @@ void FActiveGameplayEffectsContainer::CaptureAttributeForGameplayEffect(OUT FGam
 void FActiveGameplayEffectsContainer::InternalUpdateNumericalAttribute(FGameplayAttribute Attribute, float NewValue, const FGameplayEffectModCallbackData* ModData, bool bFromRecursiveCall)
 {
 	const float OldValue = Owner->GetNumericAttribute(Attribute);
-	ABILITY_LOG(Log, TEXT("[%s] InternalUpdateNumericalAttribute %s OldValue = %.2f  NewValue = %.2f."), *UEnum::GetValueAsString(Owner->GetOwnerRole()), *Attribute.GetName(), OldValue, NewValue);
+	UE_VLOG_UELOG(Owner->GetOwnerActor(), LogGameplayEffects, Log, TEXT("[%s] InternalUpdateNumericalAttribute %s OldValue = %.2f  NewValue = %.2f."), *UEnum::GetValueAsString(Owner->GetOwnerRole()), *Attribute.GetName(), OldValue, NewValue);
 	Owner->SetNumericAttribute_Internal(Attribute, NewValue);
 	
 	if (!bFromRecursiveCall)
@@ -3330,7 +3336,7 @@ void FActiveGameplayEffectsContainer::InternalUpdateNumericalAttribute(FGameplay
 		// We should only have one: either cached CurrentModcallbackData, or explicit callback data passed directly in.
 		if (ModData && CurrentModcallbackData)
 		{
-			ABILITY_LOG(Warning, TEXT("Had passed in ModData and cached CurrentModcallbackData in FActiveGameplayEffectsContainer::InternalUpdateNumericalAttribute. For attribute %s on %s."), *Attribute.GetName(), *Owner->GetFullName() );
+			UE_LOG(LogGameplayEffects, Warning, TEXT("Had passed in ModData and cached CurrentModcallbackData in FActiveGameplayEffectsContainer::InternalUpdateNumericalAttribute. For attribute %s on %s."), *Attribute.GetName(), *Owner->GetFullName() );
 		}
 		
 		const FGameplayEffectModCallbackData* DataToShare = ModData ? ModData : CurrentModcallbackData;
@@ -3516,7 +3522,7 @@ bool FActiveGameplayEffectsContainer::InternalExecuteMod(FGameplayEffectSpec& Sp
 	else
 	{
 		// Our owner doesn't have this attribute, so we can't do anything
-		ABILITY_LOG(Log, TEXT("%s does not have attribute %s. Skipping modifier"), *Owner->GetPathName(), *ModEvalData.Attribute.GetName());
+		UE_LOG(LogGameplayEffects, Log, TEXT("%s does not have attribute %s. Skipping modifier"), *Owner->GetPathName(), *ModEvalData.Attribute.GetName());
 	}
 
 	return bExecuted;
@@ -3533,7 +3539,7 @@ void FActiveGameplayEffectsContainer::ApplyModToAttribute(const FGameplayAttribu
 	if (CurrentModcallbackData)
 	{
 		// We expect this to be cleared for us in InternalUpdateNumericalAttribute
-		ABILITY_LOG(Warning, TEXT("FActiveGameplayEffectsContainer::ApplyModToAttribute CurrentModcallbackData was not consumed For attribute %s on %s."), *Attribute.GetName(), *Owner->GetFullName());
+		UE_LOG(LogGameplayEffects, Warning, TEXT("FActiveGameplayEffectsContainer::ApplyModToAttribute CurrentModcallbackData was not consumed For attribute %s on %s."), *Attribute.GetName(), *Owner->GetFullName());
 		CurrentModcallbackData = nullptr;
 	}
 }
@@ -3762,21 +3768,19 @@ PRAGMA_ENABLE_DEPRECATION_WARNINGS
 		// We cannot mod ourselves into an instant or infinite duration effect
 		if (FinalDuration <= 0.f)
 		{
-			ABILITY_LOG(Error, TEXT("GameplayEffect %s Duration was modified to %.2f. Clamping to 0.1s duration."), *AppliedEffectSpec.Def->GetName(), FinalDuration);
+			UE_LOG(LogGameplayEffects, Error, TEXT("ActiveGE %s Duration was modified to %.2f. Clamping to 0.1s duration."), *AppliedActiveGE->GetDebugString(), FinalDuration);
 			FinalDuration = 0.1f;
 		}
 
 		AppliedEffectSpec.SetDuration(FinalDuration, true);
 
-		// ABILITY_LOG(Warning, TEXT("SetDuration for %s. Base: %.2f, Final: %.2f"), *NewEffect.Spec.Def->GetName(), DurationBaseValue, FinalDuration);
-
 		// Register duration callbacks with the timer manager
 		if (Owner && bSetDuration)
 		{
 			FTimerManager& TimerManager = Owner->GetWorld()->GetTimerManager();
 			FTimerDelegate Delegate = FTimerDelegate::CreateUObject(Owner, &UAbilitySystemComponent::CheckDurationExpired, AppliedActiveGE->Handle);
 			TimerManager.SetTimer(AppliedActiveGE->DurationHandle, Delegate, FinalDuration, false);
-			if (!ensureMsgf(AppliedActiveGE->DurationHandle.IsValid(), TEXT("Invalid Duration Handle after attempting to set duration for GE %s @ %.2f"), 
+			if (!ensureMsgf(AppliedActiveGE->DurationHandle.IsValid(), TEXT("Invalid Duration Handle after attempting to set duration for GE (%s) @ %.2f"), 
 				*AppliedActiveGE->GetDebugString(), FinalDuration))
 			{
 				// Force this off next frame
@@ -3803,8 +3807,6 @@ PRAGMA_ENABLE_DEPRECATION_WARNINGS
 	if (InPredictionKey.IsLocalClientKey() == false || IsNetAuthority())	// Clients predicting a GameplayEffect must not call MarkItemDirty
 	{
 		MarkItemDirty(*AppliedActiveGE);
-
-		ABILITY_LOG(Verbose, TEXT("Added GE: %s. ReplicationID: %d. Key: %d. PredictionLey: %d"), *AppliedActiveGE->Spec.Def->GetName(), AppliedActiveGE->ReplicationID, AppliedActiveGE->ReplicationKey, InPredictionKey.Current);
 	}
 	else
 	{
@@ -3821,10 +3823,10 @@ PRAGMA_ENABLE_DEPRECATION_WARNINGS
 	if (ExistingStackableGE)
 	{
 		OnStackCountChange(*ExistingStackableGE, StartingStackCount, NewStackCount);
-		
 	}
 	else
 	{
+		UE_VLOG_UELOG(OwnerActor, LogGameplayEffects, Verbose, TEXT("Adding GE: %s. Auth: %d. ReplicationID: %d. ReplicationKey: %d."), *AppliedActiveGE->GetDebugString(), IsNetAuthority(), AppliedActiveGE->ReplicationID, AppliedActiveGE->ReplicationKey);
 		InternalOnActiveGameplayEffectAdded(*AppliedActiveGE);
 	}
 
@@ -3840,15 +3842,13 @@ void FActiveGameplayEffectsContainer::InternalOnActiveGameplayEffectAdded(FActiv
 
 	if (EffectDef == nullptr)
 	{
-		ABILITY_LOG(Error, TEXT("FActiveGameplayEffectsContainer serialized new GameplayEffect with NULL Def!"));
+		UE_LOG(LogGameplayEffects, Error, TEXT("FActiveGameplayEffectsContainer serialized new GameplayEffect with NULL Def!"));
 		return;
 	}
 
 	SCOPE_CYCLE_UOBJECT(EffectDef, EffectDef);
 
 	GAMEPLAYEFFECT_SCOPE_LOCK();
-	AActor* OwnerActor = Owner->GetOwnerActor();
-	UE_VLOG(OwnerActor ? OwnerActor : Owner->GetOuter(), LogGameplayEffects, Log, TEXT("Added: %s"), *GetNameSafe(EffectDef->GetClass()));
 
 	// Add any external dependencies that might dirty the effect, if necessary
 	AddCustomMagnitudeExternalDependencies(Effect);
@@ -3864,7 +3864,7 @@ void FActiveGameplayEffectsContainer::AddActiveGameplayEffectGrantedTagsAndModif
 {
 	if (Effect.Spec.Def == nullptr)
 	{
-		ABILITY_LOG(Error, TEXT("AddActiveGameplayEffectGrantedTagsAndModifiers called with null Def!"));
+		UE_LOG(LogGameplayEffects, Error, TEXT("AddActiveGameplayEffectGrantedTagsAndModifiers called with null Def!"));
 		return;
 	}
 
@@ -3993,24 +3993,12 @@ bool FActiveGameplayEffectsContainer::RemoveActiveGameplayEffect(FActiveGameplay
 		FActiveGameplayEffect& Effect = *GetActiveGameplayEffect(ActiveGEIdx);
 		if (Effect.Handle == Handle && Effect.IsPendingRemove == false)
 		{
-			AActor* OwnerActor = Owner->GetOwnerActor();
-			UE_VLOG(OwnerActor, LogGameplayEffects, Log, TEXT("Removed: %s"), *GetNameSafe(Effect.Spec.Def->GetClass()));
-			if (UE_LOG_ACTIVE(VLogAbilitySystem, Log))
-			{
-				ABILITY_VLOG(OwnerActor, Log, TEXT("Removed %s"), *Effect.Spec.Def->GetFName().ToString());
-				for (const FGameplayModifierInfo& Modifier : Effect.Spec.Def->Modifiers)
-				{
-					float Magnitude = 0.f;
-					Modifier.ModifierMagnitude.AttemptCalculateMagnitude(Effect.Spec, Magnitude);
-					ABILITY_VLOG(OwnerActor, Log, TEXT("         %s: %s %f"), *Modifier.Attribute.GetName(), *EGameplayModOpToString(Modifier.ModifierOp), Magnitude);
-				}
-			}
-
 			InternalRemoveActiveGameplayEffect(ActiveGEIdx, StacksToRemove, true);
 			return true;
 		}
 	}
-	ABILITY_LOG(Log, TEXT("RemoveActiveGameplayEffect called with invalid Handle: %s"), *Handle.ToString());
+
+	UE_LOG(LogGameplayEffects, Log, TEXT("RemoveActiveGameplayEffect called with invalid Handle: %s"), *Handle.ToString());
 	return false;
 }
 
@@ -4021,17 +4009,16 @@ void FActiveGameplayEffectsContainer::InternalExecutePeriodicGameplayEffect(FAct
 	{
 		FScopeCurrentGameplayEffectBeingApplied ScopedGEApplication(&ActiveEffect.Spec, Owner);
 
-		if (UE_LOG_ACTIVE(VLogAbilitySystem, Log))
-		{
-			AActor* OwnerActor = Owner->GetOwnerActor();
-			ABILITY_VLOG(OwnerActor, Log, TEXT("Executed Periodic Effect %s"), *ActiveEffect.Spec.Def->GetFName().ToString());
+		UE_IFVLOG(
+			AActor * OwnerActor = Owner->GetOwnerActor();
+			UE_VLOG(OwnerActor, LogGameplayEffects, Log, TEXT("Executed Periodic Effect %s"), *ActiveEffect.Spec.Def->GetFName().ToString());
 			for (const FGameplayModifierInfo& Modifier : ActiveEffect.Spec.Def->Modifiers)
 			{
 				float Magnitude = 0.f;
 				Modifier.ModifierMagnitude.AttemptCalculateMagnitude(ActiveEffect.Spec, Magnitude);
-				ABILITY_VLOG(OwnerActor, Log, TEXT("         %s: %s %f"), *Modifier.Attribute.GetName(), *EGameplayModOpToString(Modifier.ModifierOp), Magnitude);
+				UE_VLOG(OwnerActor, LogGameplayEffects, Log, TEXT("         %s: %s %f"), *Modifier.Attribute.GetName(), *EGameplayModOpToString(Modifier.ModifierOp), Magnitude);
 			}
-		}
+		);
 
 		// Clear modified attributes before each periodic execution
 		ActiveEffect.Spec.ModifiedAttributes.Empty();
@@ -4066,7 +4053,18 @@ bool FActiveGameplayEffectsContainer::InternalRemoveActiveGameplayEffect(int32 I
 			return true;
 		}
 
-		ABILITY_LOG(Verbose, TEXT("InternalRemoveActiveGameplayEffect: Auth: %s Handle: %s Def: %s"), IsNetAuthority() ? TEXT("TRUE") : TEXT("FALSE"), *Effect.Handle.ToString(), Effect.Spec.Def ? *Effect.Spec.Def->GetName() : TEXT("NONE"));
+		UE_LOG(LogGameplayEffects, Verbose, TEXT("Removing: %s. Auth: %d. NumToRemove: %d"), *Effect.GetDebugString(), IsNetAuthority(), StacksToRemove);
+		UE_IFVLOG(
+			AActor * OwnerActor = Owner->GetOwnerActor();
+			UE_VLOG(OwnerActor, LogGameplayEffects, Log, TEXT("Removing: %s. Auth: %d. NumToRemove: %d"), *Effect.GetDebugString(), IsNetAuthority(), StacksToRemove);
+			for (const FGameplayModifierInfo& Modifier : Effect.Spec.Def->Modifiers)
+			{
+				float Magnitude = 0.f;
+				Modifier.ModifierMagnitude.AttemptCalculateMagnitude(Effect.Spec, Magnitude);
+				UE_VLOG(OwnerActor, LogGameplayEffects, Log, TEXT("         %s: %s %f"), *Modifier.Attribute.GetName(), *EGameplayModOpToString(Modifier.ModifierOp), Magnitude);
+			}
+		);
+
 
 		FGameplayEffectRemovalInfo GameplayEffectRemovalInfo;
 		GameplayEffectRemovalInfo.ActiveEffect = &Effect;
@@ -4134,7 +4132,7 @@ bool FActiveGameplayEffectsContainer::InternalRemoveActiveGameplayEffect(int32 I
 			// We are locked, so this removal is now pending.
 			PendingRemoves++;
 
-			ABILITY_LOG(Verbose, TEXT("InternalRemoveActiveGameplayEffect while locked; Counting as a Pending Remove: Auth: %s Handle: %s Def: %s"), IsNetAuthority() ? TEXT("TRUE") : TEXT("FALSE"), *Effect.Handle.ToString(), Effect.Spec.Def ? *Effect.Spec.Def->GetName() : TEXT("NONE"));
+			UE_LOG(LogGameplayEffects, Verbose, TEXT("Begin Pending Remove: %s. Auth: %d"), *Effect.GetDebugString(), IsNetAuthority());
 		}
 		else
 		{
@@ -4159,7 +4157,7 @@ bool FActiveGameplayEffectsContainer::InternalRemoveActiveGameplayEffect(int32 I
 		return ModifiedArray;
 	}
 
-	ABILITY_LOG(Warning, TEXT("InternalRemoveActiveGameplayEffect called with invalid index: %d"), Idx);
+	UE_LOG(LogGameplayEffects, Warning, TEXT("InternalRemoveActiveGameplayEffect called with invalid index: %d"), Idx);
 	return false;
 }
 
@@ -4184,7 +4182,7 @@ void FActiveGameplayEffectsContainer::InternalOnActiveGameplayEffectRemoved(FAct
 	}
 	else
 	{
-		ABILITY_LOG(Warning, TEXT("InternalOnActiveGameplayEffectRemoved called with no GameplayEffect: %s"), *Effect.Handle.ToString());
+		UE_LOG(LogGameplayEffects, Warning, TEXT("InternalOnActiveGameplayEffectRemoved called with no GameplayEffect: %s"), *Effect.Handle.ToString());
 	}
 
 	Effect.EventSet.OnEffectRemoved.Broadcast(GameplayEffectRemovalInfo);
@@ -4655,7 +4653,7 @@ void FActiveGameplayEffectsContainer::CheckDuration(FActiveGameplayEffectHandle
 
 				if (Effect.DurationHandle.IsValid() == false)
 				{
-					ABILITY_LOG(Warning, TEXT("Failed to set new timer in ::CheckDuration. Timer trying to be set for: %.2f. Removing GE instead"), NewTimerDuration);
+					UE_LOG(LogGameplayEffects, Warning, TEXT("Failed to set new timer in ::CheckDuration. Timer trying to be set for: %.2f. Removing GE instead"), NewTimerDuration);
 					if (!Effect.IsPendingRemove)
 					{
 						InternalRemoveActiveGameplayEffect(ActiveGEIdx, -1, false);
@@ -5028,11 +5026,11 @@ void FActiveGameplayEffectsContainer::DebugCyclicAggregatorBroadcasts(FAggregato
 		{
 			if (Aggregator == TriggeredAggregator)
 			{
-				ABILITY_LOG(Warning, TEXT(" Attribute %s was the triggered aggregator (%s)"), *Attribute.GetName(), *Owner->GetPathName());
+				UE_LOG(LogGameplayEffects, Warning, TEXT(" Attribute %s was the triggered aggregator (%s)"), *Attribute.GetName(), *Owner->GetPathName());
 			}
 			else if (Aggregator->BroadcastingDirtyCount > 0)
 			{
-				ABILITY_LOG(Warning, TEXT(" Attribute %s is broadcasting dirty (%s)"), *Attribute.GetName(), *Owner->GetPathName());
+				UE_LOG(LogGameplayEffects, Warning, TEXT(" Attribute %s is broadcasting dirty (%s)"), *Attribute.GetName(), *Owner->GetPathName());
 			}
 			else
 			{
@@ -5044,7 +5042,7 @@ void FActiveGameplayEffectsContainer::DebugCyclicAggregatorBroadcasts(FAggregato
 				UAbilitySystemComponent* ASC = Handle.GetOwningAbilitySystemComponent();
 				if (ASC)
 				{
-					ABILITY_LOG(Warning, TEXT("  Dependant (%s) GE: %s"), *ASC->GetPathName(), *GetNameSafe(ASC->GetGameplayEffectDefForHandle(Handle)));
+					UE_LOG(LogGameplayEffects, Warning, TEXT("  Dependant (%s) GE: %s"), *ASC->GetPathName(), *GetNameSafe(ASC->GetGameplayEffectDefForHandle(Handle)));
 				}
 			}
 		}
@@ -5604,7 +5602,7 @@ void FActiveGameplayEffectsContainer::DecrementLock()
 
 			if (Effect.IsPendingRemove)
 			{
-				ABILITY_LOG(Verbose, TEXT("DecrementLock decrementing a pending remove: Auth: %s Handle: %s Def: %s"), IsNetAuthority() ? TEXT("TRUE") : TEXT("FALSE"), *Effect.Handle.ToString(), Effect.Spec.Def ? *Effect.Spec.Def->GetName() : TEXT("NONE"));
+				UE_VLOG_UELOG(Owner->GetOwnerActor(), LogGameplayEffects, Verbose, TEXT("Finish PendingRemove: %s. Auth: %d"), *Effect.GetDebugString(), IsNetAuthority());
 				GameplayEffects_Internal.RemoveAtSwap(idx, 1, false);
 				ModifiedArray = true;
 				PendingRemoves--;
@@ -5613,7 +5611,7 @@ void FActiveGameplayEffectsContainer::DecrementLock()
 
 		if (!ensure(PendingRemoves == 0))
 		{
-			ABILITY_LOG(Error, TEXT("~FScopedActiveGameplayEffectLock has %d pending removes after a scope lock removal"), PendingRemoves);
+			UE_LOG(LogGameplayEffects, Error, TEXT("~FScopedActiveGameplayEffectLock has %d pending removes after a scope lock removal"), PendingRemoves);
 			PendingRemoves = 0;
 		}
 
