Fix for newly created Gameplay Effects not getting their proper version number set to latest.

diff --git a/Engine/Plugins/Runtime/GameplayAbilities/Source/GameplayAbilities/Private/GameplayEffect.cpp b/Engine/Plugins/Runtime/GameplayAbilities/Source/GameplayAbilities/Private/GameplayEffect.cpp
--- a/Engine/Plugins/Runtime/GameplayAbilities/Source/GameplayAbilities/Private/GameplayEffect.cpp
+++ b/Engine/Plugins/Runtime/GameplayAbilities/Source/GameplayAbilities/Private/GameplayEffect.cpp
@@ -122,7 +122,25 @@ void UGameplayEffect::PostInitProperties()
 	Super::PostInitProperties();
 
 #if WITH_EDITORONLY_DATA
-	SetVersion(EGameplayEffectVersion::Monolithic);
+	// On the first load of this package, the version will be inherited from the parent class.
+	// We want to always stomp that value to say it's a Monolithic version (pre-UE5.3), ensuring a child class always
+	// executes the upgrade path at least once.  Once the package has been fully loaded, in PostLoad we will correctly
+	// set the package version.  However, this function will again be called on recompilation of the Blueprint.
+	// In order to account for that, we should check if the Package has a Linker:
+	// This will be True if it's already existed; False if being created (new asset).
+	if (const UPackage* Package = GetPackage())
+	{
+		const bool bAlreadyExists = (Package->GetLinker() != nullptr);
+		if (!bAlreadyExists)
+		{
+			constexpr bool bForceGameThreadValue = true;
+			SetVersion(static_cast<EGameplayEffectVersion>(UE::GameplayEffect::CVarGameplayEffectMaxVersion.GetValueOnAnyThread(bForceGameThreadValue)));
+		}
+		else
+		{
+			SetVersion(EGameplayEffectVersion::Monolithic);
+		}
+	}
 #endif
 }
 
