Fix TryActivateAbilitiesByTag crash due to iterating on pointers to a mutable (resizable) array of ActivatableAbilities.

diff --git a/Engine/Plugins/Runtime/GameplayAbilities/Source/GameplayAbilities/Private/AbilitySystemComponent_Abilities.cpp b/Engine/Plugins/Runtime/GameplayAbilities/Source/GameplayAbilities/Private/AbilitySystemComponent_Abilities.cpp
--- a/Engine/Plugins/Runtime/GameplayAbilities/Source/GameplayAbilities/Private/AbilitySystemComponent_Abilities.cpp
+++ b/Engine/Plugins/Runtime/GameplayAbilities/Source/GameplayAbilities/Private/AbilitySystemComponent_Abilities.cpp
@@ -1350,14 +1350,23 @@ void UAbilitySystemComponent::GetActivatableGameplayAbilitySpecsByAllMatchingTag
 
 bool UAbilitySystemComponent::TryActivateAbilitiesByTag(const FGameplayTagContainer& GameplayTagContainer, bool bAllowRemoteActivation)
 {
-	TArray<FGameplayAbilitySpec*> AbilitiesToActivate;
-	GetActivatableGameplayAbilitySpecsByAllMatchingTags(GameplayTagContainer, AbilitiesToActivate);
+	TArray<FGameplayAbilitySpec*> AbilitiesToActivatePtrs;
+	GetActivatableGameplayAbilitySpecsByAllMatchingTags(GameplayTagContainer, AbilitiesToActivatePtrs);
+	if (AbilitiesToActivatePtrs.Num() < 1)
+	{
+		return false;
+	}
 
-	bool bSuccess = false;
+	// Convert from pointers (which can be reallocated, since they point to internal data) to copies of that data
+	TArray<FGameplayAbilitySpec> AbilitiesToActivate;
+	AbilitiesToActivate.Reserve(AbilitiesToActivatePtrs.Num());
+	Algo::Transform(AbilitiesToActivatePtrs, AbilitiesToActivate, [](FGameplayAbilitySpec* SpecPtr) { return *SpecPtr; });
 
-	for (auto GameplayAbilitySpec : AbilitiesToActivate)
+	bool bSuccess = false;
+	for (const FGameplayAbilitySpec& GameplayAbilitySpec : AbilitiesToActivate)
 	{
-		bSuccess |= TryActivateAbility(GameplayAbilitySpec->Handle, bAllowRemoteActivation);
+		ensure(IsValid(GameplayAbilitySpec.Ability));
+		bSuccess |= TryActivateAbility(GameplayAbilitySpec.Handle, bAllowRemoteActivation);
 	}
 
 	return bSuccess;
