Gameplay Ability Montage / Bug Fix: Fixed montages not replicating what section to play on simulated clients when the AbilitySystemComponent is set to only replicate the current montage section instead of the percise montage position.
The current montage section would only be replicated when jumping to a section, and not when naturally transitioning from one section to the next. This has been fixed by ensuring SectionIdToPlay is updated in AnimMontage_UpdateReplicatedData.
NextSectionID is now also utilized when only replicating section to improve the transition between sections.

#changelist validated

#virtualized

diff --git a/Engine/Plugins/Runtime/GameplayAbilities/Source/GameplayAbilities/Private/AbilitySystemComponent_Abilities.cpp b/Engine/Plugins/Runtime/GameplayAbilities/Source/GameplayAbilities/Private/AbilitySystemComponent_Abilities.cpp
--- a/Engine/Plugins/Runtime/GameplayAbilities/Source/GameplayAbilities/Private/AbilitySystemComponent_Abilities.cpp
+++ b/Engine/Plugins/Runtime/GameplayAbilities/Source/GameplayAbilities/Private/AbilitySystemComponent_Abilities.cpp
@@ -48,6 +48,8 @@ CSV_DECLARE_CATEGORY_MODULE_EXTERN(CORE_API, Basic);
 static TAutoConsoleVariable<float> CVarReplayMontageErrorThreshold(TEXT("replay.MontageErrorThreshold"), 0.5f, TEXT("Tolerance level for when montage playback position correction occurs in replays"));
 static TAutoConsoleVariable<bool> CVarAbilitySystemSetActivationInfoMultipleTimes(TEXT("AbilitySystem.SetActivationInfoMultipleTimes"), false, TEXT("Set this to true if some replicated Gameplay Abilities aren't setting their owning actors correctly"));
 static TAutoConsoleVariable<bool> CVarGasFixClientSideMontageBlendOutTime(TEXT("AbilitySystem.Fix.ClientSideMontageBlendOutTime"), true, TEXT("Enable a fix to replicate the Montage BlendOutTime for (recently) stopped Montages"));
+static TAutoConsoleVariable<bool> CVarUpdateMontageSectionIdToPlay(TEXT("AbilitySystem.UpdateMontageSectionIdToPlay"), true, TEXT("During tick, update the section ID that replicated montages should use"));
+static TAutoConsoleVariable<bool> CVarReplicateMontageNextSectionId(TEXT("AbilitySystem.ReplicateMontageNextSectionId"), true, TEXT("Apply the replicated next section Id to montages when skipping position replication"));
 
 void UAbilitySystemComponent::InitializeComponent()
 {
@@ -2917,6 +2919,12 @@ void UAbilitySystemComponent::AnimMontage_UpdateReplicatedData(FGameplayAbilityR
 		int32 CurrentSectionID = LocalAnimMontageInfo.AnimMontage->GetSectionIndexFromPosition(OutRepAnimMontageInfo.Position);
 		if (CurrentSectionID != INDEX_NONE)
 		{
+			constexpr bool bForceGameThreadValue = true;
+			if (CVarUpdateMontageSectionIdToPlay.GetValueOnAnyThread(bForceGameThreadValue))
+			{
+				OutRepAnimMontageInfo.SectionIdToPlay = uint8(CurrentSectionID + 1);
+			}
+
 			int32 NextSectionID = AnimInstance->Montage_GetNextSectionID(LocalAnimMontageInfo.AnimMontage, CurrentSectionID);
 			if (NextSectionID >= (256 - 1))
 			{
@@ -3025,16 +3033,6 @@ void UAbilitySystemComponent::OnRep_ReplicatedAnimMontage()
 				AnimInstance->Montage_SetPlayRate(LocalAnimMontageInfo.AnimMontage, ConstRepAnimMontageInfo.PlayRate);
 			}
 
-			const int32 SectionIdToPlay = (static_cast<int32>(ConstRepAnimMontageInfo.SectionIdToPlay) - 1);
-			if (SectionIdToPlay != INDEX_NONE)
-			{
-				FName SectionNameToJumpTo = LocalAnimMontageInfo.AnimMontage->GetSectionName(SectionIdToPlay);
-				if (SectionNameToJumpTo != NAME_None)
-				{
-					AnimInstance->Montage_JumpToSection(SectionNameToJumpTo);
-				}
-			}
-
 			// Compressed Flags
 			const bool bIsStopped = AnimInstance->Montage_GetIsStopped(LocalAnimMontageInfo.AnimMontage);
 			const bool bReplicatedIsStopped = bool(ConstRepAnimMontageInfo.IsStopped);
@@ -3097,6 +3095,55 @@ void UAbilitySystemComponent::OnRep_ReplicatedAnimMontage()
 					AnimInstance->Montage_SetPosition(LocalAnimMontageInfo.AnimMontage, ConstRepAnimMontageInfo.Position);
 				}
 			}
+			// Update current and next section if not replicating position
+			else
+			{
+				const float CurrentPosition = AnimInstance->Montage_GetPosition(LocalAnimMontageInfo.AnimMontage);
+				int32 CurrentSectionID = LocalAnimMontageInfo.AnimMontage->GetSectionIndexFromPosition(CurrentPosition);
+				const int32 RepSectionIdToPlay = (static_cast<int32>(ConstRepAnimMontageInfo.SectionIdToPlay) - 1);
+				FName CurrentSectionName = LocalAnimMontageInfo.AnimMontage->GetSectionName(CurrentSectionID);
+
+				// If RepSectionIdToPlay is valid and different from the current section, then jump to it
+				if (RepSectionIdToPlay != INDEX_NONE && RepSectionIdToPlay != CurrentSectionID )
+				{
+					CurrentSectionName = LocalAnimMontageInfo.AnimMontage->GetSectionName(RepSectionIdToPlay);
+					if (CurrentSectionName != NAME_None)
+					{
+						AnimInstance->Montage_JumpToSection(CurrentSectionName);
+						CurrentSectionID = RepSectionIdToPlay;
+					}
+					else
+					{
+						ABILITY_LOG(Warning, TEXT("OnRep_ReplicatedAnimMontage: Failed to replicate current section due to invalid name. Name: %s, Section ID: %i"), 
+						*GetNameSafe(this), 
+						CurrentSectionID);
+					}
+				}
+
+				constexpr bool bForceGameThreadValue = true;
+				if (CVarReplicateMontageNextSectionId.GetValueOnAnyThread(bForceGameThreadValue))
+				{
+					const int32 NextSectionID = AnimInstance->Montage_GetNextSectionID(LocalAnimMontageInfo.AnimMontage, CurrentSectionID);
+					const int32 RepNextSectionID = int32(ConstRepAnimMontageInfo.NextSectionID) - 1;
+
+					// If NextSectionID is different than the replicated one, then set it.
+					if (RepNextSectionID != INDEX_NONE && NextSectionID != RepNextSectionID)
+					{
+						const FName NextSectionName = LocalAnimMontageInfo.AnimMontage->GetSectionName(RepNextSectionID);
+						if (CurrentSectionName != NAME_None && NextSectionName != NAME_None)
+						{
+							AnimInstance->Montage_SetNextSection(CurrentSectionName, NextSectionName, LocalAnimMontageInfo.AnimMontage);
+						}
+						else
+						{
+							ABILITY_LOG(Warning, TEXT("OnRep_ReplicatedAnimMontage: Failed to replicate next section due to invalid name. Name: %s, Current Section ID: %i, Next Section ID: %i"), 
+							*GetNameSafe(this), 
+							CurrentSectionID, 
+							RepNextSectionID);
+						}
+					}
+				}
+			}
 		}
 	}
 }
