The crash occurs because after removing an attribute set, aggregators remain in the ActiveGameplayEffectsContainer for removed attributes. Calling OnDirty() on those aggregators on removing a GE then executes AbilitySystemComponent functions that assume that a matching attribute set exists.

With this change, when an attribute set is removed, aggregators for that set's attributes are removed from ActiveGameplayEffectsContainers immediately. Then, removing a gameplay effect afterwards will not attempt to access aggregators for attributes that don't actively have an attribute set in the ASC.

Also introduced a reusable FGameplayAttribute iteration by attribute set class, instead of copying the not-immediately-clear TFieldIterator approach that was used in several places.

diff --git a/Engine/Plugins/Runtime/GameplayAbilities/Source/GameplayAbilities/Private/AbilitySystemComponent.cpp b/Engine/Plugins/Runtime/GameplayAbilities/Source/GameplayAbilities/Private/AbilitySystemComponent.cpp
--- a/Engine/Plugins/Runtime/GameplayAbilities/Source/GameplayAbilities/Private/AbilitySystemComponent.cpp
+++ b/Engine/Plugins/Runtime/GameplayAbilities/Source/GameplayAbilities/Private/AbilitySystemComponent.cpp
@@ -142,17 +142,7 @@ void UAbilitySystemComponent::GetAllAttributes(TArray<FGameplayAttribute>& OutAt
 			continue;
 		}
 
-		for (TFieldIterator<FProperty> It(Set->GetClass()); It; ++It)
-		{
-			if (FFloatProperty* FloatProperty = CastField<FFloatProperty>(*It))
-			{
-				OutAttributes.Push(FGameplayAttribute(FloatProperty));
-			}
-			else if (FGameplayAttribute::IsGameplayAttributeDataProperty(*It))
-			{
-				OutAttributes.Push(FGameplayAttribute(*It));
-			}
-		}
+		UAttributeSet::GetAttributesFromSetClass(Set->GetClass(), OutAttributes);
 	}
 }
 
@@ -2619,17 +2609,16 @@ void UAbilitySystemComponent::Debug_Internal(FAbilitySystemComponentDebugInfo& I
 				continue;
 			}
 
-			for (TFieldIterator<FProperty> It(Set->GetClass()); It; ++It)
+			TArray<FGameplayAttribute> Attributes;
+			UAttributeSet::GetAttributesFromSetClass(Set->GetClass(), Attributes);
+			for (const FGameplayAttribute& Attribute : Attributes)
 			{
-				FGameplayAttribute	Attribute(*It);
-
 				if (DrawAttributes.Contains(Attribute))
 					continue;
 
 				if (Attribute.IsValid())
 				{
-					float Value = GetNumericAttribute(Attribute);
-
+					const float Value = GetNumericAttribute(Attribute);
 					DebugLine(Info, FString::Printf(TEXT("%s %.2f"), *Attribute.GetName(), Value), 4.f, 0.f);
 				}
 			}
@@ -2866,13 +2855,21 @@ void UAbilitySystemComponent::AddSpawnedAttribute(UAttributeSet* Attribute)
 	}
 }
 
-void UAbilitySystemComponent::RemoveSpawnedAttribute(UAttributeSet* Attribute)
+void UAbilitySystemComponent::RemoveSpawnedAttribute(UAttributeSet* AttributeSet)
 {
-	if (SpawnedAttributes.RemoveSingle(Attribute) > 0)
+	if (SpawnedAttributes.RemoveSingle(AttributeSet) > 0)
 	{
 		if (IsUsingRegisteredSubObjectList())
 		{
-			RemoveReplicatedSubObject(Attribute);
+			RemoveReplicatedSubObject(AttributeSet);
+		}
+
+		TArray<FGameplayAttribute> Attributes;
+		UAttributeSet::GetAttributesFromSetClass(AttributeSet->GetClass(), Attributes);
+		for (const FGameplayAttribute& Attribute : Attributes)
+		{
+			ABILITY_LOG(Log, TEXT("Cleaning up aggregator for attribute '%s' due to RemoveSpawnedAttribute removing attribute set '%s'"), *Attribute.GetName(), *AttributeSet->GetName());
+			ActiveGameplayEffects.CleanupAttributeAggregator(Attribute);
 		}
 
 		SetSpawnedAttributesListDirty();
@@ -2939,6 +2936,21 @@ void UAbilitySystemComponent::OnRep_SpawnedAttributes(const TArray<UAttributeSet
 			}
 		}
 	}
+
+	// Find the attribute sets that got removed
+	for (UAttributeSet* PreviousAttributeSet : PreviousSpawnedAttributes)
+	{
+		if (PreviousAttributeSet && SpawnedAttributes.Find(PreviousAttributeSet) == INDEX_NONE)
+		{
+			TArray<FGameplayAttribute> Attributes;
+			UAttributeSet::GetAttributesFromSetClass(PreviousAttributeSet->GetClass(), Attributes);
+			for (const FGameplayAttribute& Attribute : Attributes)
+			{
+				ABILITY_LOG(Log, TEXT("Cleaning up aggregator for attribute '%s' due to OnRep_SpawnedAttributes detecting removal of '%s'"), *Attribute.GetName(), *PreviousAttributeSet->GetName());
+				ActiveGameplayEffects.CleanupAttributeAggregator(Attribute);
+			}
+		}
+	}
 }
 
 void UAbilitySystemComponent::AddReplicatedInstancedAbility(UGameplayAbility* GameplayAbility)
diff --git a/Engine/Plugins/Runtime/GameplayAbilities/Source/GameplayAbilities/Private/AttributeSet.cpp b/Engine/Plugins/Runtime/GameplayAbilities/Source/GameplayAbilities/Private/AttributeSet.cpp
--- a/Engine/Plugins/Runtime/GameplayAbilities/Source/GameplayAbilities/Private/AttributeSet.cpp
+++ b/Engine/Plugins/Runtime/GameplayAbilities/Source/GameplayAbilities/Private/AttributeSet.cpp
@@ -404,6 +404,21 @@ bool UAttributeSet::IsSupportedForNetworking() const
 	return true;
 }
 
+void UAttributeSet::GetAttributesFromSetClass(const TSubclassOf<UAttributeSet>& AttributeSetClass, TArray<FGameplayAttribute>& Attributes)
+{
+	for (TFieldIterator<FProperty> It(AttributeSetClass); It; ++It)
+	{
+		if (FFloatProperty* FloatProperty = CastField<FFloatProperty>(*It))
+		{
+			Attributes.Add(FGameplayAttribute(FloatProperty));
+		}
+		else if (FGameplayAttribute::IsGameplayAttributeDataProperty(*It))
+		{
+			Attributes.Add(FGameplayAttribute(*It));
+		}
+	}
+}
+
 void UAttributeSet::SetNetAddressable()
 {
 	bNetAddressable = true;
diff --git a/Engine/Plugins/Runtime/GameplayAbilities/Source/GameplayAbilities/Private/GameplayEffect.cpp b/Engine/Plugins/Runtime/GameplayAbilities/Source/GameplayAbilities/Private/GameplayEffect.cpp
--- a/Engine/Plugins/Runtime/GameplayAbilities/Source/GameplayAbilities/Private/GameplayEffect.cpp
+++ b/Engine/Plugins/Runtime/GameplayAbilities/Source/GameplayAbilities/Private/GameplayEffect.cpp
@@ -2866,7 +2866,7 @@ const FActiveGameplayEffect* FActiveGameplayEffectsContainer::GetActiveGameplayE
 	return nullptr;
 }
 
-FAggregatorRef& FActiveGameplayEffectsContainer::FindOrCreateAttributeAggregator(FGameplayAttribute Attribute)
+FAggregatorRef& FActiveGameplayEffectsContainer::FindOrCreateAttributeAggregator(const FGameplayAttribute& Attribute)
 {
 	FAggregatorRef* RefPtr = AttributeAggregatorMap.Find(Attribute);
 	if (RefPtr)
@@ -2893,6 +2893,23 @@ FAggregatorRef& FActiveGameplayEffectsContainer::FindOrCreateAttributeAggregator
 	return AttributeAggregatorMap.Add(Attribute, FAggregatorRef(NewAttributeAggregator));
 }
 
+void FActiveGameplayEffectsContainer::CleanupAttributeAggregator(const FGameplayAttribute& Attribute)
+{
+	FAggregatorRef* RefPtr = AttributeAggregatorMap.Find(Attribute);
+	if (RefPtr)
+	{
+		UE_LOG(LogGameplayEffects, Log, TEXT("Removing entry in AttributeAggregatorMap for %s."), *Attribute.GetName());
+
+		// No longer interested in OnDirty events for this aggregator, in case other sources call it.
+		RefPtr->Data->OnDirty.RemoveAll(Owner);
+		RefPtr->Data->OnDirtyRecursive.RemoveAll(Owner);
+
+		// Remove the aggregator from the map, we no longer use it. If an attribute set gets added again and gameplay effect requires an aggregator
+		// for this attribute, a new one would be created via FindOrCreateAttributeAggregator with a clean state.
+		AttributeAggregatorMap.Remove(Attribute);
+	}
+}
+
 void FActiveGameplayEffectsContainer::OnAttributeAggregatorDirty(FAggregator* Aggregator, FGameplayAttribute Attribute, bool bFromRecursiveCall)
 {
 	check(AttributeAggregatorMap.FindChecked(Attribute).Get() == Aggregator);
@@ -4213,12 +4230,11 @@ void FActiveGameplayEffectsContainer::RemoveActiveGameplayEffectGrantedTagsAndMo
 	// Update AttributeAggregators: remove mods from this ActiveGE Handle
 	if (Effect.Spec.GetPeriod() <= UGameplayEffect::NO_PERIOD)
 	{
-		for(const FGameplayModifierInfo& Mod : Effect.Spec.Def->Modifiers)
+		for (const FGameplayModifierInfo& Mod : Effect.Spec.Def->Modifiers)
 		{
-			if(Mod.Attribute.IsValid())
+			if (Mod.Attribute.IsValid())
 			{
-				FAggregatorRef* RefPtr = AttributeAggregatorMap.Find(Mod.Attribute);
-				if(RefPtr)
+				if (const FAggregatorRef* RefPtr = AttributeAggregatorMap.Find(Mod.Attribute))
 				{
 					RefPtr->Get()->RemoveAggregatorMod(Effect.Handle);
 				}
diff --git a/Engine/Plugins/Runtime/GameplayAbilities/Source/GameplayAbilities/Public/AttributeSet.h b/Engine/Plugins/Runtime/GameplayAbilities/Source/GameplayAbilities/Public/AttributeSet.h
--- a/Engine/Plugins/Runtime/GameplayAbilities/Source/GameplayAbilities/Public/AttributeSet.h
+++ b/Engine/Plugins/Runtime/GameplayAbilities/Source/GameplayAbilities/Public/AttributeSet.h
@@ -182,6 +182,8 @@ class GAMEPLAYABILITIES_API UAttributeSet : public UObject
 	GENERATED_UCLASS_BODY()
 
 public:
+	// Populates (without emptying) a TArray with all FGameplayAttributes from an attribute set class.
+	static void GetAttributesFromSetClass(const TSubclassOf<UAttributeSet>& AttributeSetClass, TArray<FGameplayAttribute>& Attributes);
 
 	/** Override to disable initialization for specific properties */
 	virtual bool ShouldInitProperty(bool FirstInit, FProperty* PropertyToInit) const { return true; }
diff --git a/Engine/Plugins/Runtime/GameplayAbilities/Source/GameplayAbilities/Public/GameplayEffect.h b/Engine/Plugins/Runtime/GameplayAbilities/Source/GameplayAbilities/Public/GameplayEffect.h
--- a/Engine/Plugins/Runtime/GameplayAbilities/Source/GameplayAbilities/Public/GameplayEffect.h
+++ b/Engine/Plugins/Runtime/GameplayAbilities/Source/GameplayAbilities/Public/GameplayEffect.h
@@ -1908,8 +1908,9 @@ private:
 	/** A map to manage stacking while we are the source */
 	TMap<TWeakObjectPtr<UGameplayEffect>, TArray<FActiveGameplayEffectHandle> >	SourceStackingMap;
 	
+	FAggregatorRef& FindOrCreateAttributeAggregator(const FGameplayAttribute& Attribute);
 
-	FAggregatorRef& FindOrCreateAttributeAggregator(FGameplayAttribute Attribute);
+	void CleanupAttributeAggregator(const FGameplayAttribute& Attribute);
 
 	void OnAttributeAggregatorDirty(FAggregator* Aggregator, FGameplayAttribute Attribute, bool FromRecursiveCall=false);
 
