Fixing a longstanding bug that replicated Gameplay Abilities are being replicated to Simulated Proxies (but we expect them to only be replicated to the Owning Client -- Autonomous Proxies).

diff --git a/Engine/Plugins/Runtime/GameplayAbilities/Source/GameplayAbilities/Private/AbilitySystemComponent.cpp b/Engine/Plugins/Runtime/GameplayAbilities/Source/GameplayAbilities/Private/AbilitySystemComponent.cpp
--- a/Engine/Plugins/Runtime/GameplayAbilities/Source/GameplayAbilities/Private/AbilitySystemComponent.cpp
+++ b/Engine/Plugins/Runtime/GameplayAbilities/Source/GameplayAbilities/Private/AbilitySystemComponent.cpp
@@ -39,6 +39,9 @@ DECLARE_CYCLE_STAT(TEXT("AbilitySystemComp ExecuteGameplayEffect"), STAT_Ability
 static bool bUseReplicationConditionForActiveGameplayEffects = true;
 static FAutoConsoleVariableRef CVarUseReplicationConditionForActiveGameplayEffects(TEXT("AbilitySystem.UseReplicationConditionForActiveGameplayEffects"), bUseReplicationConditionForActiveGameplayEffects, TEXT("Whether to be able to determine the replication condition for AbilitySystemComponent::ActiveGameplayEffects at runtime. Removes the need for executing custom logic in FActiveGameplayEffects::NetDeltaSerialize. Default is true."));
 
+static bool bReplicateAbilitiesToSimulatedProxies = false;
+static FAutoConsoleVariableRef CVarReplicateGameplayAbilitiesToOwnerOnly(TEXT("AbilitySystem.Fix.ReplicateAbilitiesToSimulatedProxies"), bReplicateAbilitiesToSimulatedProxies, TEXT("Default: False.  When false, Gameplay Abilities replicate to AutonomousProxies only, not SimulatedProxies (surmised to be a bug)"));
+
 UAbilitySystemComponent::UAbilitySystemComponent(const FObjectInitializer& ObjectInitializer)
 	: Super(ObjectInitializer)
 	, GameplayTagCountContainer()
@@ -244,7 +247,8 @@ void UAbilitySystemComponent::ReadyForReplication()
 		{
 			if (ReplicatedAbility)
 			{
-				AddReplicatedSubObject(ReplicatedAbility);
+				const ELifetimeCondition LifetimeCondition = bReplicateAbilitiesToSimulatedProxies ? COND_None : COND_ReplayOrOwner;
+				AddReplicatedSubObject(ReplicatedAbility, LifetimeCondition);
 			}
 		}
 
@@ -2962,7 +2966,8 @@ void UAbilitySystemComponent::AddReplicatedInstancedAbility(UGameplayAbility* Ga
 		
 		if (IsUsingRegisteredSubObjectList() && IsReadyForReplication())
 		{
-			AddReplicatedSubObject(GameplayAbility);
+			const ELifetimeCondition LifetimeCondition = bReplicateAbilitiesToSimulatedProxies ? COND_None : COND_ReplayOrOwner;
+			AddReplicatedSubObject(GameplayAbility, LifetimeCondition);
 		}
 	}
 }
