AbilitySystemComponent - Add method to test a Query against the component's explicit tags.
This allows callers to avoid copying all tags into another Container to test a Query

diff --git a/Engine/Plugins/Runtime/GameplayAbilities/Source/GameplayAbilities/Public/AbilitySystemComponent.h b/Engine/Plugins/Runtime/GameplayAbilities/Source/GameplayAbilities/Public/AbilitySystemComponent.h
--- a/Engine/Plugins/Runtime/GameplayAbilities/Source/GameplayAbilities/Public/AbilitySystemComponent.h
+++ b/Engine/Plugins/Runtime/GameplayAbilities/Source/GameplayAbilities/Public/AbilitySystemComponent.h
@@ -568,6 +568,12 @@ class GAMEPLAYABILITIES_API UAbilitySystemComponent : public UGameplayTasksCompo
 		TagContainer.AppendTags(GameplayTagCountContainer.GetExplicitGameplayTags());
 	}
 
+	/** Checks whether the query matches the owned GameplayTags */
+	FORCEINLINE bool MatchesGameplayTagQuery(const FGameplayTagQuery& TagQuery) const
+	{
+		return TagQuery.Matches(GameplayTagCountContainer.GetExplicitGameplayTags());
+	}
+
 	/** Returns the number of instances of a given tag */
 	FORCEINLINE int32 GetTagCount(FGameplayTag TagToCheck) const
 	{
