Removing a GameplayEffect that unhibits while being removed would previously cause an ensure to happen.  This is actually a valid use case.

diff --git a/Engine/Plugins/Runtime/GameplayAbilities/Source/GameplayAbilities/Private/GameplayEffectComponents/TargetTagRequirementsGameplayEffectComponent.cpp b/Engine/Plugins/Runtime/GameplayAbilities/Source/GameplayAbilities/Private/GameplayEffectComponents/TargetTagRequirementsGameplayEffectComponent.cpp
--- a/Engine/Plugins/Runtime/GameplayAbilities/Source/GameplayAbilities/Private/GameplayEffectComponents/TargetTagRequirementsGameplayEffectComponent.cpp
+++ b/Engine/Plugins/Runtime/GameplayAbilities/Source/GameplayAbilities/Private/GameplayEffectComponents/TargetTagRequirementsGameplayEffectComponent.cpp
@@ -116,8 +116,9 @@ void UTargetTagRequirementsGameplayEffectComponent::OnTagChanged(const FGameplay
 		return;
 	}
 
+	// It's possible for this to return nullptr if it was in the process of being removed (IsPendingRemove)
 	const FActiveGameplayEffect* ActiveGE = Owner->GetActiveGameplayEffect(ActiveGEHandle);
-	if (ensure(ActiveGE) && !ActiveGE->IsPendingRemove)
+	if (ActiveGE)
 	{
 		FGameplayTagContainer OwnedTags;
 		Owner->GetOwnedGameplayTags(OwnedTags);
