GameplayEffect's GrantedAbilities were not auto-upgrading properly.  Removing GrantedAbilities in a child GE would previously not be respected upon upgrading.

diff --git a/Engine/Plugins/Runtime/GameplayAbilities/Source/GameplayAbilities/Private/GameplayEffect.cpp b/Engine/Plugins/Runtime/GameplayAbilities/Source/GameplayAbilities/Private/GameplayEffect.cpp
--- a/Engine/Plugins/Runtime/GameplayAbilities/Source/GameplayAbilities/Private/GameplayEffect.cpp
+++ b/Engine/Plugins/Runtime/GameplayAbilities/Source/GameplayAbilities/Private/GameplayEffect.cpp
@@ -363,12 +363,21 @@ void UGameplayEffect::ConvertAbilitiesComponent()
 	// Get the archetype to compare to
 	const UGameplayEffect* Archetype = CastChecked<UGameplayEffect>(GetArchetype());
 
+	// Due to the delta serialization, this will only be true if we're not inheriting the parent's values
 	const bool bChanged = !(GrantedAbilities == Archetype->GrantedAbilities);
 	if (bChanged && UE::GameplayEffect::EditorOnly::ShouldUpgradeVersion(GetVersion(), EGameplayEffectVersion::AbilitiesComponent53))
 	{
 		UAbilitiesGameplayEffectComponent& AbilitiesComponent = FindOrAddComponent<UAbilitiesGameplayEffectComponent>();
+
+		// Since we've already determined we're not inheriting the parent values, clobber the array
+		AbilitiesComponent.GrantAbilityConfigs.Empty();
 		for (const FGameplayAbilitySpecDef& GASpecDef : GrantedAbilities)
 		{
+			if (!GASpecDef.Ability.Get())
+			{
+				continue;
+			}
+
 			FGameplayAbilitySpecConfig GASpecConfig;
 			GASpecConfig.Ability = GASpecDef.Ability;
 			GASpecConfig.InputID = GASpecDef.InputID;
@@ -379,8 +388,12 @@ void UGameplayEffect::ConvertAbilitiesComponent()
 			AbilitiesComponent.AddGrantedAbilityConfig(GASpecConfig);
 		}
 
-		// We should empty these because the old code path that executes them is still around.
-		GrantedAbilities.Empty();
+		// We shouldn't just empty out the deprecated GrantedAbilities because then we wouldn't know if a Child's Empty array was an override or inheritance.
+		// Instead, let's null-out the entries so they won't affect gameplay, but the child will still inherit the nulled-out list rather than empty list.
+		for (FGameplayAbilitySpecDef& Entry : GrantedAbilities)
+		{
+			Entry.Ability = nullptr;
+		}
 	}
 }
 
@@ -1459,15 +1472,24 @@ PRAGMA_DISABLE_DEPRECATION_WARNINGS
 	// ------------------------------------------------
 
 	// Make Granted AbilitySpecs (caller may modify these specs after creating spec, which is why we dont just reference them from the def)
-	// Note: These are going to be deprecated in the future in favor of immutable GameplayEffectComponents.
-	GrantedAbilitySpecs = InDef->GrantedAbilities;
+	// Note: GrantedAbilitySpecs is going to be removed in the future in favor of immutable GameplayEffectComponents.  Don't rely on this (dynamic) functionality.
+	GrantedAbilitySpecs.Empty();
 
 	// if we're granting abilities and they don't specify a source object use the source of this GE
-	for (FGameplayAbilitySpecDef& AbilitySpecDef : GrantedAbilitySpecs)
+	for (const FGameplayAbilitySpecDef& AbilitySpecDef : InDef->GrantedAbilities)
 	{
-		if (AbilitySpecDef.SourceObject == nullptr)
+		// Don't copy null entries over, these would have been nulled during the conversion to use GEComponents.
+		if (!AbilitySpecDef.Ability.Get())
+		{
+			continue;
+		}
+
+		UE_LOG(LogGameplayEffects, Error, TEXT("%s had GrantedAbilities dynamically added to it. This functionality is being deprecated"), *GetNameSafe(InDef));
+		FGameplayAbilitySpecDef& OurCopy = GrantedAbilitySpecs.Emplace_GetRef(AbilitySpecDef);
+
+		if (OurCopy.SourceObject == nullptr)
 		{
-			AbilitySpecDef.SourceObject = InEffectContext.GetSourceObject();
+			OurCopy.SourceObject = InEffectContext.GetSourceObject();
 		}
 	}
 PRAGMA_ENABLE_DEPRECATION_WARNINGS
@@ -3976,6 +3998,11 @@ PRAGMA_DISABLE_DEPRECATION_WARNINGS
 	{
 		for (FGameplayAbilitySpecDef& AbilitySpecDef : Effect.Spec.GrantedAbilitySpecs)
 		{
+			if (!AbilitySpecDef.Ability.Get())
+			{
+				continue;
+			}
+
 			// Only do this if we haven't assigned the ability yet! This prevents cases where stacking GEs
 			// would regrant the ability every time the stack was applied
 			if (AbilitySpecDef.AssignedHandle.IsValid() == false)
diff --git a/Engine/Plugins/Runtime/GameplayAbilities/Source/GameplayAbilities/Private/GameplayEffectComponents/AbilitiesGameplayEffectComponent.cpp b/Engine/Plugins/Runtime/GameplayAbilities/Source/GameplayAbilities/Private/GameplayEffectComponents/AbilitiesGameplayEffectComponent.cpp
--- a/Engine/Plugins/Runtime/GameplayAbilities/Source/GameplayAbilities/Private/GameplayEffectComponents/AbilitiesGameplayEffectComponent.cpp
+++ b/Engine/Plugins/Runtime/GameplayAbilities/Source/GameplayAbilities/Private/GameplayEffectComponents/AbilitiesGameplayEffectComponent.cpp
@@ -24,18 +24,9 @@ UAbilitiesGameplayEffectComponent::UAbilitiesGameplayEffectComponent()
 
 void UAbilitiesGameplayEffectComponent::AddGrantedAbilityConfig(const FGameplayAbilitySpecConfig& Config)
 {
-	for (const FGameplayAbilitySpecConfig& Existing : GrantAbilityConfigs)
-	{
-		if (Existing == Config)
-		{
-			return;
-		}
-	}
-
-	GrantAbilityConfigs.Emplace(Config);
+	GrantAbilityConfigs.AddUnique(Config);
 }
 
-
 // Reminder: this is happening on the authority only
 void UAbilitiesGameplayEffectComponent::OnInhibitionChanged(FActiveGameplayEffectHandle ActiveGEHandle, bool bIsInhibited) const
 {
diff --git a/Engine/Plugins/Runtime/GameplayAbilities/Source/GameplayAbilities/Public/GameplayEffectComponents/AbilitiesGameplayEffectComponent.h b/Engine/Plugins/Runtime/GameplayAbilities/Source/GameplayAbilities/Public/GameplayEffectComponents/AbilitiesGameplayEffectComponent.h
--- a/Engine/Plugins/Runtime/GameplayAbilities/Source/GameplayAbilities/Public/GameplayEffectComponents/AbilitiesGameplayEffectComponent.h
+++ b/Engine/Plugins/Runtime/GameplayAbilities/Source/GameplayAbilities/Public/GameplayEffectComponents/AbilitiesGameplayEffectComponent.h
@@ -35,6 +35,8 @@ struct FGameplayAbilitySpecConfig
 UCLASS()
 class GAMEPLAYABILITIES_API UAbilitiesGameplayEffectComponent : public UGameplayEffectComponent
 {
+	friend class UGameplayEffect; // for upgrade path
+
 	GENERATED_BODY()
 
 public:
