InhibitActiveGameplayEffect was reusing an ActiveGE which could have been invalidated.  Deprecating that function in favor of SetActiveGameplayEffectInhibit to make it clear to users this can happen.

diff --git a/Engine/Plugins/Runtime/GameplayAbilities/Source/GameplayAbilities/Private/AbilitySystemComponent.cpp b/Engine/Plugins/Runtime/GameplayAbilities/Source/GameplayAbilities/Private/AbilitySystemComponent.cpp
--- a/Engine/Plugins/Runtime/GameplayAbilities/Source/GameplayAbilities/Private/AbilitySystemComponent.cpp
+++ b/Engine/Plugins/Runtime/GameplayAbilities/Source/GameplayAbilities/Private/AbilitySystemComponent.cpp
@@ -270,18 +270,28 @@ void UAbilitySystemComponent::CacheIsNetSimulated()
 }
 
 void UAbilitySystemComponent::InhibitActiveGameplayEffect(FActiveGameplayEffectHandle ActiveGEHandle, bool bInhibit, bool bInvokeGameplayCueEvents)
+{
+	FActiveGameplayEffectHandle ContinuationHandle = SetActiveGameplayEffectInhibit(MoveTemp(ActiveGEHandle), bInhibit, bInvokeGameplayCueEvents);
+	ensureMsgf(ContinuationHandle.IsValid(), TEXT("InhibitActiveGameplayEffect invalidated the incoming ActiveGEHandle. Update your code to SetActiveGameplayEffectInhibit so it's clear the incoming handle can be invalidated."));
+}
+
+FActiveGameplayEffectHandle UAbilitySystemComponent::SetActiveGameplayEffectInhibit(FActiveGameplayEffectHandle&& ActiveGEHandle, bool bInhibit, bool bInvokeGameplayCueEvents)
 {
 	FActiveGameplayEffect* ActiveGE = ActiveGameplayEffects.GetActiveGameplayEffect(ActiveGEHandle);
 	if (!ActiveGE)
 	{
 		ABILITY_LOG(Error, TEXT("%s received bad Active GameplayEffect Handle: %s"), ANSI_TO_TCHAR(__func__), *ActiveGEHandle.ToString());
-		return;
+		return FActiveGameplayEffectHandle();
 	}
 
 	if (ActiveGE->bIsInhibited != bInhibit)
 	{
 		ActiveGE->bIsInhibited = bInhibit;
 
+		// It's possible the adding or removing of the tags can invalidate the ActiveGE.  As such,
+		// let's make sure we hold on to that memory until this function is done.
+		FScopedActiveGameplayEffectLock ScopeLockActiveGameplayEffects(ActiveGameplayEffects);
+
 		// All OnDirty callbacks must be inhibited until we update this entire GameplayEffect.
 		FScopedAggregatorOnDirtyBatch	AggregatorOnDirtyBatcher;
 		if (bInhibit)
@@ -294,8 +304,21 @@ void UAbilitySystemComponent::InhibitActiveGameplayEffect(FActiveGameplayEffectH
 			ActiveGameplayEffects.AddActiveGameplayEffectGrantedTagsAndModifiers(*ActiveGE, bInvokeGameplayCueEvents);
 		}
 
-		ActiveGE->EventSet.OnInhibitionChanged.Broadcast(ActiveGEHandle, ActiveGE->bIsInhibited);
+		// The act of executing anything on the ActiveGE can invalidate it.  So we need to recheck if we can continue to execute the callbacks.
+		if (!ActiveGE->IsPendingRemove)
+		{
+			ActiveGE->EventSet.OnInhibitionChanged.Broadcast(ActiveGEHandle, ActiveGE->bIsInhibited);
+		}
+
+		// We lost that it was active somewhere along the way, let the caller know
+		if (ActiveGE->IsPendingRemove)
+		{
+			return FActiveGameplayEffectHandle();
+		}
 	}
+
+	// Normal case is the passed-in ActiveGEHandle is still active and thus can continue execution
+	return MoveTemp(ActiveGEHandle);
 }
 
 
diff --git a/Engine/Plugins/Runtime/GameplayAbilities/Source/GameplayAbilities/Private/GameplayEffect.cpp b/Engine/Plugins/Runtime/GameplayAbilities/Source/GameplayAbilities/Private/GameplayEffect.cpp
--- a/Engine/Plugins/Runtime/GameplayAbilities/Source/GameplayAbilities/Private/GameplayEffect.cpp
+++ b/Engine/Plugins/Runtime/GameplayAbilities/Source/GameplayAbilities/Private/GameplayEffect.cpp
@@ -3911,10 +3911,11 @@ void FActiveGameplayEffectsContainer::InternalOnActiveGameplayEffectAdded(FActiv
 	AddCustomMagnitudeExternalDependencies(Effect);
 
 	const bool bActive = EffectDef->OnAddedToActiveContainer(*this, Effect);
+	Effect.bIsInhibited = true; // Effect has to start inhibited, so our call to Inhibit will trigger if we should be active
 
 	constexpr bool bInvokeCuesIfEnabled = false;
-	Effect.bIsInhibited = true; // Effect has to start inhibited, so our call to Inhibit will trigger if we should be active
-	Owner->InhibitActiveGameplayEffect(Effect.Handle, !bActive, bInvokeCuesIfEnabled);
+	FActiveGameplayEffectHandle EffectHandle = Effect.Handle;
+	Owner->SetActiveGameplayEffectInhibit(MoveTemp(EffectHandle), !bActive, bInvokeCuesIfEnabled);
 }
 
 void FActiveGameplayEffectsContainer::AddActiveGameplayEffectGrantedTagsAndModifiers(FActiveGameplayEffect& Effect, bool bInvokeGameplayCueEvents)
diff --git a/Engine/Plugins/Runtime/GameplayAbilities/Source/GameplayAbilities/Private/GameplayEffectComponents/TargetTagRequirementsGameplayEffectComponent.cpp b/Engine/Plugins/Runtime/GameplayAbilities/Source/GameplayAbilities/Private/GameplayEffectComponents/TargetTagRequirementsGameplayEffectComponent.cpp
--- a/Engine/Plugins/Runtime/GameplayAbilities/Source/GameplayAbilities/Private/GameplayEffectComponents/TargetTagRequirementsGameplayEffectComponent.cpp
+++ b/Engine/Plugins/Runtime/GameplayAbilities/Source/GameplayAbilities/Private/GameplayEffectComponents/TargetTagRequirementsGameplayEffectComponent.cpp
@@ -135,7 +135,7 @@ void UTargetTagRequirementsGameplayEffectComponent::OnTagChanged(const FGameplay
 			// See if we should be inhibiting the execution
 			constexpr bool bInvokeCuesIfStateChanged = true;
 			const bool bOngoingRequirementsMet = OngoingTagRequirements.IsEmpty() || OngoingTagRequirements.RequirementsMet(OwnedTags);
-			Owner->InhibitActiveGameplayEffect(ActiveGEHandle, !bOngoingRequirementsMet, bInvokeCuesIfStateChanged);
+			Owner->SetActiveGameplayEffectInhibit(MoveTemp(ActiveGEHandle), !bOngoingRequirementsMet, bInvokeCuesIfStateChanged);
 		}
 	}
 }
diff --git a/Engine/Plugins/Runtime/GameplayAbilities/Source/GameplayAbilities/Public/AbilitySystemComponent.h b/Engine/Plugins/Runtime/GameplayAbilities/Source/GameplayAbilities/Public/AbilitySystemComponent.h
--- a/Engine/Plugins/Runtime/GameplayAbilities/Source/GameplayAbilities/Public/AbilitySystemComponent.h
+++ b/Engine/Plugins/Runtime/GameplayAbilities/Source/GameplayAbilities/Public/AbilitySystemComponent.h
@@ -409,8 +409,16 @@ class GAMEPLAYABILITIES_API UAbilitySystemComponent : public UGameplayTasksCompo
 	virtual void SetActiveGameplayEffectLevelUsingQuery(FGameplayEffectQuery Query, int32 NewLevel);
 
 	/** Inhibit an active gameplay effect so that it is disabled, but not removed */
+	UE_DEPRECATED(5.4, "Use SetActiveGameplayEffectInhibit with a MoveTemp(ActiveGEHandle) so it's clear the Handle is no longer valid. Check (then use) the returned FActiveGameplayEffectHandle to continue your operation.")
 	virtual void InhibitActiveGameplayEffect(FActiveGameplayEffectHandle ActiveGEHandle, bool bInhibit, bool bInvokeGameplayCueEvents);
 
+	/**
+	 * (Un-)Inhibit an Active Gameplay Effect so it may be disabled (and perform some disabling actions, such as removing tags).
+	 * An inhibited Active Gameplay Effect will lay dormant for re-enabling on some condition (usually tags).  When it's uninhibited, it will reapply some of its functionality.
+	 * NOTE:  The passed-in ActiveGEHandle is no longer valid.  Use the returned FActiveGameplayEffectHandle to determine if the Active GE Handle is still active.
+	 */
+	virtual FActiveGameplayEffectHandle SetActiveGameplayEffectInhibit(FActiveGameplayEffectHandle&& ActiveGEHandle, bool bInhibit, bool bInvokeGameplayCueEvents);
+
 	/**
 	 * Raw accessor to ask the magnitude of a gameplay effect, not necessarily always correct. How should outside code (UI, etc) ask things like 'how much is this gameplay effect modifying my damage by'
 	 * (most likely we want to catch this on the backend - when damage is applied we can get a full dump/history of how the number got to where it is. But still we may need polling methods like below (how much would my damage be)
