FActiveGameplayEffect::GrantedAbilities was not handled property in copy/move constructor/operators.  There were other questionable cases in the code that have also been handled.

diff --git a/Engine/Plugins/Runtime/GameplayAbilities/Source/GameplayAbilities/Private/GameplayEffect.cpp b/Engine/Plugins/Runtime/GameplayAbilities/Source/GameplayAbilities/Private/GameplayEffect.cpp
--- a/Engine/Plugins/Runtime/GameplayAbilities/Source/GameplayAbilities/Private/GameplayEffect.cpp
+++ b/Engine/Plugins/Runtime/GameplayAbilities/Source/GameplayAbilities/Private/GameplayEffect.cpp
@@ -67,6 +67,21 @@ DECLARE_CYCLE_STAT(TEXT("MakeQuery"), STAT_MakeGameplayEffectQuery, STATGROUP_Ab
 
 namespace UE::GameplayEffect
 {
+	enum class EActiveGameplayEffectFix : int32
+	{
+		None = 0,									// No additional fixes to UE5.3 are attempted
+		MostRecentArrayReplicationKey	= (1 << 0),	// MostRecentArrayReplicationKey is also copied during move/copy
+		OmitPendingNextOnCopy			= (1 << 1),	// PendingNext should never be set according to the data structure
+		MoveClientCachedStackCount		= (1 << 2), // Copy the ClientCachedStackCount in the move constructor
+		CleanupAllPendingActiveGEs		= (1 << 3), // Clean-up the entire Active Pending GE list on destruction
+
+		Current = MostRecentArrayReplicationKey | OmitPendingNextOnCopy | MoveClientCachedStackCount | CleanupAllPendingActiveGEs
+	};
+
+	int32 ActiveGameplayEffectReplicationFix = (int32)EActiveGameplayEffectFix::Current;
+	FAutoConsoleVariableRef CVarActiveGameplayEffectReplicationFix{ TEXT("AbilitySystem.Fix.ActiveGEReplicationFix"), ActiveGameplayEffectReplicationFix, TEXT("Experimental code mask for fixing Active Gameplay Effects (set to 0 to disable)"), ECVF_Default };
+	inline bool HasActiveGameplayEffectFix(EActiveGameplayEffectFix Flag) { return (ActiveGameplayEffectReplicationFix & static_cast<int32>(Flag)) != 0; }
+
 	TAutoConsoleVariable<int32> CVarGameplayEffectMaxVersion(TEXT("AbilitySystem.GameplayEffects.MaxVersion"), (int32)EGameplayEffectVersion::Current, TEXT("Override the Gameplay Effect Current Version (disabling upgrade code paths)"), ECVF_Default);
 
 #if WITH_EDITOR
@@ -2272,68 +2287,34 @@ void FGameplayEffectAttributeCaptureSpecContainer::SwapAggregator(FAggregatorRef
 //
 // --------------------------------------------------------------------------------------------------------------------------------------------------------
 
-FActiveGameplayEffect::FActiveGameplayEffect()
-	: StartServerWorldTime(0)
-	, CachedStartServerWorldTime(0)
-	, StartWorldTime(0.f)
-	, bIsInhibited(true)
-	, bPendingRepOnActiveGC(false)
-	, bPendingRepWhileActiveGC(false)
-	, IsPendingRemove(false)
-	, ClientCachedStackCount(0)
-	, PendingNext(nullptr)
+FActiveGameplayEffect::FActiveGameplayEffect(const FActiveGameplayEffect& Other)
 {
+	*this = Other;
 }
 
-FActiveGameplayEffect::FActiveGameplayEffect(const FActiveGameplayEffect& Other)
+FActiveGameplayEffect::FActiveGameplayEffect(FActiveGameplayEffect&& Other)
 {
-	*this = Other;
+	*this = MoveTemp(Other);
 }
 
-FActiveGameplayEffect::FActiveGameplayEffect(FActiveGameplayEffectHandle InHandle, const FGameplayEffectSpec &InSpec, float CurrentWorldTime, float InStartServerWorldTime, FPredictionKey InPredictionKey)
+FActiveGameplayEffect::FActiveGameplayEffect(FActiveGameplayEffectHandle InHandle, const FGameplayEffectSpec &InSpec, float InCurrentWorldTime, float InStartServerWorldTime, FPredictionKey InPredictionKey)
 	: Handle(InHandle)
 	, Spec(InSpec)
 	, PredictionKey(InPredictionKey)
 	, StartServerWorldTime(InStartServerWorldTime)
 	, CachedStartServerWorldTime(InStartServerWorldTime)
-	, StartWorldTime(CurrentWorldTime)
-	, bIsInhibited(true)
-	, bPendingRepOnActiveGC(false)
-	, bPendingRepWhileActiveGC(false)
-	, IsPendingRemove(false)
-	, ClientCachedStackCount(0)
-	, PendingNext(nullptr)
+	, StartWorldTime(InCurrentWorldTime)
 {
 }
 
-FActiveGameplayEffect::FActiveGameplayEffect(FActiveGameplayEffect&& Other)
-	: Handle(Other.Handle)
-	, Spec(MoveTemp(Other.Spec))
-	, PredictionKey(Other.PredictionKey)
-	, StartServerWorldTime(Other.StartServerWorldTime)
-	, CachedStartServerWorldTime(Other.CachedStartServerWorldTime)
-	, StartWorldTime(Other.StartWorldTime)
-	, bIsInhibited(Other.bIsInhibited)
-	, bPendingRepOnActiveGC(Other.bPendingRepOnActiveGC)
-	, bPendingRepWhileActiveGC(Other.bPendingRepWhileActiveGC)
-	, IsPendingRemove(Other.IsPendingRemove)
-	, ClientCachedStackCount(0)
-	, PeriodHandle(Other.PeriodHandle)
-	, DurationHandle(Other.DurationHandle)
-	, EventSet(Other.EventSet)
-{
-
-	ReplicationID = Other.ReplicationID;
-	ReplicationKey = Other.ReplicationKey;
-
-	// Note: purposefully not copying PendingNext pointer.
-}
-
 FActiveGameplayEffect& FActiveGameplayEffect::operator=(FActiveGameplayEffect&& Other)
 {
+	using namespace UE::GameplayEffect;
+
 	Handle = Other.Handle;
 	Spec = MoveTemp(Other.Spec);
 	PredictionKey = Other.PredictionKey;
+	GrantedAbilityHandles = MoveTemp(Other.GrantedAbilityHandles);
 	StartServerWorldTime = Other.StartServerWorldTime;
 	CachedStartServerWorldTime = Other.CachedStartServerWorldTime;
 	StartWorldTime = Other.StartWorldTime;
@@ -2341,14 +2322,22 @@ FActiveGameplayEffect& FActiveGameplayEffect::operator=(FActiveGameplayEffect&&
 	bPendingRepOnActiveGC = Other.bPendingRepOnActiveGC;
 	bPendingRepWhileActiveGC = Other.bPendingRepWhileActiveGC;
 	IsPendingRemove = Other.IsPendingRemove;
-	ClientCachedStackCount = Other.ClientCachedStackCount;
+	ClientCachedStackCount = HasActiveGameplayEffectFix(EActiveGameplayEffectFix::MoveClientCachedStackCount) ? Other.ClientCachedStackCount : 0;
 	PeriodHandle = Other.PeriodHandle;
 	DurationHandle = Other.DurationHandle;
-	EventSet = Other.EventSet;
 	// Note: purposefully not copying PendingNext pointer.
+	// PendingNext = Other.PendingNext;
+	EventSet = Other.EventSet;
 
+	// FFastArraySerializerItem properties
 	ReplicationID = Other.ReplicationID;
 	ReplicationKey = Other.ReplicationKey;
+
+	if (HasActiveGameplayEffectFix(EActiveGameplayEffectFix::MostRecentArrayReplicationKey))
+	{
+		MostRecentArrayReplicationKey = Other.MostRecentArrayReplicationKey;
+	}
+
 	return *this;
 }
 
@@ -2357,6 +2346,7 @@ FActiveGameplayEffect& FActiveGameplayEffect::operator=(const FActiveGameplayEff
 	Handle = Other.Handle;
 	Spec = Other.Spec;
 	PredictionKey = Other.PredictionKey;
+	GrantedAbilityHandles = Other.GrantedAbilityHandles;
 	StartServerWorldTime = Other.StartServerWorldTime;
 	CachedStartServerWorldTime = Other.CachedStartServerWorldTime;
 	StartWorldTime = Other.StartWorldTime;
@@ -2368,10 +2358,23 @@ FActiveGameplayEffect& FActiveGameplayEffect::operator=(const FActiveGameplayEff
 	PeriodHandle = Other.PeriodHandle;
 	DurationHandle = Other.DurationHandle;
 	EventSet = Other.EventSet;
-	PendingNext = Other.PendingNext;
 
+	// Note: purposefully not copying PendingNext pointer unless the fix is disabled
+	using namespace UE::GameplayEffect;
+	if (!HasActiveGameplayEffectFix(EActiveGameplayEffectFix::OmitPendingNextOnCopy))
+	{
+		PendingNext = Other.PendingNext;
+	}
+
+	// FFastArraySerializerItem properties
 	ReplicationID = Other.ReplicationID;
 	ReplicationKey = Other.ReplicationKey;
+
+	if (HasActiveGameplayEffectFix(EActiveGameplayEffectFix::MostRecentArrayReplicationKey))
+	{
+		MostRecentArrayReplicationKey = Other.MostRecentArrayReplicationKey;
+	}
+
 	return *this;
 }
 
@@ -2562,12 +2565,27 @@ FActiveGameplayEffectsContainer::FActiveGameplayEffectsContainer()
 PRAGMA_DISABLE_DEPRECATION_WARNINGS
 FActiveGameplayEffectsContainer::~FActiveGameplayEffectsContainer()
 {
+	using namespace UE::GameplayEffect;
+
 	FActiveGameplayEffect* PendingGameplayEffect = PendingGameplayEffectHead;
-	if (PendingGameplayEffectHead)
+
+	if (HasActiveGameplayEffectFix(EActiveGameplayEffectFix::CleanupAllPendingActiveGEs))
 	{
-		FActiveGameplayEffect* Next = PendingGameplayEffectHead->PendingNext;
-		delete PendingGameplayEffectHead;
-		PendingGameplayEffectHead =	Next;
+		while (PendingGameplayEffectHead)
+		{
+			FActiveGameplayEffect* Next = PendingGameplayEffectHead->PendingNext;
+			delete PendingGameplayEffectHead;
+			PendingGameplayEffectHead = Next;
+		}
+	}
+	else
+	{
+		if (PendingGameplayEffectHead)
+		{
+			FActiveGameplayEffect* Next = PendingGameplayEffectHead->PendingNext;
+			delete PendingGameplayEffectHead;
+			PendingGameplayEffectHead = Next;
+		}
 	}
 }
 PRAGMA_ENABLE_DEPRECATION_WARNINGS
@@ -3726,6 +3744,8 @@ PRAGMA_ENABLE_DEPRECATION_WARNINGS
 			 */
 
 			check(PendingGameplayEffectNext);
+			const FActiveGameplayEffect* PreviousPendingNext = (*PendingGameplayEffectNext) ? (*PendingGameplayEffectNext)->PendingNext : nullptr;
+
 			if (*PendingGameplayEffectNext == nullptr)
 			{
 				// We have no memory allocated to put our next pending GE, so make a new one.
@@ -3741,6 +3761,9 @@ PRAGMA_ENABLE_DEPRECATION_WARNINGS
 				AppliedActiveGE = *PendingGameplayEffectNext;
 			}
 
+			// Let's check that our Pending Active GE Chain is still intact. If this triggers, the code is wrong, not the asset.
+			ensureMsgf(AppliedActiveGE->PendingNext == PreviousPendingNext, TEXT("ApplyGameplayEffectSpec Code Leaked a Pending FActiveGameplayEffect while applying %s"), *AppliedActiveGE->Spec.ToSimpleString());
+
 			// The next pending GameplayEffect goes to where our PendingNext points
 			PendingGameplayEffectNext = &AppliedActiveGE->PendingNext;
 		}
diff --git a/Engine/Plugins/Runtime/GameplayAbilities/Source/GameplayAbilities/Public/GameplayEffect.h b/Engine/Plugins/Runtime/GameplayAbilities/Source/GameplayAbilities/Public/GameplayEffect.h
--- a/Engine/Plugins/Runtime/GameplayAbilities/Source/GameplayAbilities/Public/GameplayEffect.h
+++ b/Engine/Plugins/Runtime/GameplayAbilities/Source/GameplayAbilities/Public/GameplayEffect.h
@@ -1286,16 +1286,13 @@ struct GAMEPLAYABILITIES_API FActiveGameplayEffect : public FFastArraySerializer
 	//  IMPORTANT: Any new state added to FActiveGameplayEffect must be handled in the copy/move constructor/operator
 	// ---------------------------------------------------------------------------------------------------------------------------------
 
-	FActiveGameplayEffect();
+	FActiveGameplayEffect() = default;
+	FActiveGameplayEffect(FActiveGameplayEffectHandle InHandle, const FGameplayEffectSpec& InSpec, float CurrentWorldTime, float InStartServerWorldTime, FPredictionKey InPredictionKey);
 
+	// These need to be defined because we need to omit PendingNext from move operations and the base class isn't trivially copyable
 	FActiveGameplayEffect(const FActiveGameplayEffect& Other);
-
-	FActiveGameplayEffect(FActiveGameplayEffectHandle InHandle, const FGameplayEffectSpec &InSpec, float CurrentWorldTime, float InStartServerWorldTime, FPredictionKey InPredictionKey);
-	
 	FActiveGameplayEffect(FActiveGameplayEffect&& Other);
-
 	FActiveGameplayEffect& operator=(FActiveGameplayEffect&& other);
-
 	FActiveGameplayEffect& operator=(const FActiveGameplayEffect& other);
 
 	float GetTimeRemaining(float WorldTime) const
@@ -1365,32 +1362,33 @@ struct GAMEPLAYABILITIES_API FActiveGameplayEffect : public FFastArraySerializer
 
 	/** Server time this started */
 	UPROPERTY()
-	float StartServerWorldTime;
+	float StartServerWorldTime = 0.0f;
 
 	/** Used for handling duration modifications being replicated */
 	UPROPERTY(NotReplicated)
-	float CachedStartServerWorldTime;
+	float CachedStartServerWorldTime = 0.0f;
 
 	UPROPERTY(NotReplicated)
-	float StartWorldTime;
+	float StartWorldTime = 0.0f;
 
 	// Not sure if this should replicate or not. If replicated, we may have trouble where IsInhibited doesn't appear to change when we do tag checks (because it was previously inhibited, but replication made it inhibited).
 	UPROPERTY(NotReplicated)
-	bool bIsInhibited;
+	bool bIsInhibited = true;
 
 	/** When replicated down, we cue the GC events until the entire list of active gameplay effects has been received */
-	mutable bool bPendingRepOnActiveGC;
-	mutable bool bPendingRepWhileActiveGC;
+	mutable bool bPendingRepOnActiveGC = false;
+	mutable bool bPendingRepWhileActiveGC = false;
 
-	bool IsPendingRemove;
+	bool IsPendingRemove = false;
 
 	/** Last StackCount that the client had. Used to tell if the stackcount has changed in PostReplicatedChange */
-	int32 ClientCachedStackCount;
+	int32 ClientCachedStackCount = 0;
 
 	FTimerHandle PeriodHandle;
 	FTimerHandle DurationHandle;
 
-	FActiveGameplayEffect* PendingNext;
+	/** Cached pointer.  Since these ActiveGE's can be reused in-place, this should *not* be copied during copy/move operations */
+	FActiveGameplayEffect* PendingNext = nullptr;
 	
 	/** All the bindable events for this active effect (bundled to allow easier non-const access to these events via the ASC) */
 	FActiveGameplayEffectEvents EventSet;
