Added support to set destination reached distance to UAbilityTask_ApplyRootMotionMoveToActorForce.

diff --git a/Engine/Plugins/Runtime/GameplayAbilities/Source/GameplayAbilities/Private/Abilities/Tasks/AbilityTask_ApplyRootMotionMoveToActorForce.cpp b/Engine/Plugins/Runtime/GameplayAbilities/Source/GameplayAbilities/Private/Abilities/Tasks/AbilityTask_ApplyRootMotionMoveToActorForce.cpp
--- a/Engine/Plugins/Runtime/GameplayAbilities/Source/GameplayAbilities/Private/Abilities/Tasks/AbilityTask_ApplyRootMotionMoveToActorForce.cpp
+++ b/Engine/Plugins/Runtime/GameplayAbilities/Source/GameplayAbilities/Private/Abilities/Tasks/AbilityTask_ApplyRootMotionMoveToActorForce.cpp
@@ -30,6 +30,7 @@ UAbilityTask_ApplyRootMotionMoveToActorForce::UAbilityTask_ApplyRootMotionMoveTo
 : Super(ObjectInitializer)
 {
 	bDisableDestinationReachedInterrupt = false;
+	ReachedDestinationDistance = 50.0f;
 	bSetNewMovementMode = false;
 	NewMovementMode = EMovementMode::MOVE_Walking;
 	PreviousMovementMode = EMovementMode::MOVE_None;
@@ -51,7 +52,7 @@ void UAbilityTask_ApplyRootMotionMoveToActorForce::OnTargetActorSwapped(class AA
 	}
 }
 
-UAbilityTask_ApplyRootMotionMoveToActorForce* UAbilityTask_ApplyRootMotionMoveToActorForce::ApplyRootMotionMoveToActorForce(UGameplayAbility* OwningAbility, FName TaskInstanceName, AActor* TargetActor, FVector TargetLocationOffset, ERootMotionMoveToActorTargetOffsetType OffsetAlignment, float Duration, UCurveFloat* TargetLerpSpeedHorizontal, UCurveFloat* TargetLerpSpeedVertical, bool bSetNewMovementMode, EMovementMode MovementMode, bool bRestrictSpeedToExpected, UCurveVector* PathOffsetCurve, UCurveFloat* TimeMappingCurve, ERootMotionFinishVelocityMode VelocityOnFinishMode, FVector SetVelocityOnFinish, float ClampVelocityOnFinish, bool bDisableDestinationReachedInterrupt)
+UAbilityTask_ApplyRootMotionMoveToActorForce* UAbilityTask_ApplyRootMotionMoveToActorForce::ApplyRootMotionMoveToActorForce(UGameplayAbility* OwningAbility, FName TaskInstanceName, AActor* TargetActor, FVector TargetLocationOffset, ERootMotionMoveToActorTargetOffsetType OffsetAlignment, float Duration, UCurveFloat* TargetLerpSpeedHorizontal, UCurveFloat* TargetLerpSpeedVertical, bool bSetNewMovementMode, EMovementMode MovementMode, bool bRestrictSpeedToExpected, UCurveVector* PathOffsetCurve, UCurveFloat* TimeMappingCurve, ERootMotionFinishVelocityMode VelocityOnFinishMode, FVector SetVelocityOnFinish, float ClampVelocityOnFinish, bool bDisableDestinationReachedInterrupt, float ReachedDestinationDistance/* = 50.0f*/)
 {
 	UAbilityTask_ApplyRootMotionMoveToActorForce* MyTask = NewAbilityTask<UAbilityTask_ApplyRootMotionMoveToActorForce>(OwningAbility, TaskInstanceName);
 
@@ -63,6 +64,7 @@ UAbilityTask_ApplyRootMotionMoveToActorForce* UAbilityTask_ApplyRootMotionMoveTo
 	MyTask->OffsetAlignment = OffsetAlignment;
 	MyTask->Duration = FMath::Max(Duration, KINDA_SMALL_NUMBER); // Avoid negative or divide-by-zero cases
 	MyTask->bDisableDestinationReachedInterrupt = bDisableDestinationReachedInterrupt;
+	MyTask->ReachedDestinationDistance = ReachedDestinationDistance;
 	MyTask->TargetLerpSpeedHorizontalCurve = TargetLerpSpeedHorizontal;
 	MyTask->TargetLerpSpeedVerticalCurve = TargetLerpSpeedVertical;
 	MyTask->bSetNewMovementMode = bSetNewMovementMode;
@@ -288,7 +290,7 @@ void UAbilityTask_ApplyRootMotionMoveToActorForce::TickTask(float DeltaTime)
 		}
 #endif // !(UE_BUILD_SHIPPING || UE_BUILD_TEST)
 
-		const float ReachedDestinationDistanceSqr = 50.f * 50.f;
+		const float ReachedDestinationDistanceSqr = ReachedDestinationDistance * ReachedDestinationDistance;
 		const bool bReachedDestination = FVector::DistSquared(TargetLocation, MyActor->GetActorLocation()) < ReachedDestinationDistanceSqr;
 
 		if (bTimedOut || (bReachedDestination && !bDisableDestinationReachedInterrupt))
@@ -324,6 +326,7 @@ void UAbilityTask_ApplyRootMotionMoveToActorForce::GetLifetimeReplicatedProps(TA
 	DOREPLIFETIME(UAbilityTask_ApplyRootMotionMoveToActorForce, OffsetAlignment);
 	DOREPLIFETIME(UAbilityTask_ApplyRootMotionMoveToActorForce, Duration);
 	DOREPLIFETIME(UAbilityTask_ApplyRootMotionMoveToActorForce, bDisableDestinationReachedInterrupt);
+	DOREPLIFETIME(UAbilityTask_ApplyRootMotionMoveToActorForce, ReachedDestinationDistance);
 	DOREPLIFETIME(UAbilityTask_ApplyRootMotionMoveToActorForce, TargetLerpSpeedHorizontalCurve);
 	DOREPLIFETIME(UAbilityTask_ApplyRootMotionMoveToActorForce, TargetLerpSpeedVerticalCurve);
 	DOREPLIFETIME(UAbilityTask_ApplyRootMotionMoveToActorForce, bSetNewMovementMode);
@@ -359,7 +362,7 @@ void UAbilityTask_ApplyRootMotionMoveToActorForce::OnDestroy(bool AbilityIsEndin
 	Super::OnDestroy(AbilityIsEnding);
 }
 
-UAbilityTask_ApplyRootMotionMoveToActorForce* UAbilityTask_ApplyRootMotionMoveToActorForce::ApplyRootMotionMoveToTargetDataActorForce(UGameplayAbility* OwningAbility, FName TaskInstanceName, FGameplayAbilityTargetDataHandle TargetDataHandle, int32 TargetDataIndex, int32 TargetActorIndex, FVector TargetLocationOffset, ERootMotionMoveToActorTargetOffsetType OffsetAlignment, float Duration, UCurveFloat* TargetLerpSpeedHorizontal, UCurveFloat* TargetLerpSpeedVertical, bool bSetNewMovementMode, EMovementMode MovementMode, bool bRestrictSpeedToExpected, UCurveVector* PathOffsetCurve, UCurveFloat* TimeMappingCurve, ERootMotionFinishVelocityMode VelocityOnFinishMode, FVector SetVelocityOnFinish, float ClampVelocityOnFinish, bool bDisableDestinationReachedInterrupt)
+UAbilityTask_ApplyRootMotionMoveToActorForce* UAbilityTask_ApplyRootMotionMoveToActorForce::ApplyRootMotionMoveToTargetDataActorForce(UGameplayAbility* OwningAbility, FName TaskInstanceName, FGameplayAbilityTargetDataHandle TargetDataHandle, int32 TargetDataIndex, int32 TargetActorIndex, FVector TargetLocationOffset, ERootMotionMoveToActorTargetOffsetType OffsetAlignment, float Duration, UCurveFloat* TargetLerpSpeedHorizontal, UCurveFloat* TargetLerpSpeedVertical, bool bSetNewMovementMode, EMovementMode MovementMode, bool bRestrictSpeedToExpected, UCurveVector* PathOffsetCurve, UCurveFloat* TimeMappingCurve, ERootMotionFinishVelocityMode VelocityOnFinishMode, FVector SetVelocityOnFinish, float ClampVelocityOnFinish, bool bDisableDestinationReachedInterrupt, float ReachedDestinationDistance/* = 50.0f*/)
 {
 	if (TargetDataIndex >= 0 && TargetDataIndex < TargetDataHandle.Num())
 	{
@@ -374,7 +377,7 @@ UAbilityTask_ApplyRootMotionMoveToActorForce* UAbilityTask_ApplyRootMotionMoveTo
 				if (Actors[TargetActorIndex].IsValid())
 				{
 					AActor* TargetActor = Actors[TargetActorIndex].Get();
-					UAbilityTask_ApplyRootMotionMoveToActorForce* retVal = ApplyRootMotionMoveToActorForce(OwningAbility, TaskInstanceName, TargetActor, TargetLocationOffset, OffsetAlignment, Duration, TargetLerpSpeedHorizontal, TargetLerpSpeedVertical, bSetNewMovementMode, MovementMode, bRestrictSpeedToExpected, PathOffsetCurve, TimeMappingCurve, VelocityOnFinishMode, SetVelocityOnFinish, ClampVelocityOnFinish, bDisableDestinationReachedInterrupt);
+					UAbilityTask_ApplyRootMotionMoveToActorForce* retVal = ApplyRootMotionMoveToActorForce(OwningAbility, TaskInstanceName, TargetActor, TargetLocationOffset, OffsetAlignment, Duration, TargetLerpSpeedHorizontal, TargetLerpSpeedVertical, bSetNewMovementMode, MovementMode, bRestrictSpeedToExpected, PathOffsetCurve, TimeMappingCurve, VelocityOnFinishMode, SetVelocityOnFinish, ClampVelocityOnFinish, bDisableDestinationReachedInterrupt, ReachedDestinationDistance);
 
 					if (TargetData->ShouldCheckForTargetActorSwap())
 					{
diff --git a/Engine/Plugins/Runtime/GameplayAbilities/Source/GameplayAbilities/Public/Abilities/Tasks/AbilityTask_ApplyRootMotionMoveToActorForce.h b/Engine/Plugins/Runtime/GameplayAbilities/Source/GameplayAbilities/Public/Abilities/Tasks/AbilityTask_ApplyRootMotionMoveToActorForce.h
--- a/Engine/Plugins/Runtime/GameplayAbilities/Source/GameplayAbilities/Public/Abilities/Tasks/AbilityTask_ApplyRootMotionMoveToActorForce.h
+++ b/Engine/Plugins/Runtime/GameplayAbilities/Source/GameplayAbilities/Public/Abilities/Tasks/AbilityTask_ApplyRootMotionMoveToActorForce.h
@@ -43,11 +43,11 @@ class GAMEPLAYABILITIES_API UAbilityTask_ApplyRootMotionMoveToActorForce : publi
 
 	/** Apply force to character's movement */
 	UFUNCTION(BlueprintCallable, Category = "Ability|Tasks", meta = (HidePin = "OwningAbility", DefaultToSelf = "OwningAbility", BlueprintInternalUseOnly = "TRUE"))
-	static UAbilityTask_ApplyRootMotionMoveToActorForce* ApplyRootMotionMoveToActorForce(UGameplayAbility* OwningAbility, FName TaskInstanceName, AActor* TargetActor, FVector TargetLocationOffset, ERootMotionMoveToActorTargetOffsetType OffsetAlignment, float Duration, UCurveFloat* TargetLerpSpeedHorizontal, UCurveFloat* TargetLerpSpeedVertical, bool bSetNewMovementMode, EMovementMode MovementMode, bool bRestrictSpeedToExpected, UCurveVector* PathOffsetCurve, UCurveFloat* TimeMappingCurve, ERootMotionFinishVelocityMode VelocityOnFinishMode, FVector SetVelocityOnFinish, float ClampVelocityOnFinish, bool bDisableDestinationReachedInterrupt);
+	static UAbilityTask_ApplyRootMotionMoveToActorForce* ApplyRootMotionMoveToActorForce(UGameplayAbility* OwningAbility, FName TaskInstanceName, AActor* TargetActor, FVector TargetLocationOffset, ERootMotionMoveToActorTargetOffsetType OffsetAlignment, float Duration, UCurveFloat* TargetLerpSpeedHorizontal, UCurveFloat* TargetLerpSpeedVertical, bool bSetNewMovementMode, EMovementMode MovementMode, bool bRestrictSpeedToExpected, UCurveVector* PathOffsetCurve, UCurveFloat* TimeMappingCurve, ERootMotionFinishVelocityMode VelocityOnFinishMode, FVector SetVelocityOnFinish, float ClampVelocityOnFinish, bool bDisableDestinationReachedInterrupt, float ReachedDestinationDistance = 50.0f);
 
 	/** Apply force to character's movement using an index into targetData instead of using an actor directly. */
 	UFUNCTION(BlueprintCallable, Category = "Ability|Tasks", meta = (HidePin = "OwningAbility", DefaultToSelf = "OwningAbility", BlueprintInternalUseOnly = "TRUE"))
-	static UAbilityTask_ApplyRootMotionMoveToActorForce* ApplyRootMotionMoveToTargetDataActorForce(UGameplayAbility* OwningAbility, FName TaskInstanceName, FGameplayAbilityTargetDataHandle TargetDataHandle, int32 TargetDataIndex, int32 TargetActorIndex, FVector TargetLocationOffset, ERootMotionMoveToActorTargetOffsetType OffsetAlignment, float Duration, UCurveFloat* TargetLerpSpeedHorizontal, UCurveFloat* TargetLerpSpeedVertical, bool bSetNewMovementMode, EMovementMode MovementMode, bool bRestrictSpeedToExpected, UCurveVector* PathOffsetCurve, UCurveFloat* TimeMappingCurve, ERootMotionFinishVelocityMode VelocityOnFinishMode, FVector SetVelocityOnFinish, float ClampVelocityOnFinish, bool bDisableDestinationReachedInterrupt);
+	static UAbilityTask_ApplyRootMotionMoveToActorForce* ApplyRootMotionMoveToTargetDataActorForce(UGameplayAbility* OwningAbility, FName TaskInstanceName, FGameplayAbilityTargetDataHandle TargetDataHandle, int32 TargetDataIndex, int32 TargetActorIndex, FVector TargetLocationOffset, ERootMotionMoveToActorTargetOffsetType OffsetAlignment, float Duration, UCurveFloat* TargetLerpSpeedHorizontal, UCurveFloat* TargetLerpSpeedVertical, bool bSetNewMovementMode, EMovementMode MovementMode, bool bRestrictSpeedToExpected, UCurveVector* PathOffsetCurve, UCurveFloat* TimeMappingCurve, ERootMotionFinishVelocityMode VelocityOnFinishMode, FVector SetVelocityOnFinish, float ClampVelocityOnFinish, bool bDisableDestinationReachedInterrupt, float ReachedDestinationDistance = 50.0f);
 
 	/** Tick function for this task, if bTickingTask == true */
 	virtual void TickTask(float DeltaTime) override;
@@ -98,6 +98,10 @@ protected:
 	UPROPERTY(Replicated)
 	bool bDisableDestinationReachedInterrupt;
 
+	/** This is the distance where, at or less that this distance to the destination, the user is considered to be at the destination. */
+	UPROPERTY(Replicated)
+	float ReachedDestinationDistance;
+
 	UPROPERTY(Replicated)
 	bool bSetNewMovementMode;
 
