Added support to AbilityTask_ApplyRootMotionMoveToActorForce to target a component with a relative location.

diff --git a/Engine/Plugins/Runtime/GameplayAbilities/Source/GameplayAbilities/Private/Abilities/Tasks/AbilityTask_ApplyRootMotionMoveToActorForce.cpp b/Engine/Plugins/Runtime/GameplayAbilities/Source/GameplayAbilities/Private/Abilities/Tasks/AbilityTask_ApplyRootMotionMoveToActorForce.cpp
--- a/Engine/Plugins/Runtime/GameplayAbilities/Source/GameplayAbilities/Private/Abilities/Tasks/AbilityTask_ApplyRootMotionMoveToActorForce.cpp
+++ b/Engine/Plugins/Runtime/GameplayAbilities/Source/GameplayAbilities/Private/Abilities/Tasks/AbilityTask_ApplyRootMotionMoveToActorForce.cpp
@@ -49,6 +49,10 @@ void UAbilityTask_ApplyRootMotionMoveToActorForce::OnTargetActorSwapped(class AA
 	if (OriginalTarget && OriginalTarget == TargetActor)
 	{
 		TargetActor = NewTarget;
+
+		//This is called from ApplyRootMotionMoveToTargetDataActorForce, which doesn't use TargetComponent for now, so just clearing it out for consistency
+		TargetComponent = nullptr;
+		TargetComponentRelativeLocation = FVector::ZeroVector;
 	}
 }
 
@@ -60,6 +64,8 @@ UAbilityTask_ApplyRootMotionMoveToActorForce* UAbilityTask_ApplyRootMotionMoveTo
 
 	MyTask->ForceName = TaskInstanceName;
 	MyTask->TargetActor = TargetActor;
+	MyTask->TargetComponent = nullptr;
+	MyTask->TargetComponentRelativeLocation = FVector::ZeroVector;
 	MyTask->TargetLocationOffset = TargetLocationOffset;
 	MyTask->OffsetAlignment = OffsetAlignment;
 	MyTask->Duration = FMath::Max(Duration, KINDA_SMALL_NUMBER); // Avoid negative or divide-by-zero cases
@@ -89,6 +95,45 @@ UAbilityTask_ApplyRootMotionMoveToActorForce* UAbilityTask_ApplyRootMotionMoveTo
 	return MyTask;
 }
 
+UAbilityTask_ApplyRootMotionMoveToActorForce* UAbilityTask_ApplyRootMotionMoveToActorForce::ApplyRootMotionMoveToComponentForce(UGameplayAbility* OwningAbility, FName TaskInstanceName, USceneComponent* TargetComponent, FVector TargetComponentRelativeLocation, FVector TargetLocationOffset, ERootMotionMoveToActorTargetOffsetType OffsetAlignment, float Duration, UCurveFloat* TargetLerpSpeedHorizontal, UCurveFloat* TargetLerpSpeedVertical, bool bSetNewMovementMode, EMovementMode MovementMode, bool bRestrictSpeedToExpected, UCurveVector* PathOffsetCurve, UCurveFloat* TimeMappingCurve, ERootMotionFinishVelocityMode VelocityOnFinishMode, FVector SetVelocityOnFinish, float ClampVelocityOnFinish, bool bDisableDestinationReachedInterrupt, float ReachedDestinationDistance)
+{
+	UAbilityTask_ApplyRootMotionMoveToActorForce* MyTask = NewAbilityTask<UAbilityTask_ApplyRootMotionMoveToActorForce>(OwningAbility, TaskInstanceName);
+
+	UAbilitySystemGlobals::NonShipping_ApplyGlobalAbilityScaler_Duration(Duration);
+
+	MyTask->ForceName = TaskInstanceName;
+	MyTask->TargetComponent = TargetComponent;
+	MyTask->TargetActor = TargetComponent ? TargetComponent->GetOwner() : nullptr;
+	MyTask->TargetComponentRelativeLocation = TargetComponentRelativeLocation;
+	MyTask->TargetLocationOffset = TargetLocationOffset;
+	MyTask->OffsetAlignment = OffsetAlignment;
+	MyTask->Duration = FMath::Max(Duration, KINDA_SMALL_NUMBER); // Avoid negative or divide-by-zero cases
+	MyTask->bDisableDestinationReachedInterrupt = bDisableDestinationReachedInterrupt;
+	MyTask->ReachedDestinationDistance = ReachedDestinationDistance;
+	MyTask->TargetLerpSpeedHorizontalCurve = TargetLerpSpeedHorizontal;
+	MyTask->TargetLerpSpeedVerticalCurve = TargetLerpSpeedVertical;
+	MyTask->bSetNewMovementMode = bSetNewMovementMode;
+	MyTask->NewMovementMode = MovementMode;
+	MyTask->bRestrictSpeedToExpected = bRestrictSpeedToExpected;
+	MyTask->PathOffsetCurve = PathOffsetCurve;
+	MyTask->TimeMappingCurve = TimeMappingCurve;
+	MyTask->FinishVelocityMode = VelocityOnFinishMode;
+	MyTask->FinishSetVelocity = SetVelocityOnFinish;
+	MyTask->FinishClampVelocity = ClampVelocityOnFinish;
+	if (MyTask->GetAvatarActor() != nullptr)
+	{
+		MyTask->StartLocation = MyTask->GetAvatarActor()->GetActorLocation();
+	}
+	else
+	{
+		checkf(false, TEXT("UAbilityTask_ApplyRootMotionMoveToActorForce called without valid avatar actor to get start location from."));
+		MyTask->StartLocation = TargetComponent ? TargetComponent->GetComponentTransform().TransformPosition(TargetComponentRelativeLocation) : FVector(0.f);
+	}
+	MyTask->SharedInitAndApply();
+
+	return MyTask;
+}
+
 void UAbilityTask_ApplyRootMotionMoveToActorForce::OnRep_TargetLocation()
 {
 	if (bIsSimulating)
@@ -160,7 +205,9 @@ FVector UAbilityTask_ApplyRootMotionMoveToActorForce::CalculateTargetOffset() co
 {
 	check(TargetActor != nullptr);
 
-	const FVector TargetActorLocation = TargetActor->GetActorLocation();
+	const FVector TargetActorLocation = TargetComponent ?
+											TargetComponent->GetComponentTransform().TransformPosition(TargetComponentRelativeLocation) :
+											TargetActor->GetActorLocation();
 	FVector CalculatedTargetLocation = TargetActorLocation;
 	
 	if (OffsetAlignment == ERootMotionMoveToActorTargetOffsetType::AlignFromTargetToSource)
diff --git a/Engine/Plugins/Runtime/GameplayAbilities/Source/GameplayAbilities/Public/Abilities/Tasks/AbilityTask_ApplyRootMotionMoveToActorForce.h b/Engine/Plugins/Runtime/GameplayAbilities/Source/GameplayAbilities/Public/Abilities/Tasks/AbilityTask_ApplyRootMotionMoveToActorForce.h
--- a/Engine/Plugins/Runtime/GameplayAbilities/Source/GameplayAbilities/Public/Abilities/Tasks/AbilityTask_ApplyRootMotionMoveToActorForce.h
+++ b/Engine/Plugins/Runtime/GameplayAbilities/Source/GameplayAbilities/Public/Abilities/Tasks/AbilityTask_ApplyRootMotionMoveToActorForce.h
@@ -41,10 +41,14 @@ class GAMEPLAYABILITIES_API UAbilityTask_ApplyRootMotionMoveToActorForce : publi
 	UPROPERTY(BlueprintAssignable)
 	FApplyRootMotionMoveToActorForceDelegate OnFinished;
 
-	/** Apply force to character's movement */
+	/** Apply force to character's movement to move to a target actor */
 	UFUNCTION(BlueprintCallable, Category = "Ability|Tasks", meta = (HidePin = "OwningAbility", DefaultToSelf = "OwningAbility", BlueprintInternalUseOnly = "TRUE"))
 	static UAbilityTask_ApplyRootMotionMoveToActorForce* ApplyRootMotionMoveToActorForce(UGameplayAbility* OwningAbility, FName TaskInstanceName, AActor* TargetActor, FVector TargetLocationOffset, ERootMotionMoveToActorTargetOffsetType OffsetAlignment, float Duration, UCurveFloat* TargetLerpSpeedHorizontal, UCurveFloat* TargetLerpSpeedVertical, bool bSetNewMovementMode, EMovementMode MovementMode, bool bRestrictSpeedToExpected, UCurveVector* PathOffsetCurve, UCurveFloat* TimeMappingCurve, ERootMotionFinishVelocityMode VelocityOnFinishMode, FVector SetVelocityOnFinish, float ClampVelocityOnFinish, bool bDisableDestinationReachedInterrupt, float ReachedDestinationDistance = 50.0f);
 
+	/** Apply force to character's movement to move to a target component with relative location*/
+	UFUNCTION(BlueprintCallable, Category = "Ability|Tasks", meta = (HidePin = "OwningAbility", DefaultToSelf = "OwningAbility", BlueprintInternalUseOnly = "TRUE"))
+	static UAbilityTask_ApplyRootMotionMoveToActorForce* ApplyRootMotionMoveToComponentForce(UGameplayAbility* OwningAbility, FName TaskInstanceName, USceneComponent* TargetComponent, FVector TargetComponentRelativeLocation, FVector TargetLocationOffset, ERootMotionMoveToActorTargetOffsetType OffsetAlignment, float Duration, UCurveFloat* TargetLerpSpeedHorizontal, UCurveFloat* TargetLerpSpeedVertical, bool bSetNewMovementMode, EMovementMode MovementMode, bool bRestrictSpeedToExpected, UCurveVector* PathOffsetCurve, UCurveFloat* TimeMappingCurve, ERootMotionFinishVelocityMode VelocityOnFinishMode, FVector SetVelocityOnFinish, float ClampVelocityOnFinish, bool bDisableDestinationReachedInterrupt, float ReachedDestinationDistance = 50.0f);
+	
 	/** Apply force to character's movement using an index into targetData instead of using an actor directly. */
 	UFUNCTION(BlueprintCallable, Category = "Ability|Tasks", meta = (HidePin = "OwningAbility", DefaultToSelf = "OwningAbility", BlueprintInternalUseOnly = "TRUE"))
 	static UAbilityTask_ApplyRootMotionMoveToActorForce* ApplyRootMotionMoveToTargetDataActorForce(UGameplayAbility* OwningAbility, FName TaskInstanceName, FGameplayAbilityTargetDataHandle TargetDataHandle, int32 TargetDataIndex, int32 TargetActorIndex, FVector TargetLocationOffset, ERootMotionMoveToActorTargetOffsetType OffsetAlignment, float Duration, UCurveFloat* TargetLerpSpeedHorizontal, UCurveFloat* TargetLerpSpeedVertical, bool bSetNewMovementMode, EMovementMode MovementMode, bool bRestrictSpeedToExpected, UCurveVector* PathOffsetCurve, UCurveFloat* TimeMappingCurve, ERootMotionFinishVelocityMode VelocityOnFinishMode, FVector SetVelocityOnFinish, float ClampVelocityOnFinish, bool bDisableDestinationReachedInterrupt, float ReachedDestinationDistance = 50.0f);
@@ -84,6 +88,12 @@ protected:
 	UPROPERTY(Replicated)
 	TObjectPtr<AActor> TargetActor;
 
+	UPROPERTY(Replicated)
+	TObjectPtr<USceneComponent> TargetComponent = nullptr;
+
+	UPROPERTY(Replicated)
+	FVector TargetComponentRelativeLocation = FVector::ZeroVector;
+
 	UPROPERTY(Replicated)
 	FVector TargetLocationOffset;
 
