Adding more data in the Gameplay Ability Audit

diff --git a/Engine/Plugins/Runtime/GameplayAbilities/Source/GameplayAbilitiesEditor/Private/GameplayAbilityAudit.cpp b/Engine/Plugins/Runtime/GameplayAbilities/Source/GameplayAbilitiesEditor/Private/GameplayAbilityAudit.cpp
--- a/Engine/Plugins/Runtime/GameplayAbilities/Source/GameplayAbilitiesEditor/Private/GameplayAbilityAudit.cpp
+++ b/Engine/Plugins/Runtime/GameplayAbilities/Source/GameplayAbilitiesEditor/Private/GameplayAbilityAudit.cpp
@@ -90,6 +90,32 @@ void FGameplayAbilityAuditRow::FillDataFromGameplayAbility(const UGameplayAbilit
 		AuditRow.CooldownGE = TheCooldownGE->GetClass()->GetFName();
 	}
 
+	// Gather all of the other GameplayTagContainer referenced Tags (will also include AssetTags).
+	FProperty* AbilityTagsProperty = UGameplayAbility::StaticClass()->FindPropertyByName(GET_MEMBER_NAME_CHECKED(UGameplayAbility, AbilityTags));
+	for (TPropertyValueIterator<FStructProperty> It(GameplayAbility.GetClass(), &GameplayAbility); It; ++It)
+	{
+		const bool bIsTagContainer = It.Key()->SameType(AbilityTagsProperty);
+		if (!bIsTagContainer)
+		{
+			continue;
+		}
+
+		if (const FGameplayTagContainer* TagContainer = reinterpret_cast<const FGameplayTagContainer*>(It.Value()))
+		{
+			for (const FGameplayTag& GameplayTag : *TagContainer)
+			{
+				AuditRow.ReferencedTags.Emplace(GameplayTag.GetTagName());
+			}
+		}
+	}
+
+	// Asset Tags should go in its own field because of its importance
+	for (const FGameplayTag& GameplayTag : GameplayAbility.AbilityTags)
+	{
+		AuditRow.AssetTags.Emplace(GameplayTag.GetTagName());
+	}
+
+	AuditRow.InstancingPolicy = GameplayAbility.GetInstancingPolicy();
 	AuditRow.NetExecutionPolicy = GameplayAbility.GetNetExecutionPolicy();
 	AuditRow.NetSecurityPolicy = GameplayAbility.GetNetSecurityPolicy();
 	AuditRow.ReplicationPolicy = GameplayAbility.GetReplicationPolicy();
diff --git a/Engine/Plugins/Runtime/GameplayAbilities/Source/GameplayAbilitiesEditor/Public/GameplayAbilityAudit.h b/Engine/Plugins/Runtime/GameplayAbilities/Source/GameplayAbilitiesEditor/Public/GameplayAbilityAudit.h
--- a/Engine/Plugins/Runtime/GameplayAbilities/Source/GameplayAbilitiesEditor/Public/GameplayAbilityAudit.h
+++ b/Engine/Plugins/Runtime/GameplayAbilities/Source/GameplayAbilitiesEditor/Public/GameplayAbilityAudit.h
@@ -57,6 +57,10 @@ protected:
 	UPROPERTY()
 	EGameplayAbilityActivationPath ActivationPath = EGameplayAbilityActivationPath::Native;
 
+	/** The InstancingPolicy of the GameplayAbility (Instancing Behavior can affect memory, which functions can be called, what data can be used, and replication) */
+	UPROPERTY()
+	TEnumAsByte<EGameplayAbilityInstancingPolicy::Type> InstancingPolicy = EGameplayAbilityInstancingPolicy::Type::NonInstanced;
+
 	/** The NetExecutionPolicy of the GameplayAbility */
 	UPROPERTY()
 	TEnumAsByte<EGameplayAbilityNetExecutionPolicy::Type> NetExecutionPolicy = EGameplayAbilityNetExecutionPolicy::Type::LocalPredicted;
@@ -69,6 +73,10 @@ protected:
 	UPROPERTY()
 	TEnumAsByte<EGameplayAbilityReplicationPolicy::Type> ReplicationPolicy = EGameplayAbilityReplicationPolicy::Type::ReplicateNo;
 
+	/** Gameplay Tags that the Asset itself has (AbilityTags) */
+	UPROPERTY()
+	TArray<FName>	AssetTags;
+
 	/** If the GameplayAbility has a Cost GameplayEffect, this field will tell us which one it's using. */
 	UPROPERTY()
 	FName CostGE = NAME_None;
@@ -97,6 +105,10 @@ protected:
 	UPROPERTY()
 	EGameplayAbilityEndInBlueprints EndAbility = EGameplayAbilityEndInBlueprints::Missing;
 
+	/** All of the known referenced tags (list may be incomplete -- does not take into account Blueprint Tags) */
+	UPROPERTY()
+	TArray<FName>	ReferencedTags;
+
 	/** List of functions the Blueprint calls */
 	UPROPERTY()
 	TArray<FName>	Functions;
