Fixing possible null accesses.

diff --git a/Engine/Plugins/Runtime/GameplayAbilities/Source/GameplayAbilities/Private/Abilities/Tasks/AbilityTask_WaitGameplayEvent.cpp b/Engine/Plugins/Runtime/GameplayAbilities/Source/GameplayAbilities/Private/Abilities/Tasks/AbilityTask_WaitGameplayEvent.cpp
--- a/Engine/Plugins/Runtime/GameplayAbilities/Source/GameplayAbilities/Private/Abilities/Tasks/AbilityTask_WaitGameplayEvent.cpp
+++ b/Engine/Plugins/Runtime/GameplayAbilities/Source/GameplayAbilities/Private/Abilities/Tasks/AbilityTask_WaitGameplayEvent.cpp
@@ -52,7 +52,8 @@ void UAbilityTask_WaitGameplayEvent::GameplayEventContainerCallback(FGameplayTag
 {
 	if (ShouldBroadcastAbilityTaskDelegates())
 	{
-		FGameplayEventData TempPayload = *Payload;
+		ensureMsgf(Payload, TEXT("GameplayEventCallback expected non-null Payload"));
+		FGameplayEventData TempPayload = Payload ? *Payload : FGameplayEventData{};
 		TempPayload.EventTag = MatchingTag;
 		EventReceived.Broadcast(MoveTemp(TempPayload));
 	}
