Upon AbilitySystemComponent removal, now cancels pending TimerManager callbacks for duration GEs, which otherwise cause a crash if the ASC is part of an unloaded sublevel while the persistent UWorld's timer manager still runs.

diff --git a/Engine/Plugins/Runtime/GameplayAbilities/Source/GameplayAbilities/Private/GameplayEffect.cpp b/Engine/Plugins/Runtime/GameplayAbilities/Source/GameplayAbilities/Private/GameplayEffect.cpp
--- a/Engine/Plugins/Runtime/GameplayAbilities/Source/GameplayAbilities/Private/GameplayEffect.cpp
+++ b/Engine/Plugins/Runtime/GameplayAbilities/Source/GameplayAbilities/Private/GameplayEffect.cpp
@@ -4603,9 +4603,23 @@ void FActiveGameplayEffectsContainer::PostReplicatedReceive(const FFastArraySeri
 
 void FActiveGameplayEffectsContainer::Uninitialize()
 {
+	UWorld* World = Owner->GetWorld();
 	for (FActiveGameplayEffect& CurEffect : this)
 	{
 		RemoveCustomMagnitudeExternalDependencies(CurEffect);
+
+		// Remove any timer delegates that were scheduled to tick or end the gameplay effect
+		if (World)
+		{
+			if (CurEffect.DurationHandle.IsValid())
+			{
+				World->GetTimerManager().ClearTimer(CurEffect.DurationHandle);
+			}
+			if (CurEffect.PeriodHandle.IsValid())
+			{
+				World->GetTimerManager().ClearTimer(CurEffect.PeriodHandle);
+			}
+		}
 	}
 	ensure(CustomMagnitudeClassDependencies.Num() == 0);
 }
