AbilitySystemComponent::NotifyAbilityEnded was marking Gameplay Abilities as Garbage, but they can still be used in later functions, most notably AbilitySystemComponent::OnRemoveAbility.

diff --git a/Engine/Plugins/Runtime/GameplayAbilities/Source/GameplayAbilities/Private/Abilities/GameplayAbility.cpp b/Engine/Plugins/Runtime/GameplayAbilities/Source/GameplayAbilities/Private/Abilities/GameplayAbility.cpp
--- a/Engine/Plugins/Runtime/GameplayAbilities/Source/GameplayAbilities/Private/Abilities/GameplayAbility.cpp
+++ b/Engine/Plugins/Runtime/GameplayAbilities/Source/GameplayAbilities/Private/Abilities/GameplayAbility.cpp
@@ -87,7 +87,9 @@ UGameplayAbility::UGameplayAbility(const FObjectInitializer& ObjectInitializer)
 
 	ScopeLockCount = 0;
 
+	PRAGMA_DISABLE_DEPRECATION_WARNINGS
 	bMarkPendingKillOnAbilityEnd = false;
+	PRAGMA_ENABLE_DEPRECATION_WARNINGS
 }
 
 UWorld* UGameplayAbility::GetWorld() const
diff --git a/Engine/Plugins/Runtime/GameplayAbilities/Source/GameplayAbilities/Private/AbilitySystemComponent_Abilities.cpp b/Engine/Plugins/Runtime/GameplayAbilities/Source/GameplayAbilities/Private/AbilitySystemComponent_Abilities.cpp
--- a/Engine/Plugins/Runtime/GameplayAbilities/Source/GameplayAbilities/Private/AbilitySystemComponent_Abilities.cpp
+++ b/Engine/Plugins/Runtime/GameplayAbilities/Source/GameplayAbilities/Private/AbilitySystemComponent_Abilities.cpp
@@ -50,6 +50,7 @@ static TAutoConsoleVariable<bool> CVarAbilitySystemSetActivationInfoMultipleTime
 static TAutoConsoleVariable<bool> CVarGasFixClientSideMontageBlendOutTime(TEXT("AbilitySystem.Fix.ClientSideMontageBlendOutTime"), true, TEXT("Enable a fix to replicate the Montage BlendOutTime for (recently) stopped Montages"));
 static TAutoConsoleVariable<bool> CVarUpdateMontageSectionIdToPlay(TEXT("AbilitySystem.UpdateMontageSectionIdToPlay"), true, TEXT("During tick, update the section ID that replicated montages should use"));
 static TAutoConsoleVariable<bool> CVarReplicateMontageNextSectionId(TEXT("AbilitySystem.ReplicateMontageNextSectionId"), true, TEXT("Apply the replicated next section Id to montages when skipping position replication"));
+static TAutoConsoleVariable<bool> CVarEnsureAbilitiesEndGracefully(TEXT("AbilitySystem.EnsureAbilitiesEndGracefully"), true, TEXT("When shutting down (during ClearAllAbilities) we should check if all GameplayAbilities gracefully ended. This should be disabled if you have NonInstanced abilities that are designed for multiple concurrent executions."));
 
 void UAbilitySystemComponent::InitializeComponent()
 {
@@ -421,6 +422,19 @@ void UAbilitySystemComponent::ClearAllAbilities()
 		OnRemoveAbility(Spec);
 	}
 
+	// Let's add some enhanced checking if requested
+	if (CVarEnsureAbilitiesEndGracefully.GetValueOnGameThread())
+	{
+		for (FGameplayAbilitySpec& Spec : ActivatableAbilities.Items)
+		{
+			if (Spec.IsActive())
+			{
+				ensureAlwaysMsgf(Spec.Ability->GetInstancingPolicy() != EGameplayAbilityInstancingPolicy::NonInstanced, TEXT("%hs: %s was still active (ActiveCount = %d). Since it's not instanced, it's likely that TryActivateAbility and EndAbility are not matched."), __func__, *GetNameSafe(Spec.Ability), Spec.ActiveCount);
+				ensureAlwaysMsgf(Spec.Ability->GetInstancingPolicy() == EGameplayAbilityInstancingPolicy::NonInstanced, TEXT("%hs: %s was still active. It's likely that there's an issue with the flow of EndAbility or RemoveAbility."), __func__, *GetNameSafe(Spec.Ability));
+			}
+		}
+	}
+
 	ActivatableAbilities.Items.Empty(ActivatableAbilities.Items.Num());
 	ActivatableAbilities.MarkArrayDirty();
 
@@ -564,6 +578,8 @@ void UAbilitySystemComponent::OnGiveAbility(FGameplayAbilitySpec& Spec)
 
 void UAbilitySystemComponent::OnRemoveAbility(FGameplayAbilitySpec& Spec)
 {
+	ensureMsgf(AbilityScopeLockCount > 0, TEXT("%hs called without an Ability List Lock.  It can produce side effects and should be locked to pin the Spec argument."), __func__);
+
 	if (!Spec.Ability)
 	{
 		return;
@@ -594,8 +610,6 @@ void UAbilitySystemComponent::OnRemoveAbility(FGameplayAbilitySpec& Spec)
 		{
 			if (Instance->IsActive())
 			{
-				Instance->SetMarkPendingKillOnAbilityEnd(Instance->GetInstancingPolicy() == EGameplayAbilityInstancingPolicy::InstancedPerActor);
-
 				// End the ability but don't replicate it, OnRemoveAbility gets replicated
 				bool bReplicateEndAbility = false;
 				bool bWasCancelled = false;
@@ -604,10 +618,15 @@ void UAbilitySystemComponent::OnRemoveAbility(FGameplayAbilitySpec& Spec)
 			else
 			{
 				// Ability isn't active, but still needs to be destroyed
-				if (GetOwnerRole() == ROLE_Authority || Instance->GetReplicationPolicy() == EGameplayAbilityReplicationPolicy::ReplicateNo)
+				if (GetOwnerRole() == ROLE_Authority)
 				{
 					// Only destroy if we're the server or this isn't replicated. Can't destroy on the client or replication will fail when it replicates the end state
 					RemoveReplicatedInstancedAbility(Instance);
+				}
+
+				if (Instance->GetInstancingPolicy() == EGameplayAbilityInstancingPolicy::InstancedPerExecution)
+				{
+					ABILITY_LOG(Error, TEXT("%s was InActive, yet still instanced during OnRemove"), *Instance->GetName());
 					Instance->MarkAsGarbage();
 				}
 			}
@@ -619,6 +638,7 @@ void UAbilitySystemComponent::OnRemoveAbility(FGameplayAbilitySpec& Spec)
 	if (PrimaryInstance)
 	{
 		PrimaryInstance->OnRemoveAbility(AbilityActorInfo.Get(), Spec);
+		PrimaryInstance->MarkAsGarbage();
 	}
 	else
 	{
@@ -627,14 +647,13 @@ void UAbilitySystemComponent::OnRemoveAbility(FGameplayAbilitySpec& Spec)
 
 	// If this Ability Spec specified that it was created from an Active Gameplay Effect, then unlink the handle to the Active Gameplay Effect.
 	// Note: It's possible (maybe even likely) that the ActiveGE is no longer considered active by this point.
-	if (Spec.GameplayEffectHandle.IsValid())
+	// That means we can't use FindActiveGameplayEffectHandle (which fails if ActiveGE is PendingRemove), but also many of these checks will fail
+	// if the ActiveGE has completed its removal.
+	if (Spec.GameplayEffectHandle.IsValid()) // This can only be true on the network authority
 	{
-		UAbilitySystemComponent* SourceASC = Spec.GameplayEffectHandle.GetOwningAbilitySystemComponent();
-		UE_CLOG(!SourceASC, LogAbilitySystem, Error, TEXT("OnRemoveAbility Spec '%s' GameplayEffectHandle had invalid Owning Ability System Component"), *Spec.GetDebugString());
-		if (SourceASC)
+		if (UAbilitySystemComponent* SourceASC = Spec.GameplayEffectHandle.GetOwningAbilitySystemComponent())
 		{
-			FActiveGameplayEffect* SourceActiveGE = SourceASC->ActiveGameplayEffects.GetActiveGameplayEffect(Spec.GameplayEffectHandle);
-			if (SourceActiveGE)
+			if (FActiveGameplayEffect* SourceActiveGE = SourceASC->ActiveGameplayEffects.GetActiveGameplayEffect(Spec.GameplayEffectHandle))
 			{
 				SourceActiveGE->GrantedAbilityHandles.Remove(Spec.Handle);
 			}
@@ -1049,8 +1068,6 @@ void UAbilitySystemComponent::NotifyAbilityEnded(FGameplayAbilitySpecHandle Hand
 		// The ability spec may have been removed while we were ending. We can assume everything was cleaned up if the spec isnt here.
 		return;
 	}
-	check(Spec);
-	check(Ability);
 
 	ENetRole OwnerRole = GetOwnerRole();
 
@@ -1070,9 +1087,16 @@ void UAbilitySystemComponent::NotifyAbilityEnded(FGameplayAbilitySpecHandle Hand
 	AbilityEndedCallbacks.Broadcast(Ability);
 	OnAbilityEnded.Broadcast(FAbilityEndedData(Ability, Handle, false, bWasCancelled));
 	
+	// Above callbacks could have invalidated the Spec pointer, so find it again
+	Spec = FindAbilitySpecFromHandle(Handle);
+	if (!Spec)
+	{
+		ABILITY_LOG(Error, TEXT("%hs(%s): %s lost its active handle halfway through the function."), __func__, *GetNameSafe(Ability), *Handle.ToString());
+		return;
+	}
+
 	/** If this is instanced per execution or flagged for cleanup, mark pending kill and remove it from our instanced lists if we are the authority */
-	if ((Ability->GetInstancingPolicy() == EGameplayAbilityInstancingPolicy::InstancedPerExecution) ||
-		Ability->IsMarkPendingKillOnAbilityEnd())
+	if (Ability->GetInstancingPolicy() == EGameplayAbilityInstancingPolicy::InstancedPerExecution)
 	{
 		check(Ability->HasAnyFlags(RF_ClassDefaultObject) == false);	// Should never be calling this on a CDO for an instanced ability!
 
@@ -1082,14 +1106,14 @@ void UAbilitySystemComponent::NotifyAbilityEnded(FGameplayAbilitySpecHandle Hand
 			{
 				Spec->ReplicatedInstances.Remove(Ability);
 				RemoveReplicatedInstancedAbility(Ability);
-				Ability->MarkAsGarbage();
 			}
 		}
 		else
 		{
 			Spec->NonReplicatedInstances.Remove(Ability);
-			Ability->MarkAsGarbage();
 		}
+
+		Ability->MarkAsGarbage();
 	}
 
 	if (OwnerRole == ROLE_Authority)
diff --git a/Engine/Plugins/Runtime/GameplayAbilities/Source/GameplayAbilities/Private/GameplayAbilityTypes.cpp b/Engine/Plugins/Runtime/GameplayAbilities/Source/GameplayAbilities/Private/GameplayAbilityTypes.cpp
--- a/Engine/Plugins/Runtime/GameplayAbilities/Source/GameplayAbilities/Private/GameplayAbilityTypes.cpp
+++ b/Engine/Plugins/Runtime/GameplayAbilities/Source/GameplayAbilities/Private/GameplayAbilityTypes.cpp
@@ -232,6 +232,7 @@ void FGameplayAbilitySpec::PreReplicatedRemove(const struct FGameplayAbilitySpec
 {
 	if (InArraySerializer.Owner)
 	{
+		FScopedAbilityListLock AblityListLock(*InArraySerializer.Owner);
 		InArraySerializer.Owner->OnRemoveAbility(*this);
 	}
 }
diff --git a/Engine/Plugins/Runtime/GameplayAbilities/Source/GameplayAbilities/Public/Abilities/GameplayAbility.h b/Engine/Plugins/Runtime/GameplayAbilities/Source/GameplayAbilities/Public/Abilities/GameplayAbility.h
--- a/Engine/Plugins/Runtime/GameplayAbilities/Source/GameplayAbilities/Public/Abilities/GameplayAbility.h
+++ b/Engine/Plugins/Runtime/GameplayAbilities/Source/GameplayAbilities/Public/Abilities/GameplayAbility.h
@@ -900,15 +900,16 @@ protected:
 	void DecrementListLock() const;
 
 public:
-	/** Setter for the bMarkPendingKillOnAbilityEnd */
-	void SetMarkPendingKillOnAbilityEnd(bool bInMarkPendingKillOnAbilityEnd) { bMarkPendingKillOnAbilityEnd = bInMarkPendingKillOnAbilityEnd; }
+	UE_DEPRECATED(5.4, "This is unsafe and unnecessary.  It is ignored.")
+	void SetMarkPendingKillOnAbilityEnd(bool bInMarkPendingKillOnAbilityEnd) {}
 
-	/** Is bMarkPendingKillOnAbilityEnd set */
-	bool IsMarkPendingKillOnAbilityEnd() const { return bMarkPendingKillOnAbilityEnd; }
+	UE_DEPRECATED(5.4, "This is unsafe and unnecessary.  It will always return false.")
+	bool IsMarkPendingKillOnAbilityEnd() const { return false; }
 
 protected:
 
 	/** Flag that is set by AbilitySystemComponent::OnRemoveAbility to indicate the ability needs to be cleaned up in AbilitySystemComponent::NotifyAbilityEnded */
-	UPROPERTY(BlueprintReadOnly, Category = Ability)
+	UE_DEPRECATED(5.4, "This is unsafe. Do not use.")
+	UPROPERTY(BlueprintReadOnly, Category = Ability, meta=(DeprecatedProperty, DeprecationMessage="This is unsafe. Do not use."))
 	bool bMarkPendingKillOnAbilityEnd;
 };
