Make sure the transactional flag is always set on GEComponents, shortcomings in the blueprint compiler and coreuobject mean we cannot rely on anyone but our user code to conform it

diff --git a/Engine/Plugins/Runtime/GameplayAbilities/Source/GameplayAbilities/Private/GameplayEffect.cpp b/Engine/Plugins/Runtime/GameplayAbilities/Source/GameplayAbilities/Private/GameplayEffect.cpp
--- a/Engine/Plugins/Runtime/GameplayAbilities/Source/GameplayAbilities/Private/GameplayEffect.cpp
+++ b/Engine/Plugins/Runtime/GameplayAbilities/Source/GameplayAbilities/Private/GameplayEffect.cpp
@@ -98,6 +98,21 @@ namespace UE::GameplayEffect
 
 			return FromVersion < ToVersion;
 		}
+
+		static void ConformGameplayEffectTransactionality(TArrayView<TObjectPtr<UGameplayEffectComponent>> Components)
+		{
+			// This logic is compensating for FObjectInstancingGraph::GetInstancedSubobject's failure
+			// to propagate template object's flags - specifically RF_Transactional. If our CDO
+			// were reliably RF_Transactional we would also not need to do this. For now, let's
+			// just conform here:
+			for (const TObjectPtr<UGameplayEffectComponent>& GEComponent : Components)
+			{
+				if (GEComponent)
+				{
+					GEComponent->SetFlags(RF_Transactional);
+				}
+			}
+		}
 	}
 #endif
 }
@@ -362,6 +377,10 @@ void UGameplayEffect::PostCDOCompiled(const FPostCDOCompiledContext& Context)
 	ConvertTagRequirementsComponent();
 	ConvertTargetTagsComponent();
 	ConvertUIComponent();
+	#if WITH_EDITOR
+	using namespace UE::GameplayEffect::EditorOnly;
+	ConformGameplayEffectTransactionality(GEComponents);
+	#endif
 
 	const bool bAlreadyLoaded = !HasAnyFlags(RF_NeedPostLoad);
 	if (bAlreadyLoaded)
