AbilitySystemComponent will now correctly perform cleanup in DestroyActiveState after leaving, re-entering and leaving play again via level streaming (including world partition) and seamless travel.

diff --git a/Engine/Plugins/Runtime/GameplayAbilities/Source/GameplayAbilities/Private/AbilitySystemComponent.cpp b/Engine/Plugins/Runtime/GameplayAbilities/Source/GameplayAbilities/Private/AbilitySystemComponent.cpp
--- a/Engine/Plugins/Runtime/GameplayAbilities/Source/GameplayAbilities/Private/AbilitySystemComponent.cpp
+++ b/Engine/Plugins/Runtime/GameplayAbilities/Source/GameplayAbilities/Private/AbilitySystemComponent.cpp
@@ -219,6 +219,9 @@ PRAGMA_ENABLE_DEPRECATION_WARNINGS
 	{
 		AbilityActorInfo = TSharedPtr<FGameplayAbilityActorInfo>(UAbilitySystemGlobals::Get().AllocAbilityActorInfo());
 	}
+
+	// Ensure bDestroyActiveStateInitiated is clear in case component is re-entering play 
+	bDestroyActiveStateInitiated = false;
 }
 
 void UAbilitySystemComponent::OnUnregister()
