PR #10922: fix removing gameplay effect delegate handle

diff --git a/Engine/Plugins/Runtime/GameplayAbilities/Source/GameplayAbilities/Private/Abilities/Tasks/AbilityTask_WaitGameplayEffectApplied_Self.cpp b/Engine/Plugins/Runtime/GameplayAbilities/Source/GameplayAbilities/Private/Abilities/Tasks/AbilityTask_WaitGameplayEffectApplied_Self.cpp
--- a/Engine/Plugins/Runtime/GameplayAbilities/Source/GameplayAbilities/Private/Abilities/Tasks/AbilityTask_WaitGameplayEffectApplied_Self.cpp
+++ b/Engine/Plugins/Runtime/GameplayAbilities/Source/GameplayAbilities/Private/Abilities/Tasks/AbilityTask_WaitGameplayEffectApplied_Self.cpp
@@ -58,7 +58,7 @@ void UAbilityTask_WaitGameplayEffectApplied_Self::RemoveDelegate()
 	GetASC()->OnGameplayEffectAppliedDelegateToSelf.Remove(OnApplyGameplayEffectCallbackDelegateHandle);
 	if (OnPeriodicGameplayEffectExecuteCallbackDelegateHandle.IsValid())
 	{
-		GetASC()->OnGameplayEffectAppliedDelegateToTarget.Remove(OnPeriodicGameplayEffectExecuteCallbackDelegateHandle);
+		GetASC()->OnPeriodicGameplayEffectExecuteDelegateOnSelf.Remove(OnPeriodicGameplayEffectExecuteCallbackDelegateHandle);
 	}
 }
 
diff --git a/Engine/Plugins/Runtime/GameplayAbilities/Source/GameplayAbilities/Private/Abilities/Tasks/AbilityTask_WaitGameplayEffectApplied_Target.cpp b/Engine/Plugins/Runtime/GameplayAbilities/Source/GameplayAbilities/Private/Abilities/Tasks/AbilityTask_WaitGameplayEffectApplied_Target.cpp
--- a/Engine/Plugins/Runtime/GameplayAbilities/Source/GameplayAbilities/Private/Abilities/Tasks/AbilityTask_WaitGameplayEffectApplied_Target.cpp
+++ b/Engine/Plugins/Runtime/GameplayAbilities/Source/GameplayAbilities/Private/Abilities/Tasks/AbilityTask_WaitGameplayEffectApplied_Target.cpp
@@ -58,7 +58,7 @@ void UAbilityTask_WaitGameplayEffectApplied_Target::RemoveDelegate()
 	GetASC()->OnGameplayEffectAppliedDelegateToTarget.Remove(OnApplyGameplayEffectCallbackDelegateHandle);
 	if (OnPeriodicGameplayEffectExecuteCallbackDelegateHandle.IsValid())
 	{
-		GetASC()->OnGameplayEffectAppliedDelegateToTarget.Remove(OnPeriodicGameplayEffectExecuteCallbackDelegateHandle);
+		GetASC()->OnPeriodicGameplayEffectExecuteDelegateOnTarget.Remove(OnPeriodicGameplayEffectExecuteCallbackDelegateHandle);
 	}
 }
 
diff --git a/Engine/Plugins/Runtime/GameplayAbilities/Source/GameplayAbilities/Private/Abilities/Tasks/AbilityTask_WaitGameplayEffectStackChange.cpp b/Engine/Plugins/Runtime/GameplayAbilities/Source/GameplayAbilities/Private/Abilities/Tasks/AbilityTask_WaitGameplayEffectStackChange.cpp
--- a/Engine/Plugins/Runtime/GameplayAbilities/Source/GameplayAbilities/Private/Abilities/Tasks/AbilityTask_WaitGameplayEffectStackChange.cpp
+++ b/Engine/Plugins/Runtime/GameplayAbilities/Source/GameplayAbilities/Private/Abilities/Tasks/AbilityTask_WaitGameplayEffectStackChange.cpp
@@ -50,7 +50,7 @@ void UAbilityTask_WaitGameplayEffectStackChange::OnDestroy(bool AbilityIsEnding)
 	UAbilitySystemComponent* EffectOwningAbilitySystemComponent = Handle.GetOwningAbilitySystemComponent();
 	if (EffectOwningAbilitySystemComponent && OnGameplayEffectStackChangeDelegateHandle.IsValid())
 	{
-		FOnActiveGameplayEffectRemoved_Info* DelPtr = EffectOwningAbilitySystemComponent->OnGameplayEffectRemoved_InfoDelegate(Handle);
+		FOnActiveGameplayEffectStackChange* DelPtr = EffectOwningAbilitySystemComponent->OnGameplayEffectStackChangeDelegate(Handle);
 		if (DelPtr)
 		{
 			DelPtr->Remove(OnGameplayEffectStackChangeDelegateHandle);
