PR #10990: GAS Montage Section Crash Fix

diff --git a/Engine/Plugins/Runtime/GameplayAbilities/Source/GameplayAbilities/Private/AbilitySystemComponent_Abilities.cpp b/Engine/Plugins/Runtime/GameplayAbilities/Source/GameplayAbilities/Private/AbilitySystemComponent_Abilities.cpp
--- a/Engine/Plugins/Runtime/GameplayAbilities/Source/GameplayAbilities/Private/AbilitySystemComponent_Abilities.cpp
+++ b/Engine/Plugins/Runtime/GameplayAbilities/Source/GameplayAbilities/Private/AbilitySystemComponent_Abilities.cpp
@@ -3291,7 +3291,7 @@ void UAbilitySystemComponent::ServerCurrentMontageSetNextSectionName_Implementat
 	if (AnimInstance)
 	{
 		UAnimMontage* CurrentAnimMontage = LocalAnimMontageInfo.AnimMontage;
-		if (ClientAnimMontage == CurrentAnimMontage)
+		if (CurrentAnimMontage && ClientAnimMontage == CurrentAnimMontage)
 		{
 			// Set NextSectionName
 			AnimInstance->Montage_SetNextSection(SectionName, NextSectionName, CurrentAnimMontage);
