Fixing up a case where OnRemoveAbility would not EndAbility on a non-instanced Ability that was still Active.  Without this fix, it could trigger the ensure in ClearAllAbilities (which had assumed this code was present).

Also modifying the ASC shutdown procedure and elaborating on some potential pitfalls/errors in the code that have not come up yet.

diff --git a/Engine/Plugins/Runtime/GameplayAbilities/Source/GameplayAbilities/Private/AbilitySystemComponent_Abilities.cpp b/Engine/Plugins/Runtime/GameplayAbilities/Source/GameplayAbilities/Private/AbilitySystemComponent_Abilities.cpp
--- a/Engine/Plugins/Runtime/GameplayAbilities/Source/GameplayAbilities/Private/AbilitySystemComponent_Abilities.cpp
+++ b/Engine/Plugins/Runtime/GameplayAbilities/Source/GameplayAbilities/Private/AbilitySystemComponent_Abilities.cpp
@@ -430,7 +430,7 @@ void UAbilitySystemComponent::ClearAllAbilities()
 			if (Spec.IsActive())
 			{
 				ensureAlwaysMsgf(Spec.Ability->GetInstancingPolicy() != EGameplayAbilityInstancingPolicy::NonInstanced, TEXT("%hs: %s was still active (ActiveCount = %d). Since it's not instanced, it's likely that TryActivateAbility and EndAbility are not matched."), __func__, *GetNameSafe(Spec.Ability), Spec.ActiveCount);
-				ensureAlwaysMsgf(Spec.Ability->GetInstancingPolicy() == EGameplayAbilityInstancingPolicy::NonInstanced, TEXT("%hs: %s was still active. It's likely that there's an issue with the flow of EndAbility or RemoveAbility."), __func__, *GetNameSafe(Spec.Ability));
+				ensureAlwaysMsgf(Spec.Ability->GetInstancingPolicy() == EGameplayAbilityInstancingPolicy::NonInstanced, TEXT("%hs: %s was still active. Since it's an instanced ability, it's likely that there's an issue with the flow of EndAbility or RemoveAbility (such as not calling the Super function)."), __func__, *GetNameSafe(Spec.Ability));
 			}
 		}
 	}
@@ -642,6 +642,18 @@ void UAbilitySystemComponent::OnRemoveAbility(FGameplayAbilitySpec& Spec)
 	}
 	else
 	{
+		// If we're non-instanced and still active, we need to End
+		if (Spec.IsActive())
+		{
+			if (ensureMsgf(Spec.Ability->GetInstancingPolicy() == EGameplayAbilityInstancingPolicy::NonInstanced, TEXT("We should never have an instanced Gameplay Ability that is still active by this point. All instances should have EndAbility called just before here.")))
+			{
+				// Seems like it should be cancelled, but we're just following the existing pattern (could be due to functionality from OnRep)
+				constexpr bool bReplicateEndAbility = false;
+				constexpr bool bWasCancelled = false;
+				Spec.Ability->EndAbility(Spec.Handle, AbilityActorInfo.Get(), Spec.ActivationInfo, bReplicateEndAbility, bWasCancelled);
+			}
+		}
+
 		Spec.Ability->OnRemoveAbility(AbilityActorInfo.Get(), Spec);
 	}
 
@@ -1234,27 +1246,32 @@ void UAbilitySystemComponent::DestroyActiveState()
 			CancelAbilities();
 		}
 
-		// Mark pending kill any remaining instanced abilities
-		// (CancelAbilities() will only MarkPending kill InstancePerExecution abilities).
-		for (FGameplayAbilitySpec& Spec : ActivatableAbilities.Items)
+		if (IsOwnerActorAuthoritative())
 		{
-			TArray<UGameplayAbility*> AbilitiesToCancel = Spec.GetAbilityInstances();
-			for (UGameplayAbility* InstanceAbility : AbilitiesToCancel)
+			// We should now ClearAllAbilities because not all abilities CanBeCanceled().
+			// This will gracefully call EndAbility and clean-up all instances of the abilities.
+			ClearAllAbilities();
+		}
+		else
+		{
+			// If we're a client, ClearAllAbilities won't execute and we should clean up these instances manually.
+			// CancelAbilities() will only MarkPending kill InstancePerExecution abilities.
+			// TODO: Is it correct to simply mark these as Garbage rather than EndAbility?  I suspect not, but this
+			// is ingrained behavior (circa 2015). Perhaps better to allow ClearAllAbilities on client if bDestroyActiveStateInitiated (Nov 2023).
+			for (FGameplayAbilitySpec& Spec : ActivatableAbilities.Items)
 			{
-				if (InstanceAbility)
+				TArray<UGameplayAbility*> AbilitiesToCancel = Spec.GetAbilityInstances();
+				for (UGameplayAbility* InstanceAbility : AbilitiesToCancel)
 				{
-					InstanceAbility->MarkAsGarbage();
+					if (InstanceAbility)
+					{
+						InstanceAbility->MarkAsGarbage();
+					}
 				}
-			}
-
-			Spec.ReplicatedInstances.Empty();
-			Spec.NonReplicatedInstances.Empty();
-		}
 
-		if (IsOwnerActorAuthoritative())
-		{
-			// Ability specs are no longer valid, clear them all.
-			ClearAllAbilities();
+				Spec.ReplicatedInstances.Empty();
+				Spec.NonReplicatedInstances.Empty();
+			}
 		}
 	}
 }
