Add dynamic montage support and update montage replication
* Unshelved from 22565539

diff --git a/Engine/Plugins/Runtime/GameplayAbilities/Source/GameplayAbilities/Private/Abilities/Tasks/AbilityTask_PlayAnimAndWait.cpp b/Engine/Plugins/Runtime/GameplayAbilities/Source/GameplayAbilities/Private/Abilities/Tasks/AbilityTask_PlayAnimAndWait.cpp
new file mode 100644
--- /dev/null
+++ b/Engine/Plugins/Runtime/GameplayAbilities/Source/GameplayAbilities/Private/Abilities/Tasks/AbilityTask_PlayAnimAndWait.cpp
@@ -0,0 +1,259 @@
+// Copyright Epic Games, Inc. All Rights Reserved.
+
+#include "Abilities/Tasks/AbilityTask_PlayAnimAndWait.h"
+#include "GameFramework/Character.h"
+#include "AbilitySystemComponent.h"
+#include "AbilitySystemGlobals.h"
+#include "Animation/AnimMontage.h"
+#include "Animation/AnimSequence.h"
+#include "AbilitySystemLog.h"
+
+void UAbilityTask_PlayAnimAndWait::OnMontageBlendingOut(UAnimMontage* Montage, bool bInterrupted)
+{
+	if (Montage && Ability && Ability->GetCurrentMontage() == Montage)
+	{
+		if (Montage->GetFirstAnimReference() == AnimSequenceToPlay)
+		{
+			if (UAbilitySystemComponent* ASC = AbilitySystemComponent.Get())
+			{
+				ASC->ClearAnimatingAbility(Ability);
+			}
+
+			// Reset AnimRootMotionTranslationScale
+			ACharacter* Character = Cast<ACharacter>(GetAvatarActor());
+			if (Character && (Character->GetLocalRole() == ROLE_Authority ||
+							  (Character->GetLocalRole() == ROLE_AutonomousProxy && Ability->GetNetExecutionPolicy() == EGameplayAbilityNetExecutionPolicy::LocalPredicted)))
+			{
+				Character->SetAnimRootMotionTranslationScale(1.f);
+			}
+
+		}
+	}
+
+	if (bInterrupted)
+	{
+		if (ShouldBroadcastAbilityTaskDelegates())
+		{
+			OnInterrupted.Broadcast();
+		}
+	}
+	else
+	{
+		if (ShouldBroadcastAbilityTaskDelegates())
+		{
+			OnBlendOut.Broadcast();
+		}
+	}
+}
+
+void UAbilityTask_PlayAnimAndWait::OnMontageInterrupted()
+{
+	if (StopPlayingMontage())
+	{
+		// Let the BP handle the interrupt as well
+		if (ShouldBroadcastAbilityTaskDelegates())
+		{
+			OnInterrupted.Broadcast();
+		}
+	}
+}
+
+void UAbilityTask_PlayAnimAndWait::OnMontageEnded(UAnimMontage* Montage, bool bInterrupted)
+{
+	if (!bInterrupted)
+	{
+		if (ShouldBroadcastAbilityTaskDelegates())
+		{
+			OnCompleted.Broadcast();
+		}
+	}
+
+	EndTask();
+}
+
+UAbilityTask_PlayAnimAndWait* UAbilityTask_PlayAnimAndWait::CreatePlayAnimAndWaitProxy(
+	UGameplayAbility* OwningAbility,
+	FName TaskInstanceName, 
+	UAnimSequence* AnimSequence, 
+	FName SlotName, 
+	float BlendInTime, 
+	float BlendOutTime, 
+	float InPlayRate,
+	float StartTimeSeconds /*= 0.0f*/, 
+	bool bStopWhenAbilityEnds /*= true*/, 
+	float AnimRootMotionTranslationScale /*= 1.f*/)
+{
+	UAbilityTask_PlayAnimAndWait* MyObj = NewAbilityTask<UAbilityTask_PlayAnimAndWait>(OwningAbility, TaskInstanceName);
+	MyObj->AnimSequenceToPlay = AnimSequence;
+	MyObj->SlotName = SlotName;
+	MyObj->BlendInTime = BlendInTime;
+	MyObj->BlendOutTime = BlendOutTime;
+	MyObj->PlayRate = InPlayRate;
+	MyObj->AnimRootMotionTranslationScale = AnimRootMotionTranslationScale;
+	MyObj->bStopWhenAbilityEnds = bStopWhenAbilityEnds;
+	MyObj->StartTimeSeconds = StartTimeSeconds;
+
+	return MyObj;
+}
+
+void UAbilityTask_PlayAnimAndWait::Activate()
+{
+	if (Ability == nullptr)
+	{
+		return;
+	}
+
+	bool bPlayedMontage = false;
+
+	if (UAbilitySystemComponent* ASC = AbilitySystemComponent.Get())
+	{
+		const FGameplayAbilityActorInfo* ActorInfo = Ability->GetCurrentActorInfo();
+		UAnimInstance* AnimInstance = ActorInfo->GetAnimInstance();
+		if (AnimInstance != nullptr)
+		{
+			UAnimMontage* NewDynamicMontage = ASC->PlaySlotAnimationAsDynamicMontage(Ability, Ability->GetCurrentActivationInfo(), AnimSequenceToPlay, SlotName, BlendInTime, BlendOutTime, PlayRate, StartTimeSeconds);
+			if (NewDynamicMontage)
+			{
+				
+				// Playing a montage could potentially fire off a callback into game code which could kill this ability! Early out if we are  pending kill.
+				if (ShouldBroadcastAbilityTaskDelegates() == false)
+				{
+					return;
+				}
+
+				InterruptedHandle = Ability->OnGameplayAbilityCancelled.AddUObject(this, &UAbilityTask_PlayAnimAndWait::OnMontageInterrupted);
+
+				BlendingOutDelegate.BindUObject(this, &UAbilityTask_PlayAnimAndWait::OnMontageBlendingOut);
+				AnimInstance->Montage_SetBlendingOutDelegate(BlendingOutDelegate, NewDynamicMontage);
+
+				MontageEndedDelegate.BindUObject(this, &UAbilityTask_PlayAnimAndWait::OnMontageEnded);
+				AnimInstance->Montage_SetEndDelegate(MontageEndedDelegate, NewDynamicMontage);
+
+				ACharacter* Character = Cast<ACharacter>(GetAvatarActor());
+				if (Character && (Character->GetLocalRole() == ROLE_Authority ||
+								  (Character->GetLocalRole() == ROLE_AutonomousProxy && Ability->GetNetExecutionPolicy() == EGameplayAbilityNetExecutionPolicy::LocalPredicted)))
+				{
+					Character->SetAnimRootMotionTranslationScale(AnimRootMotionTranslationScale);
+				}
+
+				bPlayedMontage = true;
+			}
+		}
+		else
+		{
+			ABILITY_LOG(Warning, TEXT("UAbilityTask_PlayAnimAndWait call to PlayMontage failed!"));
+		}
+	}
+	else
+	{
+		ABILITY_LOG(Warning, TEXT("UAbilityTask_PlayAnimAndWait called on invalid AbilitySystemComponent"));
+	}
+
+	if (!bPlayedMontage)
+	{
+		ABILITY_LOG(Warning, TEXT("UAbilityTask_PlayAnimAndWait called in Ability %s failed to play dynamic montage %s; Task Instance Name %s."), *Ability->GetName(), *GetNameSafe(AnimSequenceToPlay),*InstanceName.ToString());
+		if (ShouldBroadcastAbilityTaskDelegates())
+		{
+			OnCancelled.Broadcast();
+		}
+	}
+
+	SetWaitingOnAvatar();
+}
+
+void UAbilityTask_PlayAnimAndWait::ExternalCancel()
+{
+	if (ShouldBroadcastAbilityTaskDelegates())
+	{
+		OnCancelled.Broadcast();
+	}
+	Super::ExternalCancel();
+}
+
+void UAbilityTask_PlayAnimAndWait::OnDestroy(bool AbilityEnded)
+{
+	// Note: Clearing montage end delegate isn't necessary since its not a multicast and will be cleared when the next montage plays.
+	// (If we are destroyed, it will detect this and not do anything)
+
+	// This delegate, however, should be cleared as it is a multicast
+	if (Ability)
+	{
+		Ability->OnGameplayAbilityCancelled.Remove(InterruptedHandle);
+		if (AbilityEnded && bStopWhenAbilityEnds)
+		{
+			StopPlayingMontage();
+		}
+	}
+
+	Super::OnDestroy(AbilityEnded);
+
+}
+
+bool UAbilityTask_PlayAnimAndWait::StopPlayingMontage()
+{
+	if (Ability == nullptr)
+	{
+		return false;
+	}
+
+	const FGameplayAbilityActorInfo* ActorInfo = Ability->GetCurrentActorInfo();
+	if (ActorInfo == nullptr)
+	{
+		return false;
+	}
+
+	UAnimInstance* AnimInstance = ActorInfo->GetAnimInstance();
+	if (AnimInstance == nullptr)
+	{
+		return false;
+	}
+
+	// Check if the montage is still playing
+	// The ability would have been interrupted, in which case we should automatically stop the montage
+	UAbilitySystemComponent* ASC = AbilitySystemComponent.Get();
+	if (ASC && Ability)
+	{
+		const UAnimMontage* Montage = ASC->GetCurrentMontage();
+		if (ASC->GetAnimatingAbility() == Ability && Montage && Montage->IsDynamicMontage() && Montage->GetFirstAnimReference() == AnimSequenceToPlay)
+		{
+			// Unbind delegates so they don't get called as well
+			FAnimMontageInstance* MontageInstance = AnimInstance->GetActiveInstanceForMontage(Montage);
+			if (MontageInstance)
+			{
+				MontageInstance->OnMontageBlendingOutStarted.Unbind();
+				MontageInstance->OnMontageEnded.Unbind();
+			}
+
+			ASC->CurrentMontageStop();
+			return true;
+		}
+	}
+
+	return false;
+}
+
+FString UAbilityTask_PlayAnimAndWait::GetDebugString() const
+{
+	UAnimMontage* PlayingMontage = nullptr;
+	if (Ability)
+	{
+		const FGameplayAbilityActorInfo* ActorInfo = Ability->GetCurrentActorInfo();
+		UAnimInstance* AnimInstance = ActorInfo->GetAnimInstance();
+
+		if (AnimInstance != nullptr)
+		{
+			PlayingMontage = AnimInstance->GetCurrentActiveMontage();
+		}
+	}
+
+	if(PlayingMontage)
+	{
+		return FString::Printf(TEXT("PlayAnimAndWait. AnimSequenceToPlay: %s Currently Playing: %s IsDynamic: %d AnimSequence: %s"), 
+			*GetNameSafe(AnimSequenceToPlay), *GetNameSafe(PlayingMontage), PlayingMontage->IsDynamicMontage(), *GetNameSafe(PlayingMontage->GetFirstAnimReference()));
+	}
+	else
+	{
+		return FString::Printf(TEXT("PlayAnimAndWait. AnimSequenceToPlay: %s  (Currently Playing): None"), *GetNameSafe(AnimSequenceToPlay));
+	}
+	
+}
\ No newline at end of file
diff --git a/Engine/Plugins/Runtime/GameplayAbilities/Source/GameplayAbilities/Private/AbilitySystemComponent_Abilities.cpp b/Engine/Plugins/Runtime/GameplayAbilities/Source/GameplayAbilities/Private/AbilitySystemComponent_Abilities.cpp
--- a/Engine/Plugins/Runtime/GameplayAbilities/Source/GameplayAbilities/Private/AbilitySystemComponent_Abilities.cpp
+++ b/Engine/Plugins/Runtime/GameplayAbilities/Source/GameplayAbilities/Private/AbilitySystemComponent_Abilities.cpp
@@ -2812,10 +2812,12 @@ float UAbilitySystemComponent::PlayMontage(UGameplayAbility* InAnimatingAbility,
 				}
 			}
 
+			UAnimSequenceBase* Animation = NewAnimMontage->IsDynamicMontage() ? NewAnimMontage->GetFirstAnimReference() : NewAnimMontage;
+
 			if (NewAnimMontage->HasRootMotion() && AnimInstance->GetOwningActor())
 			{
 				UE_LOG(LogRootMotion, Log, TEXT("UAbilitySystemComponent::PlayMontage %s, Role: %s")
-					, *GetNameSafe(NewAnimMontage)
+					, *GetNameSafe(Animation)
 					, *UEnum::GetValueAsString(TEXT("Engine.ENetRole"), AnimInstance->GetOwningActor()->GetLocalRole())
 					);
 			}
@@ -2843,14 +2845,21 @@ float UAbilitySystemComponent::PlayMontage(UGameplayAbility* InAnimatingAbility,
 				FGameplayAbilityRepAnimMontage& MutableRepAnimMontageInfo = GetRepAnimMontageInfo_Mutable();
 
 				// Those are static parameters, they are only set when the montage is played. They are not changed after that.
-				MutableRepAnimMontageInfo.AnimMontage = NewAnimMontage;
+				MutableRepAnimMontageInfo.Animation = Animation;
 				MutableRepAnimMontageInfo.PlayInstanceId = (MutableRepAnimMontageInfo.PlayInstanceId < UINT8_MAX ? MutableRepAnimMontageInfo.PlayInstanceId + 1 : 0);
 
 				MutableRepAnimMontageInfo.SectionIdToPlay = 0;
-				if (MutableRepAnimMontageInfo.AnimMontage && StartSectionName != NAME_None)
+				if (MutableRepAnimMontageInfo.Animation && StartSectionName != NAME_None)
 				{
 					// we add one so INDEX_NONE can be used in the on rep
-					MutableRepAnimMontageInfo.SectionIdToPlay = MutableRepAnimMontageInfo.AnimMontage->GetSectionIndex(StartSectionName) + 1;
+					MutableRepAnimMontageInfo.SectionIdToPlay = NewAnimMontage->GetSectionIndex(StartSectionName) + 1;
+				}
+
+				if (NewAnimMontage->IsDynamicMontage())
+				{
+					check(!NewAnimMontage->SlotAnimTracks.IsEmpty());
+					MutableRepAnimMontageInfo.SlotName = NewAnimMontage->SlotAnimTracks[0].SlotName;
+					MutableRepAnimMontageInfo.BlendOutTime = NewAnimMontage->GetDefaultBlendInTime();
 				}
 
 				// Update parameters that change during Montage life time.
@@ -2909,7 +2918,16 @@ void UAbilitySystemComponent::AnimMontage_UpdateReplicatedData(FGameplayAbilityR
 	const UAnimInstance* AnimInstance = AbilityActorInfo.IsValid() ? AbilityActorInfo->GetAnimInstance() : nullptr;
 	if (AnimInstance && LocalAnimMontageInfo.AnimMontage)
 	{
-		OutRepAnimMontageInfo.AnimMontage = LocalAnimMontageInfo.AnimMontage;
+		if (LocalAnimMontageInfo.AnimMontage->IsDynamicMontage())
+		{
+			OutRepAnimMontageInfo.Animation = LocalAnimMontageInfo.AnimMontage->GetFirstAnimReference();
+			OutRepAnimMontageInfo.BlendOutTime = LocalAnimMontageInfo.AnimMontage->GetDefaultBlendOutTime();
+		}
+		else
+		{
+			OutRepAnimMontageInfo.Animation = LocalAnimMontageInfo.AnimMontage;
+			OutRepAnimMontageInfo.BlendOutTime = 0.0f;
+		}
 
 		// Compressed Flags
 		const bool bIsStopped = AnimInstance->Montage_GetIsStopped(LocalAnimMontageInfo.AnimMontage);
@@ -3041,7 +3059,7 @@ void UAbilitySystemComponent::OnRep_ReplicatedAnimMontage()
 		{
 			ABILITY_LOG( Warning, TEXT("\n\nOnRep_ReplicatedAnimMontage, %s"), *GetNameSafe(this));
 			ABILITY_LOG( Warning, TEXT("\tAnimMontage: %s\n\tPlayRate: %f\n\tPosition: %f\n\tBlendTime: %f\n\tNextSectionID: %d\n\tIsStopped: %d\n\tPlayInstanceId: %d"),
-				*GetNameSafe(ConstRepAnimMontageInfo.AnimMontage),
+				*GetNameSafe(ConstRepAnimMontageInfo.Animation),
 				ConstRepAnimMontageInfo.PlayRate,
 				ConstRepAnimMontageInfo.Position,
 				ConstRepAnimMontageInfo.BlendTime,
@@ -3052,19 +3070,33 @@ void UAbilitySystemComponent::OnRep_ReplicatedAnimMontage()
 				*GetNameSafe(LocalAnimMontageInfo.AnimMontage), AnimInstance->Montage_GetPosition(LocalAnimMontageInfo.AnimMontage));
 		}
 
-		if(ConstRepAnimMontageInfo.AnimMontage )
+		if(ConstRepAnimMontageInfo.Animation)
 		{
 			// New Montage to play
-			if ((LocalAnimMontageInfo.AnimMontage != ConstRepAnimMontageInfo.AnimMontage) || 
+			UAnimSequenceBase* LocalAnimation = LocalAnimMontageInfo.AnimMontage && LocalAnimMontageInfo.AnimMontage->IsDynamicMontage() ? LocalAnimMontageInfo.AnimMontage->GetFirstAnimReference() : LocalAnimMontageInfo.AnimMontage;
+			if ((LocalAnimation != ConstRepAnimMontageInfo.Animation) ||
 			    (LocalAnimMontageInfo.PlayInstanceId != ConstRepAnimMontageInfo.PlayInstanceId))
 			{
 				LocalAnimMontageInfo.PlayInstanceId = ConstRepAnimMontageInfo.PlayInstanceId;
-				PlayMontageSimulated(ConstRepAnimMontageInfo.AnimMontage, ConstRepAnimMontageInfo.PlayRate);
+
+				if (UAnimMontage* MontageToPlay = Cast<UAnimMontage>(ConstRepAnimMontageInfo.Animation))
+				{
+					PlayMontageSimulated(MontageToPlay, ConstRepAnimMontageInfo.PlayRate);
+				}
+				else
+				{
+					PlaySlotAnimationAsDynamicMontageSimulated(
+						ConstRepAnimMontageInfo.Animation,
+						ConstRepAnimMontageInfo.SlotName,
+						ConstRepAnimMontageInfo.BlendTime,
+						ConstRepAnimMontageInfo.BlendOutTime,
+						ConstRepAnimMontageInfo.PlayRate);
+				}
 			}
 
 			if (LocalAnimMontageInfo.AnimMontage == nullptr)
 			{ 
-				ABILITY_LOG(Warning, TEXT("OnRep_ReplicatedAnimMontage: PlayMontageSimulated failed. Name: %s, AnimMontage: %s"), *GetNameSafe(this), *GetNameSafe(ConstRepAnimMontageInfo.AnimMontage));
+				ABILITY_LOG(Warning, TEXT("OnRep_ReplicatedAnimMontage: PlayMontageSimulated failed. Name: %s, Animation: %s"), *GetNameSafe(this), *GetNameSafe(ConstRepAnimMontageInfo.Animation));
 				return;
 			}
 
@@ -3122,7 +3154,7 @@ void UAbilitySystemComponent::OnRep_ReplicatedAnimMontage()
 				if ((CurrentSectionID == RepSectionID) && (FMath::Abs(DeltaPosition) > MONTAGE_REP_POS_ERR_THRESH) && (ConstRepAnimMontageInfo.IsStopped == 0))
 				{
 					// fast forward to server position and trigger notifies
-					if (FAnimMontageInstance* MontageInstance = AnimInstance->GetActiveInstanceForMontage(ConstRepAnimMontageInfo.AnimMontage))
+					if (FAnimMontageInstance* MontageInstance = AnimInstance->GetActiveInstanceForMontage(LocalAnimMontageInfo.AnimMontage))
 					{
 						// Skip triggering notifies if we're going backwards in time, we've already triggered them.
 						const float DeltaTime = !FMath::IsNearlyZero(ConstRepAnimMontageInfo.PlayRate) ? (DeltaPosition / ConstRepAnimMontageInfo.PlayRate) : 0.f;
@@ -3239,10 +3271,14 @@ void UAbilitySystemComponent::CurrentMontageJumpToSection(FName SectionName)
 			FGameplayAbilityRepAnimMontage& MutableRepAnimMontageInfo = GetRepAnimMontageInfo_Mutable();
 
 			MutableRepAnimMontageInfo.SectionIdToPlay = 0;
-			if (MutableRepAnimMontageInfo.AnimMontage)
+			if (MutableRepAnimMontageInfo.Animation)
 			{
-				// we add one so INDEX_NONE can be used in the on rep
-				MutableRepAnimMontageInfo.SectionIdToPlay = MutableRepAnimMontageInfo.AnimMontage->GetSectionIndex(SectionName) + 1;
+				// Only change SectionIdToPlay if the anim montage's source is a montage. Dynamic montages have no sections.
+				if (const UAnimMontage* RepAnimMontage = Cast<UAnimMontage>(MutableRepAnimMontageInfo.Animation))
+				{
+					// we add one so INDEX_NONE can be used in the on rep
+					MutableRepAnimMontageInfo.SectionIdToPlay = RepAnimMontage->GetSectionIndex(SectionName) + 1;
+				}
 			}
 
 			AnimMontage_UpdateReplicatedData();
@@ -3250,8 +3286,9 @@ void UAbilitySystemComponent::CurrentMontageJumpToSection(FName SectionName)
 		
 		// If we are NOT the authority, then let the server handling jumping the montage.
 		if (!IsOwnerActorAuthoritative())
-		{
-			ServerCurrentMontageJumpToSectionName(LocalAnimMontageInfo.AnimMontage, SectionName);
+		{	
+			UAnimSequenceBase* Animation = LocalAnimMontageInfo.AnimMontage->IsDynamicMontage() ? LocalAnimMontageInfo.AnimMontage->GetFirstAnimReference() : LocalAnimMontageInfo.AnimMontage;
+			ServerCurrentMontageJumpToSectionName(Animation, SectionName);
 		}
 	}
 }
@@ -3259,9 +3296,9 @@ void UAbilitySystemComponent::CurrentMontageJumpToSection(FName SectionName)
 void UAbilitySystemComponent::CurrentMontageSetNextSectionName(FName FromSectionName, FName ToSectionName)
 {
 	UAnimInstance* AnimInstance = AbilityActorInfo.IsValid() ? AbilityActorInfo->GetAnimInstance() : nullptr;
-	if( LocalAnimMontageInfo.AnimMontage && AnimInstance )
+	if (LocalAnimMontageInfo.AnimMontage && AnimInstance)
 	{
-		// Set Next Section Name. 
+		// Set Next Section Name.
 		AnimInstance->Montage_SetNextSection(FromSectionName, ToSectionName, LocalAnimMontageInfo.AnimMontage);
 
 		// Update replicated version for Simulated Proxies if we are on the server.
@@ -3272,7 +3309,8 @@ void UAbilitySystemComponent::CurrentMontageSetNextSectionName(FName FromSection
 		else
 		{
 			float CurrentPosition = AnimInstance->Montage_GetPosition(LocalAnimMontageInfo.AnimMontage);
-			ServerCurrentMontageSetNextSectionName(LocalAnimMontageInfo.AnimMontage, CurrentPosition, FromSectionName, ToSectionName);
+			UAnimSequenceBase* Animation = LocalAnimMontageInfo.AnimMontage->IsDynamicMontage() ? LocalAnimMontageInfo.AnimMontage->GetFirstAnimReference() : LocalAnimMontageInfo.AnimMontage;
+			ServerCurrentMontageSetNextSectionName(Animation, CurrentPosition, FromSectionName, ToSectionName);
 		}
 	}
 }
@@ -3292,24 +3330,27 @@ void UAbilitySystemComponent::CurrentMontageSetPlayRate(float InPlayRate)
 		}
 		else
 		{
+			UAnimSequenceBase* Animation = LocalAnimMontageInfo.AnimMontage->IsDynamicMontage() ? LocalAnimMontageInfo.AnimMontage->GetFirstAnimReference() : LocalAnimMontageInfo.AnimMontage;
 			ServerCurrentMontageSetPlayRate(LocalAnimMontageInfo.AnimMontage, InPlayRate);
 		}
 	}
 }
 
-bool UAbilitySystemComponent::ServerCurrentMontageSetNextSectionName_Validate(UAnimMontage* ClientAnimMontage, float ClientPosition, FName SectionName, FName NextSectionName)
+bool UAbilitySystemComponent::ServerCurrentMontageSetNextSectionName_Validate(UAnimSequenceBase* ClientAnimation, float ClientPosition, FName SectionName, FName NextSectionName)
 {
 	return true;
 }
 
-void UAbilitySystemComponent::ServerCurrentMontageSetNextSectionName_Implementation(UAnimMontage* ClientAnimMontage, float ClientPosition, FName SectionName, FName NextSectionName)
+void UAbilitySystemComponent::ServerCurrentMontageSetNextSectionName_Implementation(UAnimSequenceBase* ClientAnimation, float ClientPosition, FName SectionName, FName NextSectionName)
 {
 	UAnimInstance* AnimInstance = AbilityActorInfo.IsValid() ? AbilityActorInfo->GetAnimInstance() : nullptr;
-	if (AnimInstance)
+	if (AnimInstance && LocalAnimMontageInfo.AnimMontage)
 	{
-		UAnimMontage* CurrentAnimMontage = LocalAnimMontageInfo.AnimMontage;
-		if (CurrentAnimMontage && ClientAnimMontage == CurrentAnimMontage)
+		UAnimSequenceBase* CurrentAnimation = LocalAnimMontageInfo.AnimMontage->IsDynamicMontage() ? LocalAnimMontageInfo.AnimMontage->GetFirstAnimReference() : LocalAnimMontageInfo.AnimMontage;
+		if (ClientAnimation == CurrentAnimation)
 		{
+			UAnimMontage* CurrentAnimMontage = LocalAnimMontageInfo.AnimMontage;
+
 			// Set NextSectionName
 			AnimInstance->Montage_SetNextSection(SectionName, NextSectionName, CurrentAnimMontage);
 
@@ -3335,19 +3376,21 @@ void UAbilitySystemComponent::ServerCurrentMontageSetNextSectionName_Implementat
 	}
 }
 
-bool UAbilitySystemComponent::ServerCurrentMontageJumpToSectionName_Validate(UAnimMontage* ClientAnimMontage, FName SectionName)
+bool UAbilitySystemComponent::ServerCurrentMontageJumpToSectionName_Validate(UAnimSequenceBase* ClientAnimation, FName SectionName)
 {
 	return true;
 }
 
-void UAbilitySystemComponent::ServerCurrentMontageJumpToSectionName_Implementation(UAnimMontage* ClientAnimMontage, FName SectionName)
+void UAbilitySystemComponent::ServerCurrentMontageJumpToSectionName_Implementation(UAnimSequenceBase* ClientAnimation, FName SectionName)
 {
 	UAnimInstance* AnimInstance = AbilityActorInfo.IsValid() ? AbilityActorInfo->GetAnimInstance() : nullptr;
 	if (AnimInstance)
 	{
-		UAnimMontage* CurrentAnimMontage = LocalAnimMontageInfo.AnimMontage;
-		if (ClientAnimMontage == CurrentAnimMontage)
+		UAnimSequenceBase* CurrentAnimation = LocalAnimMontageInfo.AnimMontage->IsDynamicMontage() ? LocalAnimMontageInfo.AnimMontage->GetFirstAnimReference() : LocalAnimMontageInfo.AnimMontage;
+		if (ClientAnimation == CurrentAnimation)
 		{
+			UAnimMontage* CurrentAnimMontage = LocalAnimMontageInfo.AnimMontage;
+
 			// Set NextSectionName
 			AnimInstance->Montage_JumpToSection(SectionName, CurrentAnimMontage);
 
@@ -3357,10 +3400,14 @@ void UAbilitySystemComponent::ServerCurrentMontageJumpToSectionName_Implementati
 				FGameplayAbilityRepAnimMontage& MutableRepAnimMontageInfo = GetRepAnimMontageInfo_Mutable();
 
 				MutableRepAnimMontageInfo.SectionIdToPlay = 0;
-				if (MutableRepAnimMontageInfo.AnimMontage && SectionName != NAME_None)
+				if (MutableRepAnimMontageInfo.Animation && SectionName != NAME_None)
 				{
-					// we add one so INDEX_NONE can be used in the on rep
-					MutableRepAnimMontageInfo.SectionIdToPlay = MutableRepAnimMontageInfo.AnimMontage->GetSectionIndex(SectionName) + 1;
+					// Only change SectionIdToPlay if the anim montage's source is a montage. Dynamic montages have no sections.
+					if (const UAnimMontage* RepAnimMontage = Cast<UAnimMontage>(MutableRepAnimMontageInfo.Animation))
+					{
+						// we add one so INDEX_NONE can be used in the on rep
+						MutableRepAnimMontageInfo.SectionIdToPlay = RepAnimMontage->GetSectionIndex(SectionName) + 1;
+					}
 				}
 
 				AnimMontage_UpdateReplicatedData();
@@ -3369,21 +3416,23 @@ void UAbilitySystemComponent::ServerCurrentMontageJumpToSectionName_Implementati
 	}
 }
 
-bool UAbilitySystemComponent::ServerCurrentMontageSetPlayRate_Validate(UAnimMontage* ClientAnimMontage, float InPlayRate)
+bool UAbilitySystemComponent::ServerCurrentMontageSetPlayRate_Validate(UAnimSequenceBase* ClientAnimation, float InPlayRate)
 {
 	return true;
 }
 
-void UAbilitySystemComponent::ServerCurrentMontageSetPlayRate_Implementation(UAnimMontage* ClientAnimMontage, float InPlayRate)
+void UAbilitySystemComponent::ServerCurrentMontageSetPlayRate_Implementation(UAnimSequenceBase* ClientAnimation, float InPlayRate)
 {
 	UAnimInstance* AnimInstance = AbilityActorInfo.IsValid() ? AbilityActorInfo->GetAnimInstance() : nullptr;
 	if (AnimInstance)
 	{
-		UAnimMontage* CurrentAnimMontage = LocalAnimMontageInfo.AnimMontage;
-		if (ClientAnimMontage == CurrentAnimMontage)
+		UAnimSequenceBase* CurrentAnimation = LocalAnimMontageInfo.AnimMontage->IsDynamicMontage() ? LocalAnimMontageInfo.AnimMontage->GetFirstAnimReference() : LocalAnimMontageInfo.AnimMontage;
+		if (ClientAnimation == CurrentAnimation)
 		{
+			UAnimMontage* CurrentAnimMontage = LocalAnimMontageInfo.AnimMontage;
+
 			// Set PlayRate
-			AnimInstance->Montage_SetPlayRate(LocalAnimMontageInfo.AnimMontage, InPlayRate);
+			AnimInstance->Montage_SetPlayRate(CurrentAnimMontage, InPlayRate);
 
 			// Update replicated version for Simulated Proxies if we are on the server.
 			if (IsOwnerActorAuthoritative())
@@ -3394,6 +3443,20 @@ void UAbilitySystemComponent::ServerCurrentMontageSetPlayRate_Implementation(UAn
 	}
 }
 
+UAnimMontage* UAbilitySystemComponent::PlaySlotAnimationAsDynamicMontage(UGameplayAbility* AnimatingAbility, FGameplayAbilityActivationInfo ActivationInfo, UAnimSequenceBase* AnimAsset, FName SlotName, float BlendInTime, float BlendOutTime, float InPlayRate, float StartTimeSeconds)
+{
+	UAnimMontage* DynamicMontage = UAnimMontage::CreateSlotAnimationAsDynamicMontage(AnimAsset, SlotName, BlendInTime, BlendOutTime, InPlayRate, 1, -1.0f, 0.0f);
+	PlayMontage(AnimatingAbility, ActivationInfo, DynamicMontage, InPlayRate, NAME_None, StartTimeSeconds);
+	return DynamicMontage;
+}
+
+UAnimMontage* UAbilitySystemComponent::PlaySlotAnimationAsDynamicMontageSimulated(UAnimSequenceBase* AnimAsset, FName SlotName, float BlendInTime, float BlendOutTime, float InPlayRate)
+{
+	UAnimMontage* DynamicMontage = UAnimMontage::CreateSlotAnimationAsDynamicMontage(AnimAsset, SlotName, BlendInTime, BlendOutTime, InPlayRate, 1, -1.0f, 0.0f);
+	PlayMontageSimulated(DynamicMontage, InPlayRate, NAME_None);
+	return DynamicMontage;
+}
+
 UAnimMontage* UAbilitySystemComponent::GetCurrentMontage() const
 {
 	UAnimInstance* AnimInstance = AbilityActorInfo.IsValid() ? AbilityActorInfo->GetAnimInstance() : nullptr;
diff --git a/Engine/Plugins/Runtime/GameplayAbilities/Source/GameplayAbilities/Private/GameplayAbilityRepAnimMontage.cpp b/Engine/Plugins/Runtime/GameplayAbilities/Source/GameplayAbilities/Private/GameplayAbilityRepAnimMontage.cpp
--- a/Engine/Plugins/Runtime/GameplayAbilities/Source/GameplayAbilities/Private/GameplayAbilityRepAnimMontage.cpp
+++ b/Engine/Plugins/Runtime/GameplayAbilities/Source/GameplayAbilities/Private/GameplayAbilityRepAnimMontage.cpp
@@ -1,6 +1,8 @@
 // Copyright Epic Games, Inc. All Rights Reserved.
 
 #include "Abilities/GameplayAbilityRepAnimMontage.h"
+#include "Animation/AnimMontage.h"
+#include "Animation/AnimSequenceBase.h"
 
 #include UE_INLINE_GENERATED_CPP_BY_NAME(GameplayAbilityRepAnimMontage)
 
@@ -8,6 +10,16 @@ bool FGameplayAbilityRepAnimMontage::NetSerialize(FArchive& Ar, class UPackageMa
 {
 	Ar.UsingCustomVersion(FEngineNetworkCustomVersion::Guid);
 
+	uint8 bIsMontage = 1;
+	if (Ar.EngineNetVer() > FEngineNetworkCustomVersion::DynamicMontageSerialization)
+	{
+		if (Ar.IsSaving())
+		{
+			bIsMontage = Animation && Animation->IsA<UAnimMontage>();
+		}
+		Ar.SerializeBits(&bIsMontage, 1);
+	}
+
 	uint8 RepPosition = bRepPosition;
 	Ar.SerializeBits(&RepPosition, 1);
 	if (RepPosition)
@@ -60,16 +72,25 @@ bool FGameplayAbilityRepAnimMontage::NetSerialize(FArchive& Ar, class UPackageMa
 	Ar.SerializeBits(&SkipPlayRate, 1);
 	bSkipPlayRate = SkipPlayRate & 1;
 
-	Ar << AnimMontage;
+	Ar << Animation;
 	Ar << PlayRate;
 	Ar << BlendTime;
 	Ar << NextSectionID;
+
 	if (Ar.EngineNetVer() >= FEngineNetworkCustomVersion::MontagePlayInstIdSerialization)
 	{
 		Ar << PlayInstanceId;
 	}
+	
 	PredictionKey.NetSerialize(Ar, Map, bOutSuccess);
 
+	if (!bIsMontage)
+	{
+		ensure(Ar.EngineNetVer() > FEngineNetworkCustomVersion::DynamicMontageSerialization);
+		Ar << BlendOutTime;
+		Ar << SlotName;
+	}
+
 	bOutSuccess = true;
 	return true;
 }
@@ -82,3 +103,8 @@ void FGameplayAbilityRepAnimMontage::SetRepAnimPositionMethod(ERepAnimPositionMe
 	case ERepAnimPositionMethod::CurrentSectionId: bRepPosition = false; break;
 	}
 }
+
+UAnimMontage* FGameplayAbilityRepAnimMontage::GetAnimMontage() const
+{
+	return Cast<UAnimMontage>(Animation);
+}
diff --git a/Engine/Plugins/Runtime/GameplayAbilities/Source/GameplayAbilities/Private/Serialization/GameplayAbilityRepAnimMontageNetSerializer.cpp b/Engine/Plugins/Runtime/GameplayAbilities/Source/GameplayAbilities/Private/Serialization/GameplayAbilityRepAnimMontageNetSerializer.cpp
--- a/Engine/Plugins/Runtime/GameplayAbilities/Source/GameplayAbilities/Private/Serialization/GameplayAbilityRepAnimMontageNetSerializer.cpp
+++ b/Engine/Plugins/Runtime/GameplayAbilities/Source/GameplayAbilities/Private/Serialization/GameplayAbilityRepAnimMontageNetSerializer.cpp
@@ -6,6 +6,8 @@
 
 #if UE_WITH_IRIS
 
+#include "Animation/AnimSequenceBase.h"
+#include "Animation/AnimMontage.h"
 #include "Abilities/GameplayAbilityTypes.h"
 #include "Abilities/GameplayAbilityRepAnimMontage.h"
 #include "Net/Core/NetBitArray.h"
@@ -31,17 +33,22 @@ struct FGameplayAbilityRepAnimMontageNetSerializer
 	static constexpr bool bHasConnectionSpecificSerialization = true; // One of our members requires connection specific serialization
 	static constexpr bool bHasCustomNetReference = true; // We have object references that are not directly accessible
 	static constexpr bool bUseSerializerIsEqual = true; // Since FGameplayAbilityRepAnimMontageNetSerializer conditionally replicates position and SectionId we need to use the IsEqual function provided by the NetSerializer when comparing the property
+	static constexpr bool bHasDynamicState = true;
 
 	// Types
 	enum EReplicationFlags : uint8
 	{
 		ReplicatePosition = 1U,
+		IsAnimMontage = ReplicatePosition << 1U,
 	};
 
 	struct FQuantizedType
 	{
-		alignas(16) uint8 GameplayAbilityRepAnimMontage[40];
+		alignas(16) uint8 GameplayAbilityRepAnimMontage[80];
+
 		float Position;
+		float BlendOutTime;
+
 		uint8 ReplicationFlags;
 		uint8 SectionIdToPlay;
 	};
@@ -101,7 +108,7 @@ void FGameplayAbilityRepAnimMontageNetSerializer::Serialize(FNetSerializationCon
 	const uint8 ReplicationFlags = Value.ReplicationFlags;
 
 	// We do not really need to serialize this as it is part of the basestate, but it makes it easier to access from quantized state and as long as it is a single bit it does not matter much
-	Writer->WriteBits(ReplicationFlags, 1U);
+	Writer->WriteBits(ReplicationFlags, 2U);
 
 	// Forward to normal StructNetSerializer
 	FNetSerializeArgs NetSerializeArgs = {};
@@ -124,7 +131,21 @@ void FGameplayAbilityRepAnimMontageNetSerializer::Serialize(FNetSerializationCon
 	else
 	{
 		UE_NET_TRACE_SCOPE(SectionIdToPlay, *Writer, Context.GetTraceCollector(), ENetTraceVerbosity::Verbose);
-		Writer->WriteBits(Value.SectionIdToPlay, 7U);
+		if (ReplicationFlags & EReplicationFlags::IsAnimMontage)
+		{
+			Writer->WriteBits(Value.SectionIdToPlay, 7U);
+		}
+	}
+
+	if ((ReplicationFlags & EReplicationFlags::IsAnimMontage) == 0U)
+	{
+		const FNetSerializer& FloatSerializer = UE_NET_GET_SERIALIZER(FFloatNetSerializer);
+
+		FNetSerializeArgs FloatNetSerializeArgs = {};
+		FloatNetSerializeArgs.Version = FloatSerializer.Version;
+		FloatNetSerializeArgs.Source = NetSerializerValuePointer(&Value.BlendOutTime);
+		FloatNetSerializeArgs.NetSerializerConfig = NetSerializerConfigParam(FloatSerializer.DefaultConfig);
+		FloatSerializer.Serialize(Context, FloatNetSerializeArgs);
 	}
 }
 
@@ -134,7 +155,7 @@ void FGameplayAbilityRepAnimMontageNetSerializer::Deserialize(FNetSerializationC
 
 	FNetBitStreamReader* Reader = Context.GetBitStreamReader();
 
-	const uint8 ReplicationFlags = Reader->ReadBits(1U);
+	const uint8 ReplicationFlags = Reader->ReadBits(2U);
 	Target.ReplicationFlags = ReplicationFlags;
 
 	// Forward to normal StructNetSerializer
@@ -162,7 +183,25 @@ void FGameplayAbilityRepAnimMontageNetSerializer::Deserialize(FNetSerializationC
 		UE_NET_TRACE_SCOPE(SectionIdToPlay, *Reader, Context.GetTraceCollector(), ENetTraceVerbosity::Verbose);
 
 		Target.Position = 0.f;
-		Target.SectionIdToPlay = Reader->ReadBits(7U);
+		if (ReplicationFlags & EReplicationFlags::IsAnimMontage)
+		{
+			Target.SectionIdToPlay = Reader->ReadBits(7U);
+		}
+		else
+		{
+			Target.SectionIdToPlay = 0U;
+		}
+	}
+
+	if ((ReplicationFlags & EReplicationFlags::IsAnimMontage) == 0U)
+	{
+		const FNetSerializer& FloatSerializer = UE_NET_GET_SERIALIZER(FFloatNetSerializer);
+
+		FNetDeserializeArgs FloatNetDeserializeArgs = {};
+		FloatNetDeserializeArgs.Version = FloatSerializer.Version;
+		FloatNetDeserializeArgs.Target = NetSerializerValuePointer(&Target.BlendOutTime);
+		FloatNetDeserializeArgs.NetSerializerConfig = NetSerializerConfigParam(FloatSerializer.DefaultConfig);
+		FloatSerializer.Deserialize(Context, FloatNetDeserializeArgs);		
 	}
 }
 
@@ -177,7 +216,7 @@ void FGameplayAbilityRepAnimMontageNetSerializer::SerializeDelta(FNetSerializati
 	const uint32 ReplicationFlags = Value.ReplicationFlags;
 
 	// We do not really need to serialize this as it is part of the basestate, but it makes it easier to access and as long as it is a single bit it does not matter much
-	Writer->WriteBits(ReplicationFlags, 1U);
+	Writer->WriteBits(ReplicationFlags, 2U);
 
 	// Forward to normal StructNetSerializer
 	FNetSerializeDeltaArgs NetSerializeDeltaArgs = {};
@@ -205,11 +244,26 @@ void FGameplayAbilityRepAnimMontageNetSerializer::SerializeDelta(FNetSerializati
 	else
 	{
 		UE_NET_TRACE_SCOPE(SectionIdToPlay, *Writer, Context.GetTraceCollector(), ENetTraceVerbosity::Verbose);
-		if (Writer->WriteBool(Value.SectionIdToPlay != PrevValue.SectionIdToPlay))
+		if (ReplicationFlags & EReplicationFlags::IsAnimMontage)
 		{
-			Writer->WriteBits(Value.SectionIdToPlay, 7U);
+			if (Writer->WriteBool(Value.SectionIdToPlay != PrevValue.SectionIdToPlay))
+			{
+				Writer->WriteBits(Value.SectionIdToPlay, 7U);
+			}
 		}
 	}
+
+	if ((ReplicationFlags & EReplicationFlags::IsAnimMontage) == 0U)
+	{
+		const FNetSerializer& FloatSerializer = UE_NET_GET_SERIALIZER(FFloatNetSerializer);
+
+		FNetSerializeDeltaArgs FloatNetSerializeDeltaArgs = {};
+		FloatNetSerializeDeltaArgs.Version = FloatSerializer.Version;
+		FloatNetSerializeDeltaArgs.Source = NetSerializerValuePointer(&Value.BlendOutTime);
+		FloatNetSerializeDeltaArgs.Prev = NetSerializerValuePointer(&PrevValue.BlendOutTime);
+		FloatNetSerializeDeltaArgs.NetSerializerConfig = NetSerializerConfigParam(FloatSerializer.DefaultConfig);
+		FloatSerializer.SerializeDelta(Context, FloatNetSerializeDeltaArgs);
+	}
 }
 
 void FGameplayAbilityRepAnimMontageNetSerializer::DeserializeDelta(FNetSerializationContext& Context, const FNetDeserializeDeltaArgs& Args)
@@ -219,7 +273,7 @@ void FGameplayAbilityRepAnimMontageNetSerializer::DeserializeDelta(FNetSerializa
 
 	FNetBitStreamReader* Reader = Context.GetBitStreamReader();
 
-	const uint32 ReplicationFlags = Reader->ReadBits(1U);
+	const uint32 ReplicationFlags = Reader->ReadBits(2U);
 	Target.ReplicationFlags = ReplicationFlags;
 
 	FNetDeserializeDeltaArgs NetDeserializeDeltaArgs = {};
@@ -248,15 +302,35 @@ void FGameplayAbilityRepAnimMontageNetSerializer::DeserializeDelta(FNetSerializa
 	{
 		UE_NET_TRACE_SCOPE(SectionIdToPlay, *Reader, Context.GetTraceCollector(), ENetTraceVerbosity::Verbose);
 		Target.Position = 0.f;
-		if (Reader->ReadBool())
+
+		if (ReplicationFlags & EReplicationFlags::IsAnimMontage)
 		{
-			Target.SectionIdToPlay = Reader->ReadBits(7U);
+			if (Reader->ReadBool())
+			{
+				Target.SectionIdToPlay = Reader->ReadBits(7U);
+			}
+			else
+			{
+				Target.SectionIdToPlay = Prev.SectionIdToPlay;
+			}
 		}
 		else
 		{
-			Target.SectionIdToPlay = Prev.SectionIdToPlay;
+			Target.SectionIdToPlay = 0;
 		}
 	}
+
+	if ((ReplicationFlags & EReplicationFlags::IsAnimMontage) == 0U)
+	{
+		const FNetSerializer& FloatSerializer = UE_NET_GET_SERIALIZER(FFloatNetSerializer);
+
+		FNetDeserializeDeltaArgs FloatNetDeserializeDeltaArgs = {};
+		FloatNetDeserializeDeltaArgs.Version = FloatSerializer.Version;
+		FloatNetDeserializeDeltaArgs.Target = NetSerializerValuePointer(&Target.BlendOutTime);
+		FloatNetDeserializeDeltaArgs.Prev = NetSerializerValuePointer(&Prev.BlendOutTime);
+		FloatNetDeserializeDeltaArgs.NetSerializerConfig = NetSerializerConfigParam(FloatSerializer.DefaultConfig);
+		FloatSerializer.DeserializeDelta(Context, FloatNetDeserializeDeltaArgs);
+	}
 }
 
 void FGameplayAbilityRepAnimMontageNetSerializer::Quantize(FNetSerializationContext& Context, const FNetQuantizeArgs& Args)
@@ -265,7 +339,10 @@ void FGameplayAbilityRepAnimMontageNetSerializer::Quantize(FNetSerializationCont
 	QuantizedType& TargetValue = *reinterpret_cast<QuantizedType*>(Args.Target);
 	
 	uint32 ReplicationFlags = 0;
+
 	ReplicationFlags |= SourceValue.bRepPosition == 1 ? EReplicationFlags::ReplicatePosition : 0U;
+	ReplicationFlags |= (SourceValue.Animation && SourceValue.Animation->IsA<UAnimMontage>()) ? EReplicationFlags::IsAnimMontage : 0U;
+
 	TargetValue.ReplicationFlags = ReplicationFlags;
 
 	// Forward to normal StructNetSerializer
@@ -283,7 +360,26 @@ void FGameplayAbilityRepAnimMontageNetSerializer::Quantize(FNetSerializationCont
 	else
 	{
 		TargetValue.Position = 0.f;
-		TargetValue.SectionIdToPlay = SourceValue.SectionIdToPlay;
+		if (ReplicationFlags & IsAnimMontage)
+		{
+
+			TargetValue.SectionIdToPlay = SourceValue.SectionIdToPlay;
+		}
+		else
+		{
+			// Always set to zero if not included to be deterministic
+			TargetValue.SectionIdToPlay	= 0;
+		}
+	}
+
+	if ((ReplicationFlags & IsAnimMontage) == 0U)
+	{
+		TargetValue.BlendOutTime = SourceValue.BlendOutTime;
+	}
+	else
+	{
+		// Always set to zero if not included to be deterministic
+		TargetValue.BlendOutTime = 0.f;
 	}
 }
 
@@ -310,9 +406,25 @@ void FGameplayAbilityRepAnimMontageNetSerializer::Dequantize(FNetSerializationCo
 	else
 	{
 		TargetValue.Position = 0.f;
-		TargetValue.SectionIdToPlay = SourceValue.SectionIdToPlay;
+		if (ReplicationFlags & IsAnimMontage)
+		{
+			TargetValue.SectionIdToPlay = SourceValue.SectionIdToPlay;
+		}
+		else
+		{
+			TargetValue.SectionIdToPlay = 0;
+		}
 		TargetValue.SkipPositionCorrection = 1;
 	}
+
+	if (ReplicationFlags & IsAnimMontage)
+	{
+		TargetValue.BlendOutTime = 0.f;
+	}
+	else
+	{
+		TargetValue.BlendOutTime = SourceValue.BlendOutTime;
+	}
 }
 
 bool FGameplayAbilityRepAnimMontageNetSerializer::IsEqual(FNetSerializationContext& Context, const FNetIsEqualArgs& Args)
@@ -338,7 +450,7 @@ bool FGameplayAbilityRepAnimMontageNetSerializer::IsEqual(FNetSerializationConte
 			return false;
 		}
 
-		if (Value0.Position != Value1.Position || Value0.SectionIdToPlay != Value1.SectionIdToPlay)
+		if (Value0.Position != Value1.Position || Value0.SectionIdToPlay != Value1.SectionIdToPlay || Value0.BlendOutTime != Value1.BlendOutTime)
 		{
 			return false;
 		}
@@ -368,6 +480,11 @@ bool FGameplayAbilityRepAnimMontageNetSerializer::IsEqual(FNetSerializationConte
 		{
 			return false;
 		}
+
+		if (SourceValue0.BlendOutTime != SourceValue1.BlendOutTime)
+		{
+			return false;
+		}
 	}
 
 	return true;
@@ -433,12 +550,14 @@ void FGameplayAbilityRepAnimMontageNetSerializer::FNetSerializerRegistryDelegate
 
 void FGameplayAbilityRepAnimMontageNetSerializer::FNetSerializerRegistryDelegates::OnPostFreezeNetSerializerRegistry()
 {
-	constexpr SIZE_T ExpectedSizeOfFGameplayAbilityRepAnimMontage = 56;
+#if PLATFORM_WINDOWS && UE_BUILD_DEVELOPMENT && !UE_EDITOR
+	constexpr SIZE_T ExpectedSizeOfFGameplayAbilityRepAnimMontage = 64;
 	constexpr SIZE_T ExpectedAlignOfFGameplayAbilityRepAnimMontage = 8;
 
 	// Do our best to detect changes to FGameplayAbilityRepAnimMontage
 	// If this assert triggers, this implementation must be verified against FGameplayAbilityRepAnimMontage::NetSerializee before updating the size and alignment
 	static_assert(sizeof(FGameplayAbilityRepAnimMontage) == ExpectedSizeOfFGameplayAbilityRepAnimMontage && alignof(FGameplayAbilityRepAnimMontage) == ExpectedAlignOfFGameplayAbilityRepAnimMontage, "FGameplayAbilityRepAnimMontage layout has changed. Might need to update FGameplayAbilityRepAnimMontageNetSerializer to include new data or update the size.");
+#endif
 
 	// Use helper to avoid getting hold of the FGameplayAbilityRepAnimMontageNetSerializer that we are setting up and validating.
 	const UStruct* BaseStruct = FGameplayAbilityRepAnimMontage::StaticStruct();
diff --git a/Engine/Plugins/Runtime/GameplayAbilities/Source/GameplayAbilities/Public/Abilities/GameplayAbilityRepAnimMontage.h b/Engine/Plugins/Runtime/GameplayAbilities/Source/GameplayAbilities/Public/Abilities/GameplayAbilityRepAnimMontage.h
--- a/Engine/Plugins/Runtime/GameplayAbilities/Source/GameplayAbilities/Public/Abilities/GameplayAbilityRepAnimMontage.h
+++ b/Engine/Plugins/Runtime/GameplayAbilities/Source/GameplayAbilities/Public/Abilities/GameplayAbilityRepAnimMontage.h
@@ -6,6 +6,7 @@
 #include "GameplayAbilityRepAnimMontage.generated.h"
 
 class UAnimMontage;
+class UAnimSequenceBase;
 
 /** Enum used by the Ability Rep Anim Montage struct to rep the quantized position or the current section id */
 UENUM()
@@ -21,9 +22,19 @@ struct GAMEPLAYABILITIES_API FGameplayAbilityRepAnimMontage
 {
 	GENERATED_USTRUCT_BODY()
 
+#if WITH_EDITORONLY_DATA
 	/** AnimMontage ref */
+	UPROPERTY(meta = (DeprecatedProperty, DeprecationMessage = "Use the GetAnimMontage function instead"))
+	TObjectPtr<UAnimMontage> AnimMontage_DEPRECATED;
+#endif
+
+	/** Animation ref. When playing a dynamic montage this points to the AnimSequence the montage was created from */
+	UPROPERTY()
+	TObjectPtr<UAnimSequenceBase> Animation;
+
+	/** Optional slot name used by dynamic montages. */
 	UPROPERTY()
-	TObjectPtr<UAnimMontage> AnimMontage;
+	FName SlotName;
 
 	/** Play Rate */
 	UPROPERTY()
@@ -37,6 +48,10 @@ struct GAMEPLAYABILITIES_API FGameplayAbilityRepAnimMontage
 	UPROPERTY()
 	float BlendTime;
 
+	/** Optional blend out used by dynamic montages. */
+	UPROPERTY(NotReplicated)
+	float BlendOutTime;
+
 	/** NextSectionID */
 	UPROPERTY()
 	uint8 NextSectionID;
@@ -65,14 +80,16 @@ struct GAMEPLAYABILITIES_API FGameplayAbilityRepAnimMontage
 	FPredictionKey PredictionKey;
 
 	/** The current section Id used by the montage. Will only be valid if bRepPosition is false */
-	UPROPERTY(NotReplicated)
+	UPROPERTY()
 	uint8 SectionIdToPlay;
 
 	FGameplayAbilityRepAnimMontage()
-	: AnimMontage(nullptr),
+	: Animation(nullptr),
+	SlotName(NAME_None),
 	PlayRate(0.f),
 	Position(0.f),
 	BlendTime(0.f),
+	BlendOutTime(0.f),
 	NextSectionID(0),
 	PlayInstanceId(0),
 	bRepPosition(true),
@@ -86,6 +103,8 @@ struct GAMEPLAYABILITIES_API FGameplayAbilityRepAnimMontage
 	bool NetSerialize(FArchive& Ar, class UPackageMap* Map, bool& bOutSuccess);
 
 	void SetRepAnimPositionMethod(ERepAnimPositionMethod InMethod);
+
+	UAnimMontage* GetAnimMontage() const;
 };
 
 template<>
diff --git a/Engine/Plugins/Runtime/GameplayAbilities/Source/GameplayAbilities/Public/Abilities/Tasks/AbilityTask_PlayAnimAndWait.h b/Engine/Plugins/Runtime/GameplayAbilities/Source/GameplayAbilities/Public/Abilities/Tasks/AbilityTask_PlayAnimAndWait.h
new file mode 100644
--- /dev/null
+++ b/Engine/Plugins/Runtime/GameplayAbilities/Source/GameplayAbilities/Public/Abilities/Tasks/AbilityTask_PlayAnimAndWait.h
@@ -0,0 +1,88 @@
+// Copyright Epic Games, Inc. All Rights Reserved.
+#pragma once
+
+#include "CoreMinimal.h"
+#include "UObject/ObjectMacros.h"
+#include "Animation/AnimInstance.h"
+#include "Abilities/Tasks/AbilityTask.h"
+#include "AbilityTask_PlayAnimAndWait.generated.h"
+
+DECLARE_DYNAMIC_MULTICAST_DELEGATE(FPlayAnimWaitSimpleDelegate);
+
+class UAnimSequence;
+
+UCLASS()
+class GAMEPLAYABILITIES_API UAbilityTask_PlayAnimAndWait : public UAbilityTask
+{
+	GENERATED_BODY()
+
+public:
+	UPROPERTY(BlueprintAssignable)
+	FPlayAnimWaitSimpleDelegate	OnCompleted;
+
+	UPROPERTY(BlueprintAssignable)
+	FPlayAnimWaitSimpleDelegate	OnBlendOut;
+
+	UPROPERTY(BlueprintAssignable)
+	FPlayAnimWaitSimpleDelegate	OnInterrupted;
+
+	UPROPERTY(BlueprintAssignable)
+	FPlayAnimWaitSimpleDelegate	OnCancelled;
+
+	UFUNCTION()
+	void OnMontageBlendingOut(UAnimMontage* Montage, bool bInterrupted);
+
+	UFUNCTION()
+	void OnMontageInterrupted();
+
+	UFUNCTION()
+	void OnMontageEnded(UAnimMontage* Montage, bool bInterrupted);
+
+	/** 
+	 * Start playing an animation montage on the avatar actor and wait for it to finish
+	 * If StopWhenAbilityEnds is true, this montage will be aborted if the ability ends normally. It is always stopped when the ability is explicitly cancelled.
+	 * On normal execution, OnBlendOut is called when the montage is blending out, and OnCompleted when it is completely done playing
+	 * OnInterrupted is called if another montage overwrites this, and OnCancelled is called if the ability or task is cancelled
+	 *
+	 * @param TaskInstanceName Set to override the name of this task, for later querying
+	 * @param MontageToPlay The montage to play on the character
+	 * @param Rate Change to play the montage faster or slower
+	 * @param StartSection If not empty, named montage section to start from
+	 * @param bStopWhenAbilityEnds If true, this montage will be aborted if the ability ends normally. It is always stopped when the ability is explicitly cancelled
+	 * @param AnimRootMotionTranslationScale Change to modify size of root motion or set to 0 to block it entirely
+	 * @param StartTimeSeconds Starting time offset in montage, this will be overridden by StartSection if that is also set
+	 */
+	UFUNCTION(BlueprintCallable, Category="Ability|Tasks", meta = (DisplayName="Play Anim And Wait",
+		HidePin = "OwningAbility", DefaultToSelf = "OwningAbility", BlueprintInternalUseOnly = "TRUE"))
+	static UAbilityTask_PlayAnimAndWait* CreatePlayAnimAndWaitProxy(UGameplayAbility* OwningAbility,
+		FName TaskInstanceName, UAnimSequence* AnimSequence, FName SlotName = NAME_None, float BlendInTime = 0.25f, float BlendOutTime = 0.25f, float InPlayRate = 1.0f, float StartTimeSeconds = 0.0f, bool bStopWhenAbilityEnds = true, float AnimRootMotionTranslationScale = 1.0f);
+
+	virtual void Activate() override;
+
+	/** Called when the ability is asked to cancel from an outside node. What this means depends on the individual task. By default, this does nothing other than ending the task. */
+	virtual void ExternalCancel() override;
+
+	virtual FString GetDebugString() const override;
+
+protected:
+
+	virtual void OnDestroy(bool AbilityEnded) override;
+
+	/** Checks if the ability is playing a montage and stops that montage, returns true if a montage was stopped, false if not. */
+	bool StopPlayingMontage();
+
+	FOnMontageBlendingOutStarted BlendingOutDelegate;
+	FOnMontageEnded MontageEndedDelegate;
+	FDelegateHandle InterruptedHandle;
+
+	UPROPERTY()
+	TObjectPtr<UAnimSequence> AnimSequenceToPlay;
+
+	FName SlotName = NAME_None;
+	float BlendInTime = 0.25f;
+	float BlendOutTime = 0.25f;
+	float PlayRate = 1.0f;
+	float AnimRootMotionTranslationScale = 1.0f;
+	float StartTimeSeconds = 0.0f;
+	bool bStopWhenAbilityEnds = true;
+};
diff --git a/Engine/Plugins/Runtime/GameplayAbilities/Source/GameplayAbilities/Public/AbilitySystemComponent.h b/Engine/Plugins/Runtime/GameplayAbilities/Source/GameplayAbilities/Public/AbilitySystemComponent.h
--- a/Engine/Plugins/Runtime/GameplayAbilities/Source/GameplayAbilities/Public/AbilitySystemComponent.h
+++ b/Engine/Plugins/Runtime/GameplayAbilities/Source/GameplayAbilities/Public/AbilitySystemComponent.h
@@ -34,6 +34,7 @@ class AGameplayAbilityTargetActor;
 class AHUD;
 class FDebugDisplayInfo;
 class UAnimMontage;
+class UAnimSequenceBase;
 class UCanvas;
 class UInputComponent;
 
@@ -1413,9 +1414,13 @@ class GAMEPLAYABILITIES_API UAbilitySystemComponent : public UGameplayTasksCompo
 
 	/** Plays a montage and handles replication and prediction based on passed in ability/activation info */
 	virtual float PlayMontage(UGameplayAbility* AnimatingAbility, FGameplayAbilityActivationInfo ActivationInfo, UAnimMontage* Montage, float InPlayRate, FName StartSectionName = NAME_None, float StartTimeSeconds = 0.0f);
+	
+	virtual UAnimMontage* PlaySlotAnimationAsDynamicMontage(UGameplayAbility* AnimatingAbility, FGameplayAbilityActivationInfo ActivationInfo, UAnimSequenceBase* AnimAsset, FName SlotName, float BlendInTime, float BlendOutTime, float InPlayRate = 1.f, float StartTimeSeconds = 0.0f);
 
 	/** Plays a montage without updating replication/prediction structures. Used by simulated proxies when replication tells them to play a montage. */
 	virtual float PlayMontageSimulated(UAnimMontage* Montage, float InPlayRate, FName StartSectionName = NAME_None);
+	
+	virtual UAnimMontage* PlaySlotAnimationAsDynamicMontageSimulated(UAnimSequenceBase* AnimAsset, FName SlotName, float BlendInTime, float BlendOutTime, float InPlayRate = 1.f);
 
 	/** Stops whatever montage is currently playing. Expectation is caller should only be stopping it if they are the current animating ability (or have good reason not to check) */
 	virtual void CurrentMontageStop(float OverrideBlendOutTime = -1.0f);
@@ -1778,15 +1783,15 @@ protected:
 
 	/** RPC function called from CurrentMontageSetNextSectopnName, replicates to other clients */
 	UFUNCTION(reliable, server, WithValidation)
-	void ServerCurrentMontageSetNextSectionName(UAnimMontage* ClientAnimMontage, float ClientPosition, FName SectionName, FName NextSectionName);
+	void ServerCurrentMontageSetNextSectionName(UAnimSequenceBase* ClientAnimation, float ClientPosition, FName SectionName, FName NextSectionName);
 
 	/** RPC function called from CurrentMontageJumpToSection, replicates to other clients */
 	UFUNCTION(reliable, server, WithValidation)
-	void ServerCurrentMontageJumpToSectionName(UAnimMontage* ClientAnimMontage, FName SectionName);
+	void ServerCurrentMontageJumpToSectionName(UAnimSequenceBase* ClientAnimation, FName SectionName);
 
 	/** RPC function called from CurrentMontageSetPlayRate, replicates to other clients */
 	UFUNCTION(reliable, server, WithValidation)
-	void ServerCurrentMontageSetPlayRate(UAnimMontage* ClientAnimMontage, float InPlayRate);
+	void ServerCurrentMontageSetPlayRate(UAnimSequenceBase* ClientAnimation, float InPlayRate);
 
 	/** Abilities that are triggered from a gameplay event */
 	TMap<FGameplayTag, TArray<FGameplayAbilitySpecHandle > > GameplayEventTriggeredAbilities;
diff --git a/Engine/Source/Runtime/Core/Public/Misc/EngineNetworkCustomVersion.h b/Engine/Source/Runtime/Core/Public/Misc/EngineNetworkCustomVersion.h
--- a/Engine/Source/Runtime/Core/Public/Misc/EngineNetworkCustomVersion.h
+++ b/Engine/Source/Runtime/Core/Public/Misc/EngineNetworkCustomVersion.h
@@ -84,8 +84,9 @@ struct FEngineNetworkCustomVersion
 		SoftObjectPtrNetGuids = 29,				// Bump version to support replicating SoftObjectPtrs by NetGuid instead of raw strings.
 		SubObjectDestroyFlag = 30,				// Bump version to support subobject destruction message flags
 		GameStateReplicatedTimeAsDouble = 31,	// Bump version to support AGameStateBase::ReplicatedWorldTimeSeconds as double instead of float.
+		DynamicMontageSerialization = 32,
 
-		CustomVersions = 32,                    // Bump version to switch to using custom versions
+		CustomVersions = 33,                    // Bump version to switch to using custom versions
 
 		// -----<new versions can be added above this line>-------------------------------------------------
 		VersionPlusOne,
diff --git a/Engine/Source/Runtime/Engine/Classes/Animation/AnimInstance.h b/Engine/Source/Runtime/Engine/Classes/Animation/AnimInstance.h
--- a/Engine/Source/Runtime/Engine/Classes/Animation/AnimInstance.h
+++ b/Engine/Source/Runtime/Engine/Classes/Animation/AnimInstance.h
@@ -662,6 +662,10 @@ public:
 	UFUNCTION(BlueprintPure, Category = "Animation|Montage")
 	ENGINE_API float Montage_GetEffectivePlayRate(const UAnimMontage* Montage) const;
 
+	/** Returns true if there is an animation montage is currently active and playing that was created from the provided animation. */
+	UFUNCTION(BlueprintPure, Category = "Animation|Montage")
+	ENGINE_API bool DynamicMontage_IsPlayingFrom(const UAnimSequenceBase* Animation) const;
+
 	/*********************************************************************************************
 	* AnimMontage sync. See notes in AnimMontage.h
 	********************************************************************************************* */
diff --git a/Engine/Source/Runtime/Engine/Classes/Animation/AnimMontage.h b/Engine/Source/Runtime/Engine/Classes/Animation/AnimMontage.h
--- a/Engine/Source/Runtime/Engine/Classes/Animation/AnimMontage.h
+++ b/Engine/Source/Runtime/Engine/Classes/Animation/AnimMontage.h
@@ -716,6 +716,12 @@ class UAnimMontage : public UAnimCompositeBase
 	/** return true if valid slot */
 	ENGINE_API bool IsValidSlot(FName InSlotName) const;
 
+	UFUNCTION(BlueprintPure, Category = "Montage")
+	ENGINE_API bool IsDynamicMontage() const;
+
+	UFUNCTION(BlueprintPure, Category = "Montage")
+	ENGINE_API UAnimSequenceBase* GetFirstAnimReference() const;
+
 public:
 	//~ Begin UObject Interface
 	virtual void PostLoad() override;
diff --git a/Engine/Source/Runtime/Engine/Classes/GameFramework/Character.h b/Engine/Source/Runtime/Engine/Classes/GameFramework/Character.h
--- a/Engine/Source/Runtime/Engine/Classes/GameFramework/Character.h
+++ b/Engine/Source/Runtime/Engine/Classes/GameFramework/Character.h
@@ -45,9 +45,15 @@ struct FRepRootMotionMontage
 	UPROPERTY()
 	bool bIsActive = false;
 
+#if WITH_EDITORONLY_DATA
 	/** AnimMontage providing Root Motion */
+	UPROPERTY(meta = (DeprecatedProperty, DeprecationMessage = "Use the GetAnimMontage function instead"))
+	TObjectPtr<UAnimMontage> AnimMontage_DEPRECATED = nullptr;
+#endif
+	
+	/** Animation providing Root Motion */
 	UPROPERTY()
-	TObjectPtr<UAnimMontage> AnimMontage = nullptr;
+	TObjectPtr<UAnimSequenceBase> Animation = nullptr;
 
 	/** Track position of Montage */
 	UPROPERTY()
@@ -93,15 +99,17 @@ struct FRepRootMotionMontage
 	void Clear()
 	{
 		bIsActive = false;
-		AnimMontage = nullptr;
+		Animation = nullptr;
 		AuthoritativeRootMotion.Clear();
 	}
 
 	/** Is Valid - animation root motion only */
 	bool HasRootMotion() const
 	{
-		return (AnimMontage != nullptr);
+		return (Animation != nullptr);
 	}
+
+	UAnimMontage* GetAnimMontage() const;
 };
 
 USTRUCT()
diff --git a/Engine/Source/Runtime/Engine/Private/Animation/AnimInstance.cpp b/Engine/Source/Runtime/Engine/Private/Animation/AnimInstance.cpp
--- a/Engine/Source/Runtime/Engine/Private/Animation/AnimInstance.cpp
+++ b/Engine/Source/Runtime/Engine/Private/Animation/AnimInstance.cpp
@@ -2773,6 +2773,29 @@ float UAnimInstance::Montage_GetEffectivePlayRate(const UAnimMontage* Montage) c
 	return 0.f;
 }
 
+bool UAnimInstance::DynamicMontage_IsPlayingFrom(const UAnimSequenceBase* Animation) const
+{
+	if (!Animation)
+	{
+		return false;
+	}
+
+	if (const UAnimMontage* AnimMontage = Cast<UAnimMontage>(Animation))
+	{
+		return Montage_IsPlaying(AnimMontage);
+	}
+
+	for (const TPair<UAnimMontage*, FAnimMontageInstance*>& ActiveMontage : ActiveMontagesMap)
+	{
+		if (ActiveMontage.Key->IsDynamicMontage() && ActiveMontage.Key->GetFirstAnimReference() == Animation)
+		{
+			return ActiveMontage.Value->IsPlaying();
+		}
+	}
+
+	return false;
+}
+
 void UAnimInstance::MontageSync_Follow(const UAnimMontage* MontageFollower, const UAnimInstance* OtherAnimInstance, const UAnimMontage* MontageLeader)
 {
 	if (!MontageFollower || !OtherAnimInstance || !MontageLeader)
diff --git a/Engine/Source/Runtime/Engine/Private/Animation/AnimMontage.cpp b/Engine/Source/Runtime/Engine/Private/Animation/AnimMontage.cpp
--- a/Engine/Source/Runtime/Engine/Private/Animation/AnimMontage.cpp
+++ b/Engine/Source/Runtime/Engine/Private/Animation/AnimMontage.cpp
@@ -96,6 +96,21 @@ bool UAnimMontage::IsValidSlot(FName InSlotName) const
 	return false;
 }
 
+bool UAnimMontage::IsDynamicMontage() const
+{
+	return GetPackage() == GetTransientPackage();
+}
+
+UAnimSequenceBase* UAnimMontage::GetFirstAnimReference() const
+{
+	if(!SlotAnimTracks.IsEmpty() && !SlotAnimTracks[0].AnimTrack.AnimSegments.IsEmpty())
+	{
+		return SlotAnimTracks[0].AnimTrack.AnimSegments[0].GetAnimReference();
+	}
+
+	return nullptr;
+}
+
 const FAnimTrack* UAnimMontage::GetAnimationData(FName InSlotName) const
 {
 	for (int32 I=0; I<SlotAnimTracks.Num(); ++I)
diff --git a/Engine/Source/Runtime/Engine/Private/Character.cpp b/Engine/Source/Runtime/Engine/Private/Character.cpp
--- a/Engine/Source/Runtime/Engine/Private/Character.cpp
+++ b/Engine/Source/Runtime/Engine/Private/Character.cpp
@@ -1388,9 +1388,11 @@ bool ACharacter::CanUseRootMotionRepMove(const FSimulatedRootMotionReplicatedMov
 	if( GetWorld()->TimeSince(RootMotionRepMove.Time) <= 0.5f )
 	{
 		// Make sure montage being played matched between client and server.
-		if( RootMotionRepMove.RootMotion.AnimMontage && (RootMotionRepMove.RootMotion.AnimMontage == ClientMontageInstance.Montage) )
+		UAnimSequenceBase* ClientAnimation = ClientMontageInstance.Montage && ClientMontageInstance.Montage->IsDynamicMontage() ? ClientMontageInstance.Montage->GetFirstAnimReference() : ClientMontageInstance.Montage;
+		if (RootMotionRepMove.RootMotion.Animation && RootMotionRepMove.RootMotion.Animation == ClientAnimation)
 		{
-			UAnimMontage * AnimMontage = ClientMontageInstance.Montage;
+			UAnimMontage* AnimMontage = ClientMontageInstance.Montage;
+
 			const float ServerPosition = RootMotionRepMove.RootMotion.Position;
 			const float ClientPosition = ClientMontageInstance.GetPosition();
 			const float DeltaPosition = (ClientPosition - ServerPosition);
@@ -1530,12 +1532,12 @@ void ACharacter::PreReplication( IRepChangedPropertyTracker & ChangedPropertyTra
 		RepRootMotion.MovementBaseBoneName = BasedMovement.BoneName;
 		if (RootMotionMontageInstance)
 		{
-			RepRootMotion.AnimMontage		= RootMotionMontageInstance->Montage;
+			RepRootMotion.Animation = RootMotionMontageInstance->Montage->IsDynamicMontage() ? RootMotionMontageInstance->Montage->GetFirstAnimReference() : RootMotionMontageInstance->Montage;
 			RepRootMotion.Position			= RootMotionMontageInstance->GetPosition();
 		}
 		else
 		{
-			RepRootMotion.AnimMontage = nullptr;
+			RepRootMotion.Animation = nullptr;
 		}
 
 		RepRootMotion.AuthoritativeRootMotion = CharacterMovement->CurrentRootMotion;
@@ -1911,3 +1913,7 @@ void ACharacter::ApplyAsyncOutput(const FCharacterAsyncOutput& Output)
 	}
 }
 
+UAnimMontage* FRepRootMotionMontage::GetAnimMontage() const
+{
+	return Cast<UAnimMontage>(Animation);
+}
\ No newline at end of file
