Temporarily disable FGameplayAbilityRepAnimMontage size static assert for build health/until better check in place

diff --git a/Engine/Plugins/Runtime/GameplayAbilities/Source/GameplayAbilities/Private/Serialization/GameplayAbilityRepAnimMontageNetSerializer.cpp b/Engine/Plugins/Runtime/GameplayAbilities/Source/GameplayAbilities/Private/Serialization/GameplayAbilityRepAnimMontageNetSerializer.cpp
--- a/Engine/Plugins/Runtime/GameplayAbilities/Source/GameplayAbilities/Private/Serialization/GameplayAbilityRepAnimMontageNetSerializer.cpp
+++ b/Engine/Plugins/Runtime/GameplayAbilities/Source/GameplayAbilities/Private/Serialization/GameplayAbilityRepAnimMontageNetSerializer.cpp
@@ -550,6 +550,7 @@ void FGameplayAbilityRepAnimMontageNetSerializer::FNetSerializerRegistryDelegate
 
 void FGameplayAbilityRepAnimMontageNetSerializer::FNetSerializerRegistryDelegates::OnPostFreezeNetSerializerRegistry()
 {
+	/*
 #if PLATFORM_WINDOWS && UE_BUILD_DEVELOPMENT && !UE_EDITOR
 	constexpr SIZE_T ExpectedSizeOfFGameplayAbilityRepAnimMontage = 64;
 	constexpr SIZE_T ExpectedAlignOfFGameplayAbilityRepAnimMontage = 8;
@@ -558,6 +559,7 @@ void FGameplayAbilityRepAnimMontageNetSerializer::FNetSerializerRegistryDelegate
 	// If this assert triggers, this implementation must be verified against FGameplayAbilityRepAnimMontage::NetSerializee before updating the size and alignment
 	static_assert(sizeof(FGameplayAbilityRepAnimMontage) == ExpectedSizeOfFGameplayAbilityRepAnimMontage && alignof(FGameplayAbilityRepAnimMontage) == ExpectedAlignOfFGameplayAbilityRepAnimMontage, "FGameplayAbilityRepAnimMontage layout has changed. Might need to update FGameplayAbilityRepAnimMontageNetSerializer to include new data or update the size.");
 #endif
+	*/
 
 	// Use helper to avoid getting hold of the FGameplayAbilityRepAnimMontageNetSerializer that we are setting up and validating.
 	const UStruct* BaseStruct = FGameplayAbilityRepAnimMontage::StaticStruct();
