Iris
- Fixed ordering issue when triggering gameplaycues from MinimalGameplayCueReplicationProxy.

diff --git a/Engine/Plugins/Runtime/GameplayAbilities/Source/GameplayAbilities/Private/Serialization/MinimalGameplayCueReplicationProxyReplicationFragment.cpp b/Engine/Plugins/Runtime/GameplayAbilities/Source/GameplayAbilities/Private/Serialization/MinimalGameplayCueReplicationProxyReplicationFragment.cpp
--- a/Engine/Plugins/Runtime/GameplayAbilities/Source/GameplayAbilities/Private/Serialization/MinimalGameplayCueReplicationProxyReplicationFragment.cpp
+++ b/Engine/Plugins/Runtime/GameplayAbilities/Source/GameplayAbilities/Private/Serialization/MinimalGameplayCueReplicationProxyReplicationFragment.cpp
@@ -79,6 +79,8 @@ void FMinimalGameplayCueReplicationProxyReplicationFragment::ApplyReplicatedStat
 	DequantizeArgs.NetSerializerConfig = ReplicationStateDescriptor->MemberSerializerDescriptors[0].SerializerConfig;
 	const FNetSerializer* Serializer = ReplicationStateDescriptor->MemberSerializerDescriptors[0].Serializer;
 	Serializer->Dequantize(*ApplyContext.NetSerializationContext, DequantizeArgs);
+
+	MimicMinimalGameplayCueReplicationProxyReceiveLogic(ApplyContext);
 }
 
 bool FMinimalGameplayCueReplicationProxyReplicationFragment::PollReplicatedState(EReplicationFragmentPollFlags PollOption)
@@ -102,11 +104,10 @@ bool FMinimalGameplayCueReplicationProxyReplicationFragment::PollReplicatedState
 
 void FMinimalGameplayCueReplicationProxyReplicationFragment::CallRepNotifies(FReplicationStateApplyContext& Context)
 {
-	MimicMinimalGameplayCueReplicationProxyReceiveLogic(Context);
 	CallRepNotify(Context);
 }
 
-void FMinimalGameplayCueReplicationProxyReplicationFragment::MimicMinimalGameplayCueReplicationProxyReceiveLogic(FReplicationStateApplyContext& Context)
+void FMinimalGameplayCueReplicationProxyReplicationFragment::MimicMinimalGameplayCueReplicationProxyReceiveLogic(FReplicationStateApplyContext& Context) const
 {
 	uint8* ExternalStatePointer = reinterpret_cast<uint8*>(Owner) + ReplicationStateDescriptor->MemberProperties[0]->GetOffset_ForGC();
 	FMinimalGameplayCueReplicationProxy* ExternalSourceState = reinterpret_cast<FMinimalGameplayCueReplicationProxy*>(ExternalStatePointer);
diff --git a/Engine/Plugins/Runtime/GameplayAbilities/Source/GameplayAbilities/Public/Serialization/MinimalGameplayCueReplicationProxyReplicationFragment.h b/Engine/Plugins/Runtime/GameplayAbilities/Source/GameplayAbilities/Public/Serialization/MinimalGameplayCueReplicationProxyReplicationFragment.h
--- a/Engine/Plugins/Runtime/GameplayAbilities/Source/GameplayAbilities/Public/Serialization/MinimalGameplayCueReplicationProxyReplicationFragment.h
+++ b/Engine/Plugins/Runtime/GameplayAbilities/Source/GameplayAbilities/Public/Serialization/MinimalGameplayCueReplicationProxyReplicationFragment.h
@@ -32,7 +32,7 @@ protected:
 	virtual void CallRepNotifies(FReplicationStateApplyContext& Context) override;
 
 private:
-	void MimicMinimalGameplayCueReplicationProxyReceiveLogic(FReplicationStateApplyContext& Context);
+	void MimicMinimalGameplayCueReplicationProxyReceiveLogic(FReplicationStateApplyContext& Context) const;
 	void CallRepNotify(FReplicationStateApplyContext& Context);
 
 private:
