Iris
- Moved execution of mimic old replication logic for MinimalReplicationTagCountMap to Apply, to behave as current replication system

diff --git a/Engine/Plugins/Runtime/GameplayAbilities/Source/GameplayAbilities/Private/Serialization/MinimalReplicationTagCountMapReplicationFragment.cpp b/Engine/Plugins/Runtime/GameplayAbilities/Source/GameplayAbilities/Private/Serialization/MinimalReplicationTagCountMapReplicationFragment.cpp
--- a/Engine/Plugins/Runtime/GameplayAbilities/Source/GameplayAbilities/Private/Serialization/MinimalReplicationTagCountMapReplicationFragment.cpp
+++ b/Engine/Plugins/Runtime/GameplayAbilities/Source/GameplayAbilities/Private/Serialization/MinimalReplicationTagCountMapReplicationFragment.cpp
@@ -80,6 +80,8 @@ void FMinimalReplicationTagCountMapReplicationFragment::ApplyReplicatedState(FRe
 	DequantizeArgs.NetSerializerConfig = ReplicationStateDescriptor->MemberSerializerDescriptors[0].SerializerConfig;
 	const FNetSerializer* Serializer = ReplicationStateDescriptor->MemberSerializerDescriptors[0].Serializer;
 	Serializer->Dequantize(*ApplyContext.NetSerializationContext, DequantizeArgs);
+
+	MimicMinimalReplicationTagCountMapReceiveLogic(ApplyContext);
 }
 
 bool FMinimalReplicationTagCountMapReplicationFragment::PollReplicatedState(EReplicationFragmentPollFlags PollOption)
@@ -100,11 +102,10 @@ bool FMinimalReplicationTagCountMapReplicationFragment::PollReplicatedState(ERep
 
 void FMinimalReplicationTagCountMapReplicationFragment::CallRepNotifies(FReplicationStateApplyContext& Context)
 {
-	MimicMinimalReplicationTagCountMapReceiveLogic(Context);
 	CallRepNotify(Context);
 }
 
-void FMinimalReplicationTagCountMapReplicationFragment::MimicMinimalReplicationTagCountMapReceiveLogic(FReplicationStateApplyContext& Context)
+void FMinimalReplicationTagCountMapReplicationFragment::MimicMinimalReplicationTagCountMapReceiveLogic(FReplicationStateApplyContext& Context) const
 {
 	uint8* ExternalStatePointer = reinterpret_cast<uint8*>(Owner) + ReplicationStateDescriptor->MemberProperties[0]->GetOffset_ForGC();
 	FMinimalReplicationTagCountMap* ExternalSourceState = reinterpret_cast<FMinimalReplicationTagCountMap*>(ExternalStatePointer);
diff --git a/Engine/Plugins/Runtime/GameplayAbilities/Source/GameplayAbilities/Private/Serialization/MinimalReplicationTagCountMapReplicationFragment.h b/Engine/Plugins/Runtime/GameplayAbilities/Source/GameplayAbilities/Private/Serialization/MinimalReplicationTagCountMapReplicationFragment.h
--- a/Engine/Plugins/Runtime/GameplayAbilities/Source/GameplayAbilities/Private/Serialization/MinimalReplicationTagCountMapReplicationFragment.h
+++ b/Engine/Plugins/Runtime/GameplayAbilities/Source/GameplayAbilities/Private/Serialization/MinimalReplicationTagCountMapReplicationFragment.h
@@ -32,7 +32,7 @@ protected:
 	virtual void CallRepNotifies(FReplicationStateApplyContext& Context) override;
 
 private:
-	void MimicMinimalReplicationTagCountMapReceiveLogic(FReplicationStateApplyContext& Context);
+	void MimicMinimalReplicationTagCountMapReceiveLogic(FReplicationStateApplyContext& Context) const;
 	void CallRepNotify(FReplicationStateApplyContext& Context);
 
 private:
