Calling ASC->ClearAllAbilities() is now allowed inside ABILITY_SCOPE_LOCK and will be deferred until exiting the scope. This allows GameplayAbilities to call DestroyActor on owner. Previously would crash via check due to attempting to modify the ability list while inside an ABILITY_SCOPE_LOCK.

diff --git a/Engine/Plugins/Runtime/GameplayAbilities/Source/GameplayAbilities/Private/AbilitySystemComponent.cpp b/Engine/Plugins/Runtime/GameplayAbilities/Source/GameplayAbilities/Private/AbilitySystemComponent.cpp
--- a/Engine/Plugins/Runtime/GameplayAbilities/Source/GameplayAbilities/Private/AbilitySystemComponent.cpp
+++ b/Engine/Plugins/Runtime/GameplayAbilities/Source/GameplayAbilities/Private/AbilitySystemComponent.cpp
@@ -63,6 +63,9 @@ UAbilitySystemComponent::UAbilitySystemComponent(const FObjectInitializer& Objec
 	bSuppressGameplayCues = false;
 	bPendingMontageRep = false;
 	AffectedAnimInstanceTag = NAME_None; 
+
+	AbilityScopeLockCount = 0;
+	bAbilityPendingClearAll = false;
 	AbilityLastActivatedTime = 0.f;
 
 	ReplicationMode = EGameplayEffectReplicationMode::Full;
diff --git a/Engine/Plugins/Runtime/GameplayAbilities/Source/GameplayAbilities/Private/AbilitySystemComponent_Abilities.cpp b/Engine/Plugins/Runtime/GameplayAbilities/Source/GameplayAbilities/Private/AbilitySystemComponent_Abilities.cpp
--- a/Engine/Plugins/Runtime/GameplayAbilities/Source/GameplayAbilities/Private/AbilitySystemComponent_Abilities.cpp
+++ b/Engine/Plugins/Runtime/GameplayAbilities/Source/GameplayAbilities/Private/AbilitySystemComponent_Abilities.cpp
@@ -406,7 +406,14 @@ void UAbilitySystemComponent::SetRemoveAbilityOnEnd(FGameplayAbilitySpecHandle A
 
 void UAbilitySystemComponent::ClearAllAbilities()
 {
-	check(AbilityScopeLockCount == 0);	// We should never be calling this from a scoped lock situation.
+	// If this is called inside an ability scope lock, postpone the workload until end of scope.
+	// This was introduced for abilities that trigger their owning actor's destruction on ability
+	// activation.
+	if (AbilityScopeLockCount > 0)
+	{
+		bAbilityPendingClearAll = true;
+		return;
+	}
 
 	if (!IsOwnerActorAuthoritative())
 	{
@@ -439,6 +446,7 @@ void UAbilitySystemComponent::ClearAllAbilities()
 	ActivatableAbilities.MarkArrayDirty();
 
 	CheckForClearedAbilities();
+	bAbilityPendingClearAll = false;
 }
 
 void UAbilitySystemComponent::ClearAllAbilitiesWithInputID(int32 InputID /*= 0*/)
@@ -797,26 +805,45 @@ void UAbilitySystemComponent::IncrementAbilityListLock()
 }
 void UAbilitySystemComponent::DecrementAbilityListLock()
 {
-	if (--AbilityScopeLockCount == 0 &&
-		(AbilityPendingAdds.Num() > 0 || AbilityPendingRemoves.Num() > 0))
+	if (--AbilityScopeLockCount == 0)
 	{
-		FAbilityListLockActiveChange ActiveChange(*this, AbilityPendingAdds, AbilityPendingRemoves);
-
-		for (FGameplayAbilitySpec& Spec : ActiveChange.Adds)
+		if (bAbilityPendingClearAll)
 		{
-			if (Spec.bActivateOnce)
+			ClearAllAbilities();
+
+			// When there are pending adds but also a pending clear-all, prioritize clear-all since ClearAllAbilities() based on an assumption 
+			// that the clear-all is likely end-of-life cleanup. There may be cases where someone intentionally calls ClearAllAbilities() and 
+			// then GiveAbility() within one ability scope lock like an ability that removes all abilities and grants an ability. In the future 
+			// we could support this by keeping a chronological list of pending add/remove/clear-all actions and executing them in order.
+			if (AbilityPendingAdds.Num() > 0)
 			{
-				GiveAbilityAndActivateOnce(Spec, Spec.GameplayEventData.Get());
+				ABILITY_LOG(Warning, TEXT("GiveAbility and ClearAllAbilities were both called within an ability scope lock. Prioritizing clear all abilities by ignoring pending adds."));
+				AbilityPendingAdds.Reset();
 			}
-			else
+
+			// Pending removes are no longer relevant since all abilities have been removed
+			AbilityPendingRemoves.Reset();
+		}
+		else if (AbilityPendingAdds.Num() > 0 || AbilityPendingRemoves.Num() > 0)
+		{
+			FAbilityListLockActiveChange ActiveChange(*this, AbilityPendingAdds, AbilityPendingRemoves);
+
+			for (FGameplayAbilitySpec& Spec : ActiveChange.Adds)
 			{
-				GiveAbility(Spec);
+				if (Spec.bActivateOnce)
+				{
+					GiveAbilityAndActivateOnce(Spec, Spec.GameplayEventData.Get());
+				}
+				else
+				{
+					GiveAbility(Spec);
+				}
 			}
-		}
 
-		for (FGameplayAbilitySpecHandle& Handle : ActiveChange.Removes)
-		{
-			ClearAbility(Handle);
+			for (FGameplayAbilitySpecHandle& Handle : ActiveChange.Removes)
+			{
+				ClearAbility(Handle);
+			}
 		}
 	}
 }
diff --git a/Engine/Plugins/Runtime/GameplayAbilities/Source/GameplayAbilities/Public/AbilitySystemComponent.h b/Engine/Plugins/Runtime/GameplayAbilities/Source/GameplayAbilities/Public/AbilitySystemComponent.h
--- a/Engine/Plugins/Runtime/GameplayAbilities/Source/GameplayAbilities/Public/AbilitySystemComponent.h
+++ b/Engine/Plugins/Runtime/GameplayAbilities/Source/GameplayAbilities/Public/AbilitySystemComponent.h
@@ -1693,9 +1693,14 @@ protected:
 	/** Creates a new instance of an ability, storing it in the spec */
 	virtual UGameplayAbility* CreateNewInstanceOfAbility(FGameplayAbilitySpec& Spec, const UGameplayAbility* Ability);
 
+	/** Indicates how many levels of ABILITY_SCOPE_LOCK() we are in. The ability list may not be modified while AbilityScopeLockCount > 0. */
 	int32 AbilityScopeLockCount;
+	/** Abilities that will be removed when exiting the current ability scope lock. */
 	TArray<FGameplayAbilitySpecHandle, TInlineAllocator<2> > AbilityPendingRemoves;
+	/** Abilities that will be added when exiting the current ability scope lock. */
 	TArray<FGameplayAbilitySpec, TInlineAllocator<2> > AbilityPendingAdds;
+	/** Whether all abilities should be removed when exiting the current ability scope lock. Will be prioritized over pending adds. */
+	bool bAbilityPendingClearAll;
 
 	/** Local World time of the last ability activation. This is used for AFK/idle detection */
 	float AbilityLastActivatedTime;
