AbilitySystemComponent: Add GetOwnedGameplayTags and GetBlockedAbilityTags that return a const-ref to avoid allocating and copying

diff --git a/Engine/Plugins/Runtime/GameplayAbilities/Source/GameplayAbilities/Public/AbilitySystemComponent.h b/Engine/Plugins/Runtime/GameplayAbilities/Source/GameplayAbilities/Public/AbilitySystemComponent.h
--- a/Engine/Plugins/Runtime/GameplayAbilities/Source/GameplayAbilities/Public/AbilitySystemComponent.h
+++ b/Engine/Plugins/Runtime/GameplayAbilities/Source/GameplayAbilities/Public/AbilitySystemComponent.h
@@ -574,7 +574,12 @@ class GAMEPLAYABILITIES_API UAbilitySystemComponent : public UGameplayTasksCompo
 	FORCEINLINE void GetOwnedGameplayTags(FGameplayTagContainer& TagContainer) const override
 	{
 		TagContainer.Reset();
-		TagContainer.AppendTags(GameplayTagCountContainer.GetExplicitGameplayTags());
+		TagContainer.AppendTags(GetOwnedGameplayTags());
+	}
+
+	UE_NODISCARD FORCEINLINE const FGameplayTagContainer& GetOwnedGameplayTags() const
+	{
+		return GameplayTagCountContainer.GetExplicitGameplayTags();
 	}
 
 	/** Checks whether the query matches the owned GameplayTags */
@@ -616,7 +621,12 @@ class GAMEPLAYABILITIES_API UAbilitySystemComponent : public UGameplayTasksCompo
 	/** Fills TagContainer with BlockedAbilityTags */
 	FORCEINLINE void GetBlockedAbilityTags(FGameplayTagContainer& TagContainer) const
 	{
-		TagContainer.AppendTags(BlockedAbilityTags.GetExplicitGameplayTags());
+		TagContainer.AppendTags(GetBlockedAbilityTags());
+	}
+
+	UE_NODISCARD FORCEINLINE const FGameplayTagContainer& GetBlockedAbilityTags() const
+	{
+		return BlockedAbilityTags.GetExplicitGameplayTags();
 	}
 
 	/** 	 
