Fixes for GameplayCueNotify OnActive executed twice for Duration GEs on Standalone + ListenServer in Mixed/Minimal Replication Mode.

Previously, when in Mixed/Minimal Replication mode, we would add call a Client RPC on top of adding the GameplayCue Tags to the Minimal list.  The RPCs would execute in Standalone/Listen since the "Server" is also a Client.  This previously worked in Client/Dedicated Server mode because the Dedicated Server mode typically ignores Gameplay Cue functions (so even though the logic was incorrect, the functions weren't invoked).

As well, since typically GameplayCue RPC's don't execute in Dedicated Server mode, there was logic to force Invoke the functions in special cases (ApplyGameplayEffectSpecToTarget).  That has been removed in favor of a single unified code path when the Gameplay Cues actually get added (InternalOnGameplayEffectAdded).

diff --git a/Engine/Plugins/Runtime/GameplayAbilities/Source/GameplayAbilities/Private/AbilitySystemComponent.cpp b/Engine/Plugins/Runtime/GameplayAbilities/Source/GameplayAbilities/Private/AbilitySystemComponent.cpp
--- a/Engine/Plugins/Runtime/GameplayAbilities/Source/GameplayAbilities/Private/AbilitySystemComponent.cpp
+++ b/Engine/Plugins/Runtime/GameplayAbilities/Source/GameplayAbilities/Private/AbilitySystemComponent.cpp
@@ -905,33 +905,23 @@ FActiveGameplayEffectHandle UAbilitySystemComponent::ApplyGameplayEffectSpecToSe
 	// Update (not push) the global spec being applied [we want to switch it to our copy, from the const input copy)
 	UAbilitySystemGlobals::Get().SetCurrentAppliedGE(OurCopyOfSpec);
 
-	// We still probably want to apply tags and stuff even if instant?
-	// If bSuppressStackingCues is set for this GameplayEffect, only add the GameplayCue if this is the first instance of the GameplayEffect
-	if (!bSuppressGameplayCues && bInvokeGameplayCueApplied && AppliedEffect && !AppliedEffect->bIsInhibited && 
-		(!bFoundExistingStackableGE || !Spec.Def->bSuppressStackingCues))
+	// UE5.4: We are following the same previous implementation that there is a special case for Gameplay Cues here (caveat: may not be true):
+	// We are Stacking an existing Gameplay Effect.  That means the GameplayCues should already be Added/WhileActive and we do not have a proper
+	// way to replicate the fact that it's been retriggered, hence the RPC here.  I say this may not be true because any number of things could have
+	// removed the GameplayCue by the time we getting a Stacking GE (e.g. RemoveGameplayCue).
+	if (!bSuppressGameplayCues && !Spec.Def->bSuppressStackingCues && bFoundExistingStackableGE && AppliedEffect && !AppliedEffect->bIsInhibited)
 	{
-		// We both added and activated the GameplayCue here.
-		// On the client, which will invoke the gameplay cue from an OnRep, it will need to look at the StartTime to determine
-		// if the Cue was actually added+activated or just added (due to relevancy)
-
-		// Fixme: what if we wanted to scale Cue magnitude based on damage? E.g, scale an cue effect when the GE is buffed?
-
-		if (OurCopyOfSpec->GetStackCount() > Spec.GetStackCount())
+		ensureMsgf(OurCopyOfSpec, TEXT("OurCopyOfSpec will always be valid if bFoundExistingStackableGE"));
+		if (OurCopyOfSpec && OurCopyOfSpec->GetStackCount() > Spec.GetStackCount())
 		{
 			// Because PostReplicatedChange will get called from modifying the stack count
 			// (and not PostReplicatedAdd) we won't know which GE was modified.
 			// So instead we need to explicitly RPC the client so it knows the GC needs updating
 			UAbilitySystemGlobals::Get().GetGameplayCueManager()->InvokeGameplayCueAddedAndWhileActive_FromSpec(this, *OurCopyOfSpec, PredictionKey);
 		}
-		else
-		{
-			// Otherwise these will get replicated to the client when the GE gets added to the replicated array
-			InvokeGameplayCueEvent(*OurCopyOfSpec, EGameplayCueEvent::OnActive);
-			InvokeGameplayCueEvent(*OurCopyOfSpec, EGameplayCueEvent::WhileActive);
-		}
 	}
 	
-	// Execute the GE at least once (if instant, this will execute once and be done. If persistent, it was added to ActiveGameplayEffects above)
+	// Execute the GE at least once (if instant, this will execute once and be done. If persistent, it was added to ActiveGameplayEffects in ApplyGameplayEffectSpec)
 	
 	// Execute if this is an instant application effect
 	if (bTreatAsInfiniteDuration)
@@ -1312,7 +1302,7 @@ void UAbilitySystemComponent::AddGameplayCue_Internal(const FGameplayTag Gamepla
 {
 	if (IsOwnerActorAuthoritative())
 	{
-		bool bWasInList = HasMatchingGameplayTag(GameplayCueTag);
+		const bool bWasInList = GameplayCueContainer.HasCue(GameplayCueTag);
 
 		ForceReplication();
 		GameplayCueContainer.AddCue(GameplayCueTag, ScopedPredictionKey, GameplayCueParameters);
@@ -1379,7 +1369,7 @@ void UAbilitySystemComponent::RemoveGameplayCue_Internal(const FGameplayTag Game
 {
 	if (IsOwnerActorAuthoritative())
 	{
-		bool bWasInList = HasMatchingGameplayTag(GameplayCueTag);
+		const bool bWasInList = GameplayCueContainer.HasCue(GameplayCueTag);
 
 		// Force replication so GameplayCue removals are properly replicated to all clients during Mixed and Minimal replication modes
 		ForceReplication();
@@ -1782,7 +1772,7 @@ void UAbilitySystemComponent::ReinvokeActiveGameplayCues()
 {
 	for (const FActiveGameplayEffect& Effect : &ActiveGameplayEffects)
 	{
-		if (Effect.bIsInhibited == false)
+		if (Effect.bIsInhibited == false && Effect.Spec.Def && Effect.Spec.Def->DurationPolicy != EGameplayEffectDurationType::Instant)
 		{
 			InvokeGameplayCueEvent(Effect.Spec, EGameplayCueEvent::WhileActive);
 		}
diff --git a/Engine/Plugins/Runtime/GameplayAbilities/Source/GameplayAbilities/Private/GameplayEffect.cpp b/Engine/Plugins/Runtime/GameplayAbilities/Source/GameplayAbilities/Private/GameplayEffect.cpp
--- a/Engine/Plugins/Runtime/GameplayAbilities/Source/GameplayAbilities/Private/GameplayEffect.cpp
+++ b/Engine/Plugins/Runtime/GameplayAbilities/Source/GameplayAbilities/Private/GameplayEffect.cpp
@@ -2501,8 +2501,9 @@ void FActiveGameplayEffect::PostReplicatedAdd(const struct FActiveGameplayEffect
 
 	UE_VLOG_UELOG(InArray.Owner->GetOwnerActor(), LogGameplayEffects, Verbose, TEXT("%s (Non-Auth): %s. Pending( OnActive: %d WhileActive: %d )"), ANSI_TO_TCHAR(__func__), *GetDebugString(), bPendingRepOnActiveGC, bPendingRepWhileActiveGC);
 
-	// Do stuff for adding GEs (add mods, tags, *invoke callbacks*
-	const_cast<FActiveGameplayEffectsContainer&>(InArray).InternalOnActiveGameplayEffectAdded(*this);	// Const cast is ok. It is there to prevent mutation of the GameplayEffects array, which this wont do.
+	// Do stuff for adding GEs (add mods, tags, *invoke callbacks*).  But do NOT invoke the GameplayCues as we don't know if this GE ends up inhibited or not (thus the bPendingRepOnActiveGC variables).
+	constexpr bool bInvokeGameplayCueEvents = false;
+	const_cast<FActiveGameplayEffectsContainer&>(InArray).InternalOnActiveGameplayEffectAdded(*this, bInvokeGameplayCueEvents);	// Const cast is ok. It is there to prevent mutation of the GameplayEffects array, which this wont do.
 	
 }
 
@@ -3926,14 +3927,17 @@ PRAGMA_ENABLE_DEPRECATION_WARNINGS
 	else
 	{
 		UE_VLOG_UELOG(OwnerActor, LogGameplayEffects, Verbose, TEXT("Adding GE: %s. Auth: %d. ReplicationID: %d. ReplicationKey: %d."), *AppliedActiveGE->GetDebugString(), IsNetAuthority(), AppliedActiveGE->ReplicationID, AppliedActiveGE->ReplicationKey);
-		InternalOnActiveGameplayEffectAdded(*AppliedActiveGE);
+
+		// Since we are applying it locally (and possibly predictively) invoke the cues.  Unless it's an Instant Cue, in which case we're not invoking the OnActive/WhileActive cues.
+		const bool bInvokeGameplayCueEvents = (Spec.Def->DurationPolicy != EGameplayEffectDurationType::Instant);
+		InternalOnActiveGameplayEffectAdded(*AppliedActiveGE, bInvokeGameplayCueEvents);
 	}
 
 	return AppliedActiveGE;
 }
 
 /** This is called anytime a new ActiveGameplayEffect is added, on both client and server in all cases */
-void FActiveGameplayEffectsContainer::InternalOnActiveGameplayEffectAdded(FActiveGameplayEffect& Effect)
+void FActiveGameplayEffectsContainer::InternalOnActiveGameplayEffectAdded(FActiveGameplayEffect& Effect, const bool bInvokeGameplayCueEvents)
 {
 	SCOPE_CYCLE_COUNTER(STAT_OnActiveGameplayEffectAdded);
 
@@ -3955,9 +3959,8 @@ void FActiveGameplayEffectsContainer::InternalOnActiveGameplayEffectAdded(FActiv
 	const bool bActive = EffectDef->OnAddedToActiveContainer(*this, Effect);
 	Effect.bIsInhibited = true; // Effect has to start inhibited, so our call to Inhibit will trigger if we should be active
 
-	constexpr bool bInvokeCuesIfEnabled = false;
 	FActiveGameplayEffectHandle EffectHandle = Effect.Handle;
-	Owner->SetActiveGameplayEffectInhibit(MoveTemp(EffectHandle), !bActive, bInvokeCuesIfEnabled);
+	Owner->SetActiveGameplayEffectInhibit(MoveTemp(EffectHandle), !bActive, bInvokeGameplayCueEvents);
 }
 
 void FActiveGameplayEffectsContainer::AddActiveGameplayEffectGrantedTagsAndModifiers(FActiveGameplayEffect& Effect, bool bInvokeGameplayCueEvents)
@@ -4064,15 +4067,9 @@ PRAGMA_ENABLE_DEPRECATION_WARNINGS
 	{
 		for (const FGameplayEffectCue& Cue : Effect.Spec.Def->GameplayCues)
 		{
-			Owner->UpdateTagMap(Cue.GameplayCueTags, 1);
-
-			if (bInvokeGameplayCueEvents)
-			{
-				Owner->InvokeGameplayCueEvent(Effect.Spec, EGameplayCueEvent::OnActive);
-				Owner->InvokeGameplayCueEvent(Effect.Spec, EGameplayCueEvent::WhileActive);
-			}
-
-			if (ShouldUseMinimalReplication())
+			// If we use Minimal/Mixed Replication, then this path will AddCue and broadcast the RPC that calls EGameplayCueEvent::OnActive (and WhileActive)
+			// Note: This is going to ignore bInvokeGameplayCueEvents and invoke them anyway (with a bunch of caveats e.g. you're the server and ignoring them)
+			if (ShouldUseMinimalReplication()) // This can only be true on Authority
 			{
 				for (const FGameplayTag& CueTag : Cue.GameplayCueTags)
 				{
@@ -4081,6 +4078,17 @@ PRAGMA_ENABLE_DEPRECATION_WARNINGS
 					Owner->AddGameplayCue_MinimalReplication(CueTag, Effect.Spec.GetEffectContext());
 				}
 			}
+			else // ActiveGameplayEffects are replicating to everyone (this path can also execute on client)
+			{
+				// Do a pseudo-AddGameplayCue (but don't add to ActiveGameplayCues so it doesn't replicate in addition to the AGE we're replicating).
+				Owner->UpdateTagMap(Cue.GameplayCueTags, 1);
+
+				if (bInvokeGameplayCueEvents)
+				{
+					Owner->InvokeGameplayCueEvent(Effect.Spec, EGameplayCueEvent::OnActive);
+					Owner->InvokeGameplayCueEvent(Effect.Spec, EGameplayCueEvent::WhileActive);
+				}
+			}
 		}
 	}
 
@@ -4197,21 +4205,21 @@ bool FActiveGameplayEffectsContainer::InternalRemoveActiveGameplayEffect(int32 I
 		}
 		
 		// Invoke Remove GameplayCue event
-		bool ShouldInvokeGameplayCueEvent = true;
+		bool bShouldInvokeGameplayCueEvent = true;
 		if (!bIsNetAuthority && Effect.PredictionKey.IsLocalClientKey() && Effect.PredictionKey.WasReceived() == false)
 		{
 			// This was an effect that we predicted. Don't invoke GameplayCue event if we have another GameplayEffect that shares the same predictionkey and was received from the server
 			if (HasReceivedEffectWithPredictedKey(Effect.PredictionKey))
 			{
-				ShouldInvokeGameplayCueEvent = false;
+				bShouldInvokeGameplayCueEvent = false;
 			}
 		}
 
 		// Don't invoke the GC event if the effect is inhibited, and thus the GC is already not active
-		ShouldInvokeGameplayCueEvent &= !Effect.bIsInhibited;
+		bShouldInvokeGameplayCueEvent &= !Effect.bIsInhibited;
 
 		// Mark the effect pending remove, and remove all side effects from the effect
-		InternalOnActiveGameplayEffectRemoved(Effect, ShouldInvokeGameplayCueEvent, GameplayEffectRemovalInfo);
+		InternalOnActiveGameplayEffectRemoved(Effect, bShouldInvokeGameplayCueEvent, GameplayEffectRemovalInfo);
 
 		// Check world validity in case RemoveActiveGameplayEffect is called during world teardown
 		if (UWorld* World = Owner->GetWorld())
@@ -4362,13 +4370,7 @@ PRAGMA_ENABLE_DEPRECATION_WARNINGS
 	{
 		for (const FGameplayEffectCue& Cue : Effect.Spec.Def->GameplayCues)
 		{
-			Owner->UpdateTagMap(Cue.GameplayCueTags, -1);
-
-			if (bInvokeGameplayCueEvents)
-			{
-				Owner->InvokeGameplayCueEvent(Effect.Spec, EGameplayCueEvent::Removed);
-			}
-
+			// If we use Minimal/Mixed Replication, then this will cause EGameplayCueEvent::Removed
 			if (ShouldUseMinimalReplication())
 			{
 				for (const FGameplayTag& CueTag : Cue.GameplayCueTags)
@@ -4376,6 +4378,16 @@ PRAGMA_ENABLE_DEPRECATION_WARNINGS
 					Owner->RemoveGameplayCue_MinimalReplication(CueTag);
 				}
 			}
+			else
+			{
+				// Perform pseudo-RemoveCue (without affecting ActiveGameplayCues, as we were not inserted there - see AddActiveGameplayEffectGrantedTagsAndModifiers)
+				Owner->UpdateTagMap(Cue.GameplayCueTags, -1);
+
+				if (bInvokeGameplayCueEvents)
+				{
+					Owner->InvokeGameplayCueEvent(Effect.Spec, EGameplayCueEvent::Removed);
+				}
+			}
 		}
 	}
 }
diff --git a/Engine/Plugins/Runtime/GameplayAbilities/Source/GameplayAbilities/Public/AbilitySystemComponent.h b/Engine/Plugins/Runtime/GameplayAbilities/Source/GameplayAbilities/Public/AbilitySystemComponent.h
--- a/Engine/Plugins/Runtime/GameplayAbilities/Source/GameplayAbilities/Public/AbilitySystemComponent.h
+++ b/Engine/Plugins/Runtime/GameplayAbilities/Source/GameplayAbilities/Public/AbilitySystemComponent.h
@@ -932,6 +932,7 @@ class GAMEPLAYABILITIES_API UAbilitySystemComponent : public UGameplayTasksCompo
 	virtual void HandleDeferredGameplayCues(const FActiveGameplayEffectsContainer* GameplayEffectsContainer);
 
 	/** Invokes the WhileActive event for all GCs on active, non inhibited, GEs. This would typically be used on "respawn" or something where the mesh/avatar has changed */
+	UE_DEPRECATED(5.4, "ReinvokeActiveGameplayCues was unused and had logic inconsistent with predicting Gameplay Effects.  You can implement it in your own project if desired.")
 	virtual void ReinvokeActiveGameplayCues();
 
 	/**
@@ -1845,11 +1846,11 @@ protected:
 	UPROPERTY(Replicated)
 	FActiveGameplayEffectsContainer ActiveGameplayEffects;
 
-	/** List of all active gameplay cues, including ones applied manually */
+	/** List of all active gameplay cues (executed outside of Gameplay Effects) */
 	UPROPERTY(Replicated)
 	FActiveGameplayCueContainer ActiveGameplayCues;
 
-	/** Replicated gameplaycues when in minimal replication mode. These are cues that would come normally come from ActiveGameplayEffects */
+	/** Replicated gameplaycues when in minimal replication mode. These are cues that would come normally come from ActiveGameplayEffects (but since we do not replicate AGE in minimal mode, they must be replicated through here) */
 	UPROPERTY(Replicated)
 	FActiveGameplayCueContainer MinimalReplicationGameplayCues;
 
diff --git a/Engine/Plugins/Runtime/GameplayAbilities/Source/GameplayAbilities/Public/GameplayEffect.h b/Engine/Plugins/Runtime/GameplayAbilities/Source/GameplayAbilities/Public/GameplayEffect.h
--- a/Engine/Plugins/Runtime/GameplayAbilities/Source/GameplayAbilities/Public/GameplayEffect.h
+++ b/Engine/Plugins/Runtime/GameplayAbilities/Source/GameplayAbilities/Public/GameplayEffect.h
@@ -1871,7 +1871,7 @@ private:
 	bool InternalRemoveActiveGameplayEffect(int32 Idx, int32 StacksToRemove, bool bPrematureRemoval);
 	
 	/** Called both in server side creation and replication creation/deletion */
-	void InternalOnActiveGameplayEffectAdded(FActiveGameplayEffect& Effect);
+	void InternalOnActiveGameplayEffectAdded(FActiveGameplayEffect& Effect, bool bInvokeGameplayCueEvents);
 	void InternalOnActiveGameplayEffectRemoved(FActiveGameplayEffect& Effect, bool bInvokeGameplayCueEvents, const FGameplayEffectRemovalInfo& GameplayEffectRemovalInfo);
 
 	void RemoveActiveGameplayEffectGrantedTagsAndModifiers(const FActiveGameplayEffect& Effect, bool bInvokeGameplayCueEvents);
