Gameplay Cue Unit Tests that can be a first line of defense when chasing bugs related to OnActive/WhileActive/OnRemove calls.  Network replication results aren't taken into account yet.

diff --git a/Engine/Plugins/Runtime/GameplayAbilities/Source/GameplayAbilities/Private/Tests/GameplayCueTests.h b/Engine/Plugins/Runtime/GameplayAbilities/Source/GameplayAbilities/Private/Tests/GameplayCueTests.h
new file mode 100644
--- /dev/null
+++ b/Engine/Plugins/Runtime/GameplayAbilities/Source/GameplayAbilities/Private/Tests/GameplayCueTests.h
@@ -0,0 +1,52 @@
+// Copyright Epic Games, Inc. All Rights Reserved.
+
+#pragma once
+
+#include "CoreMinimal.h"
+#include "GameplayCueNotify_Static.h"
+#include "NativeGameplayTags.h"
+#include "GameplayCueTests.generated.h"
+
+namespace UE::GameplayTags
+{
+	UE_DECLARE_GAMEPLAY_TAG_EXTERN(GameplayCue_Test);
+}
+
+/** Helper GameplayCueNotify that is used for unit testing.  Note: Since this is a GCN_Static, we're using the CDO during testing */
+UCLASS(NotBlueprintType, MinimalAPI)
+class UGameplayCueNotify_UnitTest : public UGameplayCueNotify_Static
+{
+	GENERATED_BODY()
+
+public:
+	virtual void PostInitProperties() override
+	{
+		Super::PostInitProperties();
+		GameplayCueTag = UE::GameplayTags::GameplayCue_Test;
+	}
+
+	/** Called when a GameplayCue is executed, this is used for instant effects or periodic ticks */
+	virtual bool OnExecute_Implementation(AActor* MyTarget, const FGameplayCueParameters& Parameters) const override { ++NumOnExecuteCalls; return true; }
+
+	/** Called when a GameplayCue with duration is first activated, this will only be called if the client witnessed the activation */
+	virtual bool OnActive_Implementation(AActor* MyTarget, const FGameplayCueParameters& Parameters) const override { ++NumOnActiveCalls; return true; }
+
+	/** Called when a GameplayCue with duration is first seen as active, even if it wasn't actually just applied (Join in progress, etc) */
+	virtual bool WhileActive_Implementation(AActor* MyTarget, const FGameplayCueParameters& Parameters) const override { ++NumWhileActiveCalls; return true; }
+
+	/** Called when a GameplayCue with duration is removed */
+	virtual bool OnRemove_Implementation(AActor* MyTarget, const FGameplayCueParameters& Parameters) const override { ++NumOnRemoveCalls; return true; }
+
+	/** These are counters for how many times the functions are called (so we can verify we get the right amounts).  All mutable because the UGCN_Static function overrides are const. */
+	mutable int NumOnExecuteCalls = 0;
+	mutable int NumOnActiveCalls = 0;
+	mutable int NumWhileActiveCalls = 0;
+	mutable int NumOnRemoveCalls = 0;
+
+	/** Reset all of the counts */
+	inline void ResetCallCounts()
+	{
+		NumOnExecuteCalls = NumOnActiveCalls = NumWhileActiveCalls = NumOnRemoveCalls = 0;
+	}
+};
+
diff --git a/Engine/Plugins/Runtime/GameplayAbilities/Source/GameplayAbilities/Private/Tests/GameplayEffectTests.cpp b/Engine/Plugins/Runtime/GameplayAbilities/Source/GameplayAbilities/Private/Tests/GameplayEffectTests.cpp
--- a/Engine/Plugins/Runtime/GameplayAbilities/Source/GameplayAbilities/Private/Tests/GameplayEffectTests.cpp
+++ b/Engine/Plugins/Runtime/GameplayAbilities/Source/GameplayAbilities/Private/Tests/GameplayEffectTests.cpp
@@ -18,8 +18,20 @@
 #include "AbilitySystemTestAttributeSet.h"
 #include "AbilitySystemGlobals.h"
 
+// For Testing Gameplay Cues
+#include "GameplayCueTests.h"
+#include "NativeGameplayTags.h"
+#include "GameplayCueManager.h"
+#include "GameplayCueSet.h"
+#include "GameplayCueNotify_Static.h"
+
 #define SKILL_TEST_TEXT( Format, ... ) FString::Printf(TEXT("%s - %d: %s"), TEXT(__FILE__) , __LINE__ , *FString::Printf(TEXT(Format), ##__VA_ARGS__) )
 
+namespace UE::GameplayTags
+{
+	UE_DEFINE_GAMEPLAY_TAG_COMMENT(GameplayCue_Test, "GameplayCue.Test", "GameplayCue Test Tag used for Unit Tests");
+}
+
 #if WITH_EDITOR
 
 static UDataTable* CreateGameplayDataTable()
@@ -315,6 +327,336 @@ public: // the tests
 		Test->TestEqual(TEXT("Stacking GEs"), DestComponent->GetSet<UAbilitySystemTestAttributeSet>()->StackingAttribute1, StartingAttributeValue);
 	}
 
+	void Test_GameplayCues_DirectAPI()
+	{
+		// Construct a TestGameplayCuesGE Gameplay Effect that triggers the GameplayCue.Test tag.
+		FGameplayEffectCue GameplayEffectCue{ UE::GameplayTags::GameplayCue_Test, 0.0f, 0.0f };
+		CONSTRUCT_CLASS(UGameplayEffect, TestGameplayCuesGE);
+		TestGameplayCuesGE->GameplayCues.Emplace(GameplayEffectCue);
+		TestGameplayCuesGE->DurationPolicy = EGameplayEffectDurationType::Infinite;
+		TestGameplayCuesGE->StackingType = EGameplayEffectStackingType::AggregateByTarget;
+		TestGameplayCuesGE->StackLimitCount = 5; // Make it high enough to not confuse between "expected 1 and got 2" etc.
+
+		UGameplayCueNotify_UnitTest* GCNotify_Test_CDO = GetMutableDefault<UGameplayCueNotify_UnitTest>();
+		const FGameplayEffectSpec GESpecInfinite{ TestGameplayCuesGE, FGameplayEffectContextHandle{ new FGameplayEffectContext{SourceActor, SourceActor} }, 1.f };
+
+		// Setup some parameters so we can test all of the routes we can execute a Cue through
+		FGameplayCueParameters GameplayCueParameters;
+		{
+			UAbilitySystemGlobals::Get().InitGameplayCueParameters_GESpec(GameplayCueParameters, GESpecInfinite);
+			GameplayCueParameters.EffectCauser = SourceActor;
+			GameplayCueParameters.Instigator = SourceActor;
+		}
+
+		Test->AddInfo(TEXT(" Testing Duration GameplayCues via UAbilitySystemComponent::AddGameplayCue"));
+		{
+			GCNotify_Test_CDO->ResetCallCounts();
+
+			// We should test with EGameplayCueExecutionOptions::IgnoreInterfaces and then with EGameplayCueExecutionOptions::IgnoreNotifies to see the functionality.
+			// Add GameplayCue should execute OnActive / While Active (used for duration)
+			DestComponent->AddGameplayCue(UE::GameplayTags::GameplayCue_Test, GameplayCueParameters);
+
+			Test->TestTrue(TEXT("  IsGameplayCueActive"), DestComponent->IsGameplayCueActive(UE::GameplayTags::GameplayCue_Test));
+			Test->TestEqual(TEXT("  OnActive Calls"), GCNotify_Test_CDO->NumOnActiveCalls, DestComponent->bSuppressGameplayCues ? 0 : 1);
+			Test->TestEqual(TEXT("  WhileActive Calls"), GCNotify_Test_CDO->NumWhileActiveCalls, DestComponent->bSuppressGameplayCues ? 0 : 1);
+			Test->TestEqual(TEXT("  OnExecute Calls"), GCNotify_Test_CDO->NumOnExecuteCalls, 0);
+			Test->TestEqual(TEXT("  OnRemove Calls"), GCNotify_Test_CDO->NumOnRemoveCalls, 0);
+			Test->TestEqual(TEXT("  OwnerTagCount"), DestComponent->GetTagCount(UE::GameplayTags::GameplayCue_Test), 1);
+		}
+
+		Test->AddInfo(TEXT(" Testing Duration GameplayCues via UAbilitySystemComponent::RemoveGameplayCue"));
+		{
+			GCNotify_Test_CDO->ResetCallCounts();
+
+			// Remove GameplayCue should execute OnRemove.
+			DestComponent->RemoveGameplayCue(UE::GameplayTags::GameplayCue_Test);
+
+			Test->TestFalse(TEXT("  IsGameplayCueActive"), DestComponent->IsGameplayCueActive(UE::GameplayTags::GameplayCue_Test));
+			Test->TestEqual(TEXT("  OnActive Calls"), GCNotify_Test_CDO->NumOnActiveCalls, 0);
+			Test->TestEqual(TEXT("  WhileActive Calls"), GCNotify_Test_CDO->NumWhileActiveCalls, 0);
+			Test->TestEqual(TEXT("  OnExecute Calls"), GCNotify_Test_CDO->NumOnExecuteCalls, 0);
+			Test->TestEqual(TEXT("  OnRemove Calls"), GCNotify_Test_CDO->NumOnRemoveCalls, DestComponent->bSuppressGameplayCues ? 0 : 1);
+			Test->TestEqual(TEXT("  OwnerTagCount"), DestComponent->GetTagCount(UE::GameplayTags::GameplayCue_Test), 0);
+		}
+
+		Test->AddInfo(TEXT(" Testing Duration GameplayCues via UAbilitySystemComponent::AddGameplayCue_MinimalReplication"));
+		{
+			GCNotify_Test_CDO->ResetCallCounts();
+
+			// Add the Minimal Replication Cue should execute OnActive/WhileActive
+			DestComponent->AddGameplayCue_MinimalReplication(UE::GameplayTags::GameplayCue_Test, GESpecInfinite.GetEffectContext());
+
+			Test->TestTrue(TEXT("  IsGameplayCueActive"), DestComponent->IsGameplayCueActive(UE::GameplayTags::GameplayCue_Test));
+			Test->TestEqual(TEXT("  OnActive Calls"), GCNotify_Test_CDO->NumOnActiveCalls, DestComponent->bSuppressGameplayCues ? 0 : 1);
+			Test->TestEqual(TEXT("  WhileActive Calls"), GCNotify_Test_CDO->NumWhileActiveCalls, DestComponent->bSuppressGameplayCues ? 0 : 1);
+			Test->TestEqual(TEXT("  OnExecute Calls"), GCNotify_Test_CDO->NumOnExecuteCalls, 0);
+			Test->TestEqual(TEXT("  OnRemove Calls"), GCNotify_Test_CDO->NumOnRemoveCalls, 0);
+			Test->TestEqual(TEXT("  OwnerTagCount"), DestComponent->GetTagCount(UE::GameplayTags::GameplayCue_Test), 1);
+		}
+
+		Test->AddInfo(TEXT(" Testing Duration GameplayCues via UAbilitySystemComponent::RemoveGameplayCue_MinimalReplication"));
+		{
+			GCNotify_Test_CDO->ResetCallCounts();
+
+			// Remove GameplayCue should execute OnRemove.
+			DestComponent->RemoveGameplayCue_MinimalReplication(UE::GameplayTags::GameplayCue_Test);
+
+			Test->TestFalse(TEXT("  IsGameplayCueActive"), DestComponent->IsGameplayCueActive(UE::GameplayTags::GameplayCue_Test));
+			Test->TestEqual(TEXT("  OnActive Calls"), GCNotify_Test_CDO->NumOnActiveCalls, 0);
+			Test->TestEqual(TEXT("  WhileActive Calls"), GCNotify_Test_CDO->NumWhileActiveCalls, 0);
+			Test->TestEqual(TEXT("  OnExecute Calls"), GCNotify_Test_CDO->NumOnExecuteCalls, 0);
+			Test->TestEqual(TEXT("  OnRemove Calls"), GCNotify_Test_CDO->NumOnRemoveCalls, DestComponent->bSuppressGameplayCues ? 0 : 1);
+			Test->TestEqual(TEXT("  OwnerTagCount"), DestComponent->GetTagCount(UE::GameplayTags::GameplayCue_Test), 0);
+		}
+
+		Test->AddInfo(TEXT(" Testing UAbilitySystemComponent::InvokeGameplayCueEvent"));
+		{
+			GCNotify_Test_CDO->ResetCallCounts();
+
+			Test->TestFalse(TEXT("  IsGameplayCueActive"), DestComponent->IsGameplayCueActive(UE::GameplayTags::GameplayCue_Test));
+
+			DestComponent->InvokeGameplayCueEvent(GESpecInfinite, EGameplayCueEvent::OnActive);
+			DestComponent->InvokeGameplayCueEvent(GESpecInfinite, EGameplayCueEvent::WhileActive);
+			DestComponent->InvokeGameplayCueEvent(GESpecInfinite, EGameplayCueEvent::Executed);
+			DestComponent->InvokeGameplayCueEvent(GESpecInfinite, EGameplayCueEvent::Removed);
+
+			// One extra one, just to make sure
+			DestComponent->InvokeGameplayCueEvent(GESpecInfinite, EGameplayCueEvent::WhileActive);
+
+			Test->TestFalse(TEXT("  IsGameplayCueActive"), DestComponent->IsGameplayCueActive(UE::GameplayTags::GameplayCue_Test));
+			Test->TestEqual(TEXT("  OnActive Calls"), GCNotify_Test_CDO->NumOnActiveCalls, DestComponent->bSuppressGameplayCues ? 0 : 1);
+			Test->TestEqual(TEXT("  WhileActive Calls"), GCNotify_Test_CDO->NumWhileActiveCalls, DestComponent->bSuppressGameplayCues ? 0 : 2);
+			Test->TestEqual(TEXT("  OnExecute Calls"), GCNotify_Test_CDO->NumOnExecuteCalls, DestComponent->bSuppressGameplayCues ? 0 : 1);
+			Test->TestEqual(TEXT("  OnRemove Calls"), GCNotify_Test_CDO->NumOnRemoveCalls, DestComponent->bSuppressGameplayCues ? 0 : 1);
+			Test->TestEqual(TEXT("  OwnerTagCount"), DestComponent->GetTagCount(UE::GameplayTags::GameplayCue_Test), 0);
+		}
+
+		Test->AddInfo(TEXT(" Testing UAbilitySystemComponent::ExecuteGameplayCue"));
+		{
+			GCNotify_Test_CDO->ResetCallCounts();
+
+			Test->TestFalse(TEXT("  IsGameplayCueActive"), DestComponent->IsGameplayCueActive(UE::GameplayTags::GameplayCue_Test));
+
+			DestComponent->ExecuteGameplayCue(UE::GameplayTags::GameplayCue_Test, GameplayCueParameters);
+			Test->TestFalse(TEXT("  IsGameplayCueActive"), DestComponent->IsGameplayCueActive(UE::GameplayTags::GameplayCue_Test));
+			Test->TestEqual(TEXT("  OnActive Calls"), GCNotify_Test_CDO->NumOnActiveCalls, 0);
+			Test->TestEqual(TEXT("  WhileActive Calls"), GCNotify_Test_CDO->NumWhileActiveCalls, 0);
+			Test->TestEqual(TEXT("  OnExecute Calls"), GCNotify_Test_CDO->NumOnExecuteCalls, DestComponent->bSuppressGameplayCues ? 0 : 1);
+			Test->TestEqual(TEXT("  OnRemove Calls"), GCNotify_Test_CDO->NumOnRemoveCalls, 0);
+			Test->TestEqual(TEXT("  OwnerTagCount"), DestComponent->GetTagCount(UE::GameplayTags::GameplayCue_Test), 0);
+		}
+	}
+
+	void Test_GameplayCues_GameplayEffectsAPI()
+	{
+		// Construct a TestGameplayCuesGE Gameplay Effect that triggers the GameplayCue.Test tag.
+		FGameplayEffectCue GameplayEffectCue{ UE::GameplayTags::GameplayCue_Test, 0.0f, 0.0f };
+		CONSTRUCT_CLASS(UGameplayEffect, TestGameplayCuesGE);
+		TestGameplayCuesGE->GameplayCues.Emplace(GameplayEffectCue);
+		TestGameplayCuesGE->DurationPolicy = EGameplayEffectDurationType::Infinite;
+		TestGameplayCuesGE->StackingType = EGameplayEffectStackingType::AggregateByTarget;
+		TestGameplayCuesGE->StackDurationRefreshPolicy = EGameplayEffectStackingDurationPolicy::NeverRefresh;
+		TestGameplayCuesGE->StackPeriodResetPolicy = EGameplayEffectStackingPeriodPolicy::NeverReset;
+		TestGameplayCuesGE->StackLimitCount = 5; // Make it high enough to not confuse between "expected 1 and got 2" etc.
+
+		UGameplayCueNotify_UnitTest* GCNotify_Test_CDO = GetMutableDefault<UGameplayCueNotify_UnitTest>();
+		const FGameplayEffectSpec GESpecInfinite{ TestGameplayCuesGE, FGameplayEffectContextHandle{ new FGameplayEffectContext{SourceActor, SourceActor} }, 1.f };
+
+		auto TestGameplayEffectsTriggerGameplayCuesProperly = [&](const FGameplayEffectSpec& GESpec)
+			{
+				GCNotify_Test_CDO->ResetCallCounts();
+
+				Test->TestFalse(TEXT("  IsGameplayCueActive (Before Apply)"), DestComponent->IsGameplayCueActive(UE::GameplayTags::GameplayCue_Test));
+
+				// Apply the GE.  We expect OnActive/WhileActive/OnRemove for non-instant; OnExecute for the instant ones.
+				FActiveGameplayEffectHandle ActiveGEHandle = DestComponent->ApplyGameplayEffectSpecToSelf(GESpec, FPredictionKey::CreateNewPredictionKey(DestComponent));
+
+				const bool bPredicted = !DestComponent->IsOwnerActorAuthoritative();
+				const bool bInstant = (TestGameplayCuesGE->DurationPolicy != EGameplayEffectDurationType::Instant);
+				const bool bExpectedAdd = (!DestComponent->bSuppressGameplayCues) && (bPredicted || bInstant);
+				const int ExpectedOnActive = bExpectedAdd && (TestGameplayCuesGE->DurationPolicy != EGameplayEffectDurationType::Instant);
+				const int ExpectedOnExecute = ((TestGameplayCuesGE->DurationPolicy == EGameplayEffectDurationType::Instant) && (!DestComponent->bSuppressGameplayCues)) ? 1 : 0;
+				const int ExpectedWhileActive = ExpectedOnActive;
+				const int ExpectedOnRemove = bExpectedAdd;
+				const int ExpectedOwnerTagCount = bExpectedAdd;
+
+				Test->TestEqual(TEXT("  IsGameplayCueActive (After Apply)"), DestComponent->IsGameplayCueActive(UE::GameplayTags::GameplayCue_Test), bExpectedAdd);
+				Test->TestEqual(TEXT("  OnActive Calls"), GCNotify_Test_CDO->NumOnActiveCalls, ExpectedOnActive);
+				Test->TestEqual(TEXT("  WhileActive Calls"), GCNotify_Test_CDO->NumWhileActiveCalls, ExpectedWhileActive);
+				Test->TestEqual(TEXT("  OnExecute Calls"), GCNotify_Test_CDO->NumOnExecuteCalls, ExpectedOnExecute);
+				Test->TestEqual(TEXT("  OnRemove Calls"), GCNotify_Test_CDO->NumOnRemoveCalls, 0);
+				Test->TestEqual(TEXT("  OwnerTagCount"), DestComponent->GetTagCount(UE::GameplayTags::GameplayCue_Test), ExpectedOwnerTagCount);
+
+				// Now try to remove it
+				DestComponent->RemoveActiveGameplayEffect(ActiveGEHandle);
+
+				Test->TestFalse(TEXT("  IsGameplayCueActive (After RemoveActiveEffects)"), DestComponent->IsGameplayCueActive(UE::GameplayTags::GameplayCue_Test));
+				Test->TestEqual(TEXT("  OnRemove Calls"), GCNotify_Test_CDO->NumOnRemoveCalls, ExpectedOnRemove);
+				Test->TestEqual(TEXT("  OwnerTagCount"), DestComponent->GetTagCount(UE::GameplayTags::GameplayCue_Test), 0);
+			};
+
+		auto TestGameplayEffectsStackTriggerGameplayCuesProperly = [&](const FGameplayEffectSpec& GESpec)
+			{
+				GCNotify_Test_CDO->ResetCallCounts();
+
+				Test->TestFalse(TEXT("  IsGameplayCueActive (Before Stack Application)"), DestComponent->IsGameplayCueActive(UE::GameplayTags::GameplayCue_Test));
+
+				// Apply the GE until the StackCount.  We expect OnActive/WhileActive/OnRemove for non-instant; OnExecute cannot stack.
+				FActiveGameplayEffectHandle ActiveGEHandle = DestComponent->ApplyGameplayEffectSpecToSelf(GESpec, FPredictionKey::CreateNewPredictionKey(DestComponent));
+				const int StackLimitCount = GESpec.Def->GetStackLimitCount();
+				for (int Index = 1; Index < StackLimitCount; ++Index)
+				{
+					DestComponent->ApplyGameplayEffectSpecToSelf(GESpec, FPredictionKey::CreateNewPredictionKey(DestComponent));
+				}
+
+				// More complex rules:  We expect as many OnActive's as there is applications, unless we're predicting.
+				// Unexpected behavior:  We can never predict a stacked application (why?!) so it ends up with 1.
+				const bool bPredicting = !DestComponent->IsOwnerActorAuthoritative();
+				bool bExpectedAdded = (!DestComponent->bSuppressGameplayCues);
+				int ExpectedOnActive = bExpectedAdded ? (bPredicting ? 1 : StackLimitCount) : 0; // can only predict the first cue
+				int ExpectedOnExecute = 0;
+				int ExpectedWhileActive = ExpectedOnActive;
+				int ExpectedOnRemove = (ExpectedOnActive > 0) ? 1 : 0;		// Even with stacking, the GameplayCue is actually only applied once (and therefore removed once when fully unstacked)
+				int ExpectedOwnerTagCount = (ExpectedOnActive > 0) ? 1 : 0;	// Even with stacking, we only have the tag applied once
+
+				const bool bInstant = (TestGameplayCuesGE->DurationPolicy == EGameplayEffectDurationType::Instant);
+				if (bInstant)
+				{
+					// Instant Cues should not actually stack, so we don't get any of the Added/While/Remove calls, just Execute...
+					// Except if we're predicted, in which case we're added so that we can reconcile the prediction.
+					bExpectedAdded = (!DestComponent->bSuppressGameplayCues) && bPredicting;
+					ExpectedOnExecute = (!DestComponent->bSuppressGameplayCues) ? StackLimitCount : 0;
+					ExpectedOnActive = ExpectedWhileActive = 0;
+
+					// Unexpected behavior here:  If we're predicting, we get Tags.  If not, we don't get tags.
+					ExpectedOwnerTagCount = bPredicting && (!DestComponent->bSuppressGameplayCues) ? ExpectedOnExecute : 0;
+					ExpectedOnRemove = bExpectedAdded && (!DestComponent->bSuppressGameplayCues) ? ExpectedOnExecute : 0;
+				}
+
+				Test->TestEqual(TEXT("  IsGameplayCueActive (After Stack Application)"), DestComponent->IsGameplayCueActive(UE::GameplayTags::GameplayCue_Test), bExpectedAdded);
+				Test->TestEqual(TEXT("  OnActive Calls"), GCNotify_Test_CDO->NumOnActiveCalls, ExpectedOnActive);
+				Test->TestEqual(TEXT("  WhileActive Calls"), GCNotify_Test_CDO->NumWhileActiveCalls, ExpectedWhileActive);
+				Test->TestEqual(TEXT("  OnExecute Calls"), GCNotify_Test_CDO->NumOnExecuteCalls, ExpectedOnExecute);
+				Test->TestEqual(TEXT("  OnRemove Calls"), GCNotify_Test_CDO->NumOnRemoveCalls, 0);
+				Test->TestEqual(TEXT("  OwnerTagCount"), DestComponent->GetTagCount(UE::GameplayTags::GameplayCue_Test), ExpectedOwnerTagCount);
+
+				// Now try to remove it.  If not instant, try one at a time.
+				if (!bInstant)
+				{
+					constexpr int StacksToRemove = 1;
+					for (int Index = 0; Index < ExpectedOnActive; ++Index)
+					{
+						Test->TestEqual(TEXT("  IsGameplayCueActive before RemoveActiveEffects Unstack"), DestComponent->IsGameplayCueActive(UE::GameplayTags::GameplayCue_Test), (ExpectedOnActive > 0));
+						DestComponent->RemoveActiveGameplayEffect(ActiveGEHandle, StacksToRemove);
+					}
+				}
+				else
+				{
+					// If instant, since we ActiveGEHandle was different, we need to remove this way.  This function for some reason allows us to bypass the Authority check.
+					DestComponent->RemoveActiveGameplayEffectBySourceEffect(UGameplayEffect::StaticClass(), SourceComponent);
+				}
+
+				Test->TestFalse(TEXT("  IsGameplayCueActive (After Stack Removal)"), DestComponent->IsGameplayCueActive(UE::GameplayTags::GameplayCue_Test));
+				Test->TestEqual(TEXT("  OnRemove Calls"), GCNotify_Test_CDO->NumOnRemoveCalls, ExpectedOnRemove);
+				Test->TestEqual(TEXT("  OwnerTagCount"), DestComponent->GetTagCount(UE::GameplayTags::GameplayCue_Test), 0);
+			};
+
+		Test->AddInfo(TEXT(" Testing Infinite Duration GameplayCues via UAbilitySystemComponent::ApplyGameplayEffectSpecToSelf"));
+		TestGameplayCuesGE->DurationPolicy = EGameplayEffectDurationType::Infinite;
+		TestGameplayEffectsTriggerGameplayCuesProperly(GESpecInfinite);
+
+		Test->AddInfo(TEXT(" Testing Stacking Infinite Duration GameplayCues via UAbilitySystemComponent::ApplyGameplayEffectSpecToSelf"));
+		TestGameplayEffectsStackTriggerGameplayCuesProperly(GESpecInfinite);
+
+		Test->AddInfo(TEXT(" Testing HasDuration GameplayCues via UAbilitySystemComponent::ApplyGameplayEffectSpecToSelf"));
+		TestGameplayCuesGE->DurationPolicy = EGameplayEffectDurationType::HasDuration;
+		FGameplayEffectSpec GESpecHasDuration{ GESpecInfinite, GESpecInfinite.GetContext() };
+		GESpecHasDuration.SetDuration(5.0f, false);
+		TestGameplayEffectsTriggerGameplayCuesProperly(GESpecHasDuration);
+
+		Test->AddInfo(TEXT(" Testing Stacking HasDuration GameplayCues via UAbilitySystemComponent::ApplyGameplayEffectSpecToSelf"));
+		TestGameplayEffectsStackTriggerGameplayCuesProperly(GESpecHasDuration);
+
+		Test->AddInfo(TEXT(" Testing Instant GameplayCues via UAbilitySystemComponent::ApplyGameplayEffectSpecToSelf"));
+		TestGameplayCuesGE->DurationPolicy = EGameplayEffectDurationType::Instant;
+		FGameplayEffectSpec GESpecInstant{ GESpecInfinite, GESpecInfinite.GetContext() };
+		GESpecInstant.SetDuration(FGameplayEffectConstants::INSTANT_APPLICATION, false);
+		TestGameplayEffectsTriggerGameplayCuesProperly(GESpecInstant);
+
+		Test->AddInfo(TEXT(" Testing Stacking Instant GameplayCues via UAbilitySystemComponent::ApplyGameplayEffectSpecToSelf"));
+		TestGameplayEffectsStackTriggerGameplayCuesProperly(GESpecInstant);
+
+		Test->TestFalse(TEXT(" Ending with IsGameplayCueActive?"), DestComponent->IsGameplayCueActive(UE::GameplayTags::GameplayCue_Test));
+	}
+
+	void Test_GameplayCues()
+	{
+		Test->TestTrue(TEXT("DestComponent IsReadyForGameplayCues()"), DestComponent->IsReadyForGameplayCues());
+
+		// Setup a temporary UGameplayCueNotify_Test instance that will intercept the GameplayCue.Test tag
+		// We are poking the data that should be internal (GameplayCueData) because we want to manually populate
+		// the LoadedGameplayCueClass (to avoid the async load path).  That means we will eventually need to
+		const UGameplayCueNotify_UnitTest* GCNotify_Test_CDO = GetDefault<UGameplayCueNotify_UnitTest>();
+		FGameplayCueNotifyData GameplayCueNotifyData;
+		GameplayCueNotifyData.GameplayCueTag = UE::GameplayTags::GameplayCue_Test;
+		GameplayCueNotifyData.GameplayCueNotifyObj = GCNotify_Test_CDO->GetPathName();
+		GameplayCueNotifyData.LoadedGameplayCueClass = UGameplayCueNotify_UnitTest::StaticClass();
+
+		UGameplayCueManager* GameplayCueManager = UAbilitySystemGlobals::Get().GetGameplayCueManager();
+		UGameplayCueSet* RuntimeCueSet = GameplayCueManager->GetRuntimeCueSet();
+		RuntimeCueSet->GameplayCueData.Emplace(GameplayCueNotifyData);
+
+		// Now just force a rebuild of the internal table since we weren't supposed to poke the data manually.
+		RuntimeCueSet->UpdateCueByStringRefs(GameplayCueNotifyData.GameplayCueNotifyObj, GCNotify_Test_CDO->GetPathName());
+		ON_SCOPE_EXIT{ RuntimeCueSet->RemoveCuesByStringRefs({ GCNotify_Test_CDO->GetPathName() }); };
+
+		// Perform all of the Tests once as an Authority (Server) and once as an AutonomousProxy (Client)
+		Test->AddInfo(TEXT("--- Testing as ROLE_Authority ---"));
+		// Full Mode (Gameplay Effects are replicated to all clients)
+		Test->AddInfo("Testing SetReplicationMode(Full)");
+		DestComponent->SetReplicationMode(EGameplayEffectReplicationMode::Full);
+		Test_GameplayCues_DirectAPI();
+		Test_GameplayCues_GameplayEffectsAPI();
+
+		// Mixed Mode (GE's are replicated to Autonomous Proxies, Cues from the GE's are replicated to Simulated as they won't have the GE's)
+		Test->AddInfo("Testing SetReplicationMode(Mixed)");
+		DestComponent->SetReplicationMode(EGameplayEffectReplicationMode::Mixed);
+		Test_GameplayCues_DirectAPI();
+		Test_GameplayCues_GameplayEffectsAPI();
+
+		// Minimal Mode (GE's are never replicated; Cues from the GE's are replicated to the Simulated as they won't have the GE's.  There should be no autonomous proxies when this is set (e.g. used for AI).)
+		Test->AddInfo("Testing SetReplicationMode(Minimal)");
+		DestComponent->SetReplicationMode(EGameplayEffectReplicationMode::Minimal);
+		Test_GameplayCues_DirectAPI();
+		Test_GameplayCues_GameplayEffectsAPI();
+
+		// Test Suppression
+		DestComponent->bSuppressGameplayCues = true;
+		DestComponent->SetReplicationMode(EGameplayEffectReplicationMode::Full);
+		Test->AddInfo(TEXT("Testing bSuppressGameplayCues (expecting no Gameplay Cues)"));
+		Test_GameplayCues_DirectAPI();
+		Test_GameplayCues_GameplayEffectsAPI();
+		DestComponent->bSuppressGameplayCues = false;
+
+		// As Autonomous Proxy, we are mimicking the prediction of Gameplay Effects.  As such, we should not
+		// be calling the Direct API, only the Gameplay Effects API.
+		Test->AddInfo(TEXT("--- Testing as ROLE_AutonomousProxy ---"));
+		SourceActor->SetRole(ENetRole::ROLE_AutonomousProxy);
+		SourceComponent->CacheIsNetSimulated();
+		DestActor->SetRole(ENetRole::ROLE_AutonomousProxy);
+		DestComponent->CacheIsNetSimulated();
+
+		// Replication modes are server-side only, so we really only need to test the one...
+		Test->AddInfo("Testing SetReplicationMode(Mixed)");
+		DestComponent->SetReplicationMode(EGameplayEffectReplicationMode::Mixed);
+		Test_GameplayCues_GameplayEffectsAPI();
+
+		// Reset the variables
+		SourceActor->SetRole(ENetRole::ROLE_Authority);
+		SourceComponent->CacheIsNetSimulated();
+		DestActor->SetRole(ENetRole::ROLE_Authority);
+		DestComponent->CacheIsNetSimulated();
+	}
 private: // test helpers
 
 	template<typename MODIFIER_T>
@@ -373,6 +715,7 @@ public:
 		ADD_TEST(Test_PeriodicDamage);
 		ADD_TEST(Test_StackLimit);
 		ADD_TEST(Test_SetByCallerStackingDuration);
+		ADD_TEST(Test_GameplayCues);
 	}
 
 	virtual uint32 GetTestFlags() const override { return EAutomationTestFlags::EditorContext | EAutomationTestFlags::EngineFilter; }
