PR #11119: GameplayEffect ApplicationRequirements Upgrade Fix

diff --git a/Engine/Plugins/Runtime/GameplayAbilities/Source/GameplayAbilities/Private/GameplayEffect.cpp b/Engine/Plugins/Runtime/GameplayAbilities/Source/GameplayAbilities/Private/GameplayEffect.cpp
--- a/Engine/Plugins/Runtime/GameplayAbilities/Source/GameplayAbilities/Private/GameplayEffect.cpp
+++ b/Engine/Plugins/Runtime/GameplayAbilities/Source/GameplayAbilities/Private/GameplayEffect.cpp
@@ -440,7 +440,7 @@ void UGameplayEffect::ConvertCustomCanApplyComponent()
 	if (bChanged && UE::GameplayEffect::EditorOnly::ShouldUpgradeVersion(GetVersion(), EGameplayEffectVersion::Modular53))
 	{
 		UCustomCanApplyGameplayEffectComponent& ApplicationComponent = FindOrAddComponent<UCustomCanApplyGameplayEffectComponent>();
-		ApplicationComponent.ApplicationRequirements.Append(ApplicationRequirements_DEPRECATED);
+		ApplicationComponent.ApplicationRequirements = ApplicationRequirements_DEPRECATED;
 	}
 
 	// Keep backwards compatibility (at least in terms of reading from the data)
