PR #11117: Fixed typos and wrong parameter name in comments in the AbilitySystemComponent header file

diff --git a/Engine/Plugins/Runtime/GameplayAbilities/Source/GameplayAbilities/Public/AbilitySystemComponent.h b/Engine/Plugins/Runtime/GameplayAbilities/Source/GameplayAbilities/Public/AbilitySystemComponent.h
--- a/Engine/Plugins/Runtime/GameplayAbilities/Source/GameplayAbilities/Public/AbilitySystemComponent.h
+++ b/Engine/Plugins/Runtime/GameplayAbilities/Source/GameplayAbilities/Public/AbilitySystemComponent.h
@@ -267,7 +267,7 @@ class GAMEPLAYABILITIES_API UAbilitySystemComponent : public UGameplayTasksCompo
 		return ScopedPredictionKey.IsValidForMorePrediction();
 	}
 
-	/** Returns true if this is running on the server or has a valid prediciton key */
+	/** Returns true if this is running on the server or has a valid prediction key */
 	bool HasAuthorityOrPredictionKey(const FGameplayAbilityActivationInfo* ActivationInfo) const;
 
 	/** Returns true if this component's actor has authority */
@@ -490,7 +490,7 @@ class GAMEPLAYABILITIES_API UAbilitySystemComponent : public UGameplayTasksCompo
 	void OnPeriodicGameplayEffectExecuteOnTarget(UAbilitySystemComponent* Target, const FGameplayEffectSpec& SpecExecuted, FActiveGameplayEffectHandle ActiveHandle);
 	void OnPeriodicGameplayEffectExecuteOnSelf(UAbilitySystemComponent* Source, const FGameplayEffectSpec& SpecExecuted, FActiveGameplayEffectHandle ActiveHandle);
 
-	/** Called when the duration of a gamepaly effect has changed */
+	/** Called when the duration of a gameplay effect has changed */
 	virtual void OnGameplayEffectDurationChange(struct FActiveGameplayEffect& ActiveEffect);
 
 	/** Called on server whenever a GE is applied to self. This includes instant and duration based GEs. */
@@ -499,7 +499,7 @@ class GAMEPLAYABILITIES_API UAbilitySystemComponent : public UGameplayTasksCompo
 	/** Called on server whenever a GE is applied to someone else. This includes instant and duration based GEs. */
 	FOnGameplayEffectAppliedDelegate OnGameplayEffectAppliedDelegateToTarget;
 
-	/** Called on both client and server whenever a duraton based GE is added (E.g., instant GEs do not trigger this). */
+	/** Called on both client and server whenever a duration based GE is added (E.g., instant GEs do not trigger this). */
 	FOnGameplayEffectAppliedDelegate OnActiveGameplayEffectAddedDelegateToSelf;
 
 	/** Called on server whenever a periodic GE executes on self */
@@ -915,7 +915,7 @@ class GAMEPLAYABILITIES_API UAbilitySystemComponent : public UGameplayTasksCompo
 	void InvokeGameplayCueEvent(const FGameplayTag GameplayCueTag, EGameplayCueEvent::Type EventType, FGameplayEffectContextHandle EffectContext = FGameplayEffectContextHandle());
 	void InvokeGameplayCueEvent(const FGameplayTag GameplayCueTag, EGameplayCueEvent::Type EventType, const FGameplayCueParameters& GameplayCueParameters);
 
-	/** Allows polling to see if a GameplayCue is active. We expect most GameplayCue handling to be event based, but some cases we may need to check if a GamepalyCue is active (Animation Blueprint for example) */
+	/** Allows polling to see if a GameplayCue is active. We expect most GameplayCue handling to be event based, but some cases we may need to check if a GameplayCue is active (Animation Blueprint for example) */
 	UFUNCTION(BlueprintCallable, Category="GameplayCue", meta=(GameplayTagFilter="GameplayCue"))
 	bool IsGameplayCueActive(const FGameplayTag GameplayCueTag) const
 	{
@@ -1162,7 +1162,7 @@ class GAMEPLAYABILITIES_API UAbilitySystemComponent : public UGameplayTasksCompo
 	 *
 	 * @param OutAbilityHandles This array will be filled with matching Ability Spec Handles
 	 * @param Tags Gameplay Tags to match
-	 * @param bMatchAll If true, tags must be matched exactly. Otherwise, abilities matching any of the tags will be returned
+	 * @param bExactMatch If true, tags must be matched exactly. Otherwise, abilities matching any of the tags will be returned
 	 */
 	UFUNCTION(BlueprintCallable, BlueprintPure = false, Category = "Gameplay Abilities")
 	void FindAllAbilitiesWithTags(TArray<FGameplayAbilitySpecHandle>& OutAbilityHandles, FGameplayTagContainer Tags, bool bExactMatch = true) const;
@@ -1797,7 +1797,7 @@ protected:
 	/** Returns true if we are ready to handle replicated montage information */
 	virtual bool IsReadyForReplicatedMontage();
 
-	/** RPC function called from CurrentMontageSetNextSectopnName, replicates to other clients */
+	/** RPC function called from CurrentMontageSetNextSectionName, replicates to other clients */
 	UFUNCTION(reliable, server, WithValidation)
 	void ServerCurrentMontageSetNextSectionName(UAnimSequenceBase* ClientAnimation, float ClientPosition, FName SectionName, FName NextSectionName);
 
