Ability System Debug HUD now uses the correct source for CanActivateAbility.

diff --git a/Engine/Plugins/Runtime/GameplayAbilities/Source/GameplayAbilities/Private/AbilitySystemComponent.cpp b/Engine/Plugins/Runtime/GameplayAbilities/Source/GameplayAbilities/Private/AbilitySystemComponent.cpp
--- a/Engine/Plugins/Runtime/GameplayAbilities/Source/GameplayAbilities/Private/AbilitySystemComponent.cpp
+++ b/Engine/Plugins/Runtime/GameplayAbilities/Source/GameplayAbilities/Private/AbilitySystemComponent.cpp
@@ -2665,6 +2665,13 @@ void UAbilitySystemComponent::Debug_Internal(FAbilitySystemComponentDebugInfo& I
 			FGameplayTagContainer FailureTags;
 			const TArray<uint8>& LocalBlockedAbilityBindings = GetBlockedAbilityBindings();
 
+			// We prefer executing on the instanced ability if we're instancing
+			UGameplayAbility* AbilitySource = AbilitySpec.GetPrimaryInstance();
+			if (!AbilitySource)
+			{
+				AbilitySource = AbilitySpec.Ability.Get();
+			}
+
 			if (AbilitySpec.IsActive())
 			{
 				StatusText = FString::Printf(TEXT(" (Active %d)"), AbilitySpec.ActiveCount);
@@ -2675,12 +2682,12 @@ void UAbilitySystemComponent::Debug_Internal(FAbilitySystemComponentDebugInfo& I
 				StatusText = TEXT(" (InputBlocked)");
 				AbilityTextColor = FColor::Red;
 			}
-			else if (AbilitySpec.Ability->AbilityTags.HasAny(BlockedAbilityTags.GetExplicitGameplayTags()))
+			else if (AbilitySource->AbilityTags.HasAny(BlockedAbilityTags.GetExplicitGameplayTags()))
 			{
 				StatusText = TEXT(" (TagBlocked)");
 				AbilityTextColor = FColor::Red;
 			}
-			else if (AbilitySpec.Ability->CanActivateAbility(AbilitySpec.Handle, AbilityActorInfo.Get(), nullptr, nullptr, &FailureTags) == false)
+			else if (AbilitySource->CanActivateAbility(AbilitySpec.Handle, AbilityActorInfo.Get(), nullptr, nullptr, &FailureTags) == false)
 			{
 				StatusText = FString::Printf(TEXT(" (CantActivate %s)"), *FailureTags.ToString());
 				AbilityTextColor = FColor::Red;
