Fixing up an EngineNetworkCustomVersion mishap where two version numbers got swapped.

diff --git a/Engine/Plugins/Runtime/GameplayAbilities/Source/GameplayAbilities/Private/GameplayAbilityRepAnimMontage.cpp b/Engine/Plugins/Runtime/GameplayAbilities/Source/GameplayAbilities/Private/GameplayAbilityRepAnimMontage.cpp
--- a/Engine/Plugins/Runtime/GameplayAbilities/Source/GameplayAbilities/Private/GameplayAbilityRepAnimMontage.cpp
+++ b/Engine/Plugins/Runtime/GameplayAbilities/Source/GameplayAbilities/Private/GameplayAbilityRepAnimMontage.cpp
@@ -11,7 +11,7 @@ bool FGameplayAbilityRepAnimMontage::NetSerialize(FArchive& Ar, class UPackageMa
 	Ar.UsingCustomVersion(FEngineNetworkCustomVersion::Guid);
 
 	uint8 bIsMontage = 1;
-	if (Ar.EngineNetVer() > FEngineNetworkCustomVersion::DynamicMontageSerialization)
+	if (Ar.EngineNetVer() >= FEngineNetworkCustomVersion::DynamicMontageSerialization)
 	{
 		if (Ar.IsSaving())
 		{
@@ -86,7 +86,7 @@ bool FGameplayAbilityRepAnimMontage::NetSerialize(FArchive& Ar, class UPackageMa
 
 	if (!bIsMontage)
 	{
-		ensure(Ar.EngineNetVer() > FEngineNetworkCustomVersion::DynamicMontageSerialization);
+		ensure(Ar.EngineNetVer() >= FEngineNetworkCustomVersion::DynamicMontageSerialization);
 		Ar << BlendOutTime;
 		Ar << SlotName;
 	}
diff --git a/Engine/Source/Runtime/Core/Public/Misc/EngineNetworkCustomVersion.h b/Engine/Source/Runtime/Core/Public/Misc/EngineNetworkCustomVersion.h
--- a/Engine/Source/Runtime/Core/Public/Misc/EngineNetworkCustomVersion.h
+++ b/Engine/Source/Runtime/Core/Public/Misc/EngineNetworkCustomVersion.h
@@ -84,9 +84,8 @@ struct FEngineNetworkCustomVersion
 		SoftObjectPtrNetGuids = 29,				// Bump version to support replicating SoftObjectPtrs by NetGuid instead of raw strings.
 		SubObjectDestroyFlag = 30,				// Bump version to support subobject destruction message flags
 		GameStateReplicatedTimeAsDouble = 31,	// Bump version to support AGameStateBase::ReplicatedWorldTimeSeconds as double instead of float.
-		DynamicMontageSerialization = 32,
-
-		CustomVersions = 33,                    // Bump version to switch to using custom versions
+		CustomVersions = 32,                    // Bump version to switch to using custom versions
+		DynamicMontageSerialization = 33,		// Bump version to support dynamic montage serialization in the Gameplay Ability System
 
 		// -----<new versions can be added above this line>-------------------------------------------------
 		VersionPlusOne,
