Cleaning up some misconceptions (and bugs) about dependent FPredictionKeys inside the Ability System.  We no longer replicate FPredictionKey::Base because it's never used.

inspired by #contrib 11151

diff --git a/Engine/Plugins/Runtime/GameplayAbilities/Source/GameplayAbilities/Private/AbilitySystemDebugHUD.cpp b/Engine/Plugins/Runtime/GameplayAbilities/Source/GameplayAbilities/Private/AbilitySystemDebugHUD.cpp
--- a/Engine/Plugins/Runtime/GameplayAbilities/Source/GameplayAbilities/Private/AbilitySystemDebugHUD.cpp
+++ b/Engine/Plugins/Runtime/GameplayAbilities/Source/GameplayAbilities/Private/AbilitySystemDebugHUD.cpp
@@ -549,7 +549,7 @@ void AAbilitySystemDebugHUD::DrawDebugAbilitySystemComponent(UAbilitySystemCompo
 		// Draw Active GameplayEffect
 		for (FActiveGameplayEffect& Effect : &Component->ActiveGameplayEffects)
 		{
-			String = FString::Printf(TEXT("%s. [%d, %d] %.2f"), *Effect.Spec.ToSimpleString(), Effect.PredictionKey.Current, Effect.PredictionKey.Base, Effect.GetTimeRemaining(GameWorldTime));
+			String = FString::Printf(TEXT("%s. %s %.2f"), *Effect.Spec.ToSimpleString(), *Effect.PredictionKey.ToString(), Effect.GetTimeRemaining(GameWorldTime));
 			DrawWithBackground(Font, String, Color, EAlignHorizontal::Left, X, EAlignVertical::Top, Y);	
 		}
 	}
diff --git a/Engine/Plugins/Runtime/GameplayAbilities/Source/GameplayAbilities/Private/GameplayPrediction.cpp b/Engine/Plugins/Runtime/GameplayAbilities/Source/GameplayAbilities/Private/GameplayPrediction.cpp
--- a/Engine/Plugins/Runtime/GameplayAbilities/Source/GameplayAbilities/Private/GameplayPrediction.cpp
+++ b/Engine/Plugins/Runtime/GameplayAbilities/Source/GameplayAbilities/Private/GameplayPrediction.cpp
@@ -48,6 +48,9 @@ FReplicatedPredictionKeyItem& FReplicatedPredictionKeyItem::operator=(FReplicate
 /** The key to understanding this function is that when a key is received by the server, we note which connection gave it to us. We only serialize the key back to that client.  */
 bool FPredictionKey::NetSerialize(FArchive& Ar, class UPackageMap* Map, bool& bOutSuccess)
 {
+	Ar.UsingCustomVersion(FEngineNetworkCustomVersion::Guid);
+	const bool bReplicateDeprecatedBaseForDemoPurposes = (Ar.EngineNetVer() < FEngineNetworkCustomVersion::PredictionKeyBaseNotReplicated);
+
 	// First bit for valid key for this connection or not. (most keys are not valid)
 	uint8 ValidKeyForConnection = 0;
 	if (Ar.IsSaving())
@@ -61,9 +64,9 @@ bool FPredictionKey::NetSerialize(FArchive& Ar, class UPackageMap* Map, bool& bO
 	}
 	Ar.SerializeBits(&ValidKeyForConnection, 1);
 
-	// Second bit for base key (only if valid connection)
+	// Second bit for the now-deprecated base key (only if valid connection)
 	uint8 HasBaseKey = 0;
-	if (ValidKeyForConnection)
+	if (bReplicateDeprecatedBaseForDemoPurposes && ValidKeyForConnection)
 	{
 		if (Ar.IsSaving())
 		{
@@ -83,6 +86,7 @@ bool FPredictionKey::NetSerialize(FArchive& Ar, class UPackageMap* Map, bool& bO
 		Ar << Current;
 		if (HasBaseKey)
 		{
+			ensureMsgf(bReplicateDeprecatedBaseForDemoPurposes, TEXT("We should only ever be replicating a Base Key if we're loading an old demo, see bReplicateDeprecatedBaseForDemoPurposes"));
 			Ar << Base;
 		}
 	}	
@@ -103,7 +107,7 @@ void FPredictionKey::GenerateNewPredictionKey()
 {
 	static KeyType GKey = 1;
 	Current = GKey++;
-	if (GKey < 0)
+	if (GKey <= 0)
 	{
 		GKey = 1;
 	}
@@ -118,15 +122,11 @@ void FPredictionKey::GenerateDependentPredictionKey()
 		return;
 	}
 
-	KeyType Previous = 0;
+	KeyType Previous = Current;
 	if (Base == 0)
 	{
 		Base = Current;
 	}
-	else
-	{
-		Previous = Current;
-	}
 
 	GenerateNewPredictionKey();
 
diff --git a/Engine/Plugins/Runtime/GameplayAbilities/Source/GameplayAbilities/Private/Serialization/PredictionKeyNetSerializer.cpp b/Engine/Plugins/Runtime/GameplayAbilities/Source/GameplayAbilities/Private/Serialization/PredictionKeyNetSerializer.cpp
--- a/Engine/Plugins/Runtime/GameplayAbilities/Source/GameplayAbilities/Private/Serialization/PredictionKeyNetSerializer.cpp
+++ b/Engine/Plugins/Runtime/GameplayAbilities/Source/GameplayAbilities/Private/Serialization/PredictionKeyNetSerializer.cpp
@@ -29,7 +29,6 @@ struct FPredictionKeyNetSerializer
 	struct FQuantizedType
 	{
 		uint16 CurrentKey;
-		uint16 BaseKey;
 		uint16 OwningConnectionId;
 		uint16 bIsServerInitiated : 1;
 	};
@@ -78,10 +77,6 @@ void FPredictionKeyNetSerializer::Serialize(FNetSerializationContext& Context, c
 	if (Writer->WriteBool(bShouldWriteKey))
 	{
 		Writer->WriteBits(Value.CurrentKey, 15U);
-		if (Writer->WriteBool(Value.BaseKey > 0))
-		{
-			Writer->WriteBits(Value.BaseKey, 15U);
-		}
 	}
 }
 
@@ -94,10 +89,6 @@ void FPredictionKeyNetSerializer::Deserialize(FNetSerializationContext& Context,
 	if (Reader->ReadBool())
 	{
 		TempValue.CurrentKey = Reader->ReadBits(15U);
-		if (Reader->ReadBool())
-		{
-			TempValue.BaseKey = Reader->ReadBits(15U);
-		}
 	}
 
 	if (TempValue.CurrentKey > 0 && !TempValue.bIsServerInitiated)
@@ -119,7 +110,6 @@ void FPredictionKeyNetSerializer::Quantize(FNetSerializationContext&, const FNet
 	if (Source.IsValidKey())
 	{
 		Target.CurrentKey = uint16(Source.Current);
-		Target.BaseKey = Source.Base > 0 ? uint16(Source.Base) : uint16(0);
 
 		// The owning connection is only relevant if the key is valid.
 		if (UPackageMapClient* PackageMap = Cast<UPackageMapClient>(reinterpret_cast<UPackageMap*>(Source.GetPredictiveConnectionKey())))
@@ -139,7 +129,6 @@ void FPredictionKeyNetSerializer::Dequantize(FNetSerializationContext& Context,
 
 	Target.bIsServerInitiated = Source.bIsServerInitiated;
 	Target.Current = int16(Source.CurrentKey);
-	Target.Base = int16(Source.BaseKey);
 	if (Source.OwningConnectionId != 0)
 	{
 		UObject* UserData = Context.GetLocalConnectionUserData(Source.OwningConnectionId);
@@ -161,8 +150,8 @@ bool FPredictionKeyNetSerializer::IsEqual(FNetSerializationContext&, const FNetI
 		const QuantizedType& Value0 = *reinterpret_cast<const QuantizedType*>(Args.Source0);
 		const QuantizedType& Value1 = *reinterpret_cast<const QuantizedType*>(Args.Source1);
 
-		const bool bIsEqual = (Value0.CurrentKey == Value1.CurrentKey) & (Value0.BaseKey == Value1.BaseKey)
-			& (Value0.bIsServerInitiated == Value1.bIsServerInitiated) & (Value0.OwningConnectionId == Value1.OwningConnectionId);
+		const bool bIsEqual = (Value0.CurrentKey == Value1.CurrentKey) & 
+			(Value0.bIsServerInitiated == Value1.bIsServerInitiated) & (Value0.OwningConnectionId == Value1.OwningConnectionId);
 
 		return bIsEqual;
 	}
@@ -177,7 +166,7 @@ bool FPredictionKeyNetSerializer::IsEqual(FNetSerializationContext&, const FNetI
 
 		if (Value0.IsValidKey())
 		{
-			const bool bIsSameKey = (Value0.Current == Value1.Current) & (Value0.Base == Value1.Base) & (Value0.bIsServerInitiated == Value1.bIsServerInitiated);
+			const bool bIsSameKey = (Value0.Current == Value1.Current) & (Value0.bIsServerInitiated == Value1.bIsServerInitiated);
 			if (!bIsSameKey)
 			{
 				return false;
diff --git a/Engine/Plugins/Runtime/GameplayAbilities/Source/GameplayAbilities/Public/GameplayPrediction.h b/Engine/Plugins/Runtime/GameplayAbilities/Source/GameplayAbilities/Public/GameplayPrediction.h
--- a/Engine/Plugins/Runtime/GameplayAbilities/Source/GameplayAbilities/Public/GameplayPrediction.h
+++ b/Engine/Plugins/Runtime/GameplayAbilities/Source/GameplayAbilities/Public/GameplayPrediction.h
@@ -28,7 +28,7 @@ DECLARE_DELEGATE(FPredictionKeyEvent);
  *	For example, footsteps are completely client side and never interact with this system. But clients predicting their mana going from 100 to 90 when they cast a spell is 'client side prediction'.
  *		
  *	What do we currently predict?
- *	-Ability activation
+ *	-Initial GameplayAbility activation (and chained activation with caveats)
  *	-Triggered Events
  *	-GameplayEffect application:
  *		-Attribute modification	(EXCEPTIONS: Executions do not currently predict, only attribute modifiers)
@@ -67,39 +67,29 @@ DECLARE_DELEGATE(FPredictionKeyEvent);
  *	
  *	*** Ability Activation ***
  *	
- *	Ability Activation is a first class predictive action. Whenever a client predictively activates an ability, it explicitly asks the server and the server explicitly responds. Once an ability has been
- *	predictively activated, the client has a valid 'prediction window' where predictive side effects can happen which are not explicitly 'asked about'. (E.g., we do not explicitly ask 'Can I decrement mana, Can I
- *	put this ability on cooldown. Those actions are considered logically atomic with activating an ability).
+ *	Ability Activation is a first class predictive action -- it generates an initial prediction key. Whenever a client predictively activates an ability, it explicitly asks the server and the server explicitly responds. Once an ability has been
+ *	predictively activated (but the request has not yet been sent), the client has a valid 'prediction window' where predictive side effects can happen which are not explicitly 'asked about'. (E.g., we do not explicitly ask 'Can I decrement mana, Can I
+ *	put this ability on cooldown. Those actions are considered logically atomic with activating an ability). You can think of this prediction window as being the initial callstack of ActivateAbility. Once ActivateAbility ends, your
+ *  prediction window (and therefore your prediction key) is no longer valid. This is important, because many things can invalidate your prediction window such as any timers or latent nodes in your Blueprint; we do not predict over multiple frames.
  *	
  *	
  *	AbilitySystemComponent provides a set of functions for communicating ability activation between clients and server: TryActivateAbility -> ServerTryActivateAbility ->  ClientActivateAbility(Failed/Succeed).
  *	
  *	1. Client calls TryActivateAbility which generates a new FPredictionKey and calls ServerTryActivateAbility.
  *	2. Client continues (before hearing back from server) and calls ActivateAbility with the generated PredictionKey associated with the Ability's ActivationInfo.
- *	3. Any side effects that happen /before the call to ActivatAbility finish/ have the generated FPredictionKey associated with them.
- *	4. Server decides if the ability really happened in ServerTryActivateAbility, calls ClientActivateAbility(Failed/Succeed) and sets UAbilitySystemComponent::ReplicatedPredictionKey to the generated key that was sent.
+ *	3. Any side effects that happen /before the call to ActivateAbility finish/ have the generated FPredictionKey associated with them.
+ *	4. Server decides if the ability really happened in ServerTryActivateAbility, calls ClientActivateAbility(Failed/Succeed) and sets UAbilitySystemComponent::ReplicatedPredictionKey to the generated key that was sent with the request by the client.
  *	5. If client receives ClientAbilityFailed, it immediately kills the ability and rolls back side effects that were associated with the prediction key.
- *		5a. 'Rolling back' is accomplished via FPredictionKeyDelegates and FPredictionKey::NewRejectedDelegate/NewCaughtUpDelegate/NewRejectOrCaughtUpDelegate. 
- *
- *
- *		Registering the callback in TryActivateAbility:
- *		
- *		// If this PredictionKey is rejected, we will call OnClientActivateAbilityFailed.
- *		ThisPredictionKey.NewRejectedDelegate().BindUObject(this, &UAbilitySystemComponent::OnClientActivateAbilityFailed, Handle, ThisPredictionKey.Current);
- *		
- *		
- *		Invoking the callback in ClientActivateAbilityFailed_Implementation:
- *		FPredictionKeyDelegates::BroadcastRejectedDelegate(PredictionKey);
- *		
- *	6. If accepted, client must wait until property replication catches up (the Succeed RPC will be sent immediately, property replication will happen on its own). Once the ReplicatedPredictionKey catches up to the
- *		key used previous steps, the client can undo its predictive side effects. See UAbilitySystemComponent::OnRep_PredictionKey.
+ *		5a. 'Rolling back' logic is registered via FPredictionKeyDelegates and FPredictionKey::NewRejectedDelegate/NewCaughtUpDelegate/NewRejectOrCaughtUpDelegate.
+ *		5b. ClientAbilityFailed is really the only case where we 'reject' prediction keys and thus all of our current predictions rely on if an ability activates or not.
+ *	6. If ServerTryActivateAbility succeeds, client must wait until property replication catches up (the Succeed RPC will be sent immediately, property replication will happen on its own). Once the ReplicatedPredictionKey catches up to the
+ *		key used previous steps, the client can undo its predictive side effects.
+ *		See FReplicatedPredictionKeyItem::OnRep for the CatchUpTo logic. See UAbilitySystemComponent::ReplicatedPredictionKeyMap for how the keys actually get replicated. See ~FScopedPredictionWindow where the server acknowledges keys.
  *						
  *			 
- *			 
- *			 
  *	*** GameplayEffect Prediction ***
  *
- *	GameplayEffects are considered side effects of prediction and are not explicitly asked about.
+ *	GameplayEffects are considered side effects of ability activation and are not separately accepted/rejected.
  *	
  *	1. GameplayEffects are only applied on clients if there is a valid prediction key. (If no prediction key, it simply skips the application on client).
  *	2. Attributes, GameplayCues, and GameplayTags are all predicted if the GameplayEffect is predicted.
@@ -108,12 +98,14 @@ DECLARE_DELEGATE(FPredictionKeyEvent);
  *	4. On the server, the same prediction key is also set on the server's FActiveGameplayEffect that will be replicated down.
  *	5. As a client, if you get a replicated FActiveGameplayEffect with a valid prediction key on it, you check to see if you have an ActiveGameplayEffect with that same key, if there is match, we do not apply
  *		the 'on applied' type of logic, e.g., GameplayCues. The solves the "Redo" problem. However we will have 2 of the 'same' GameplayEffects in our ActiveGameplayEffects container, temporarily:
- *	6. At the same time, UAbilitySystemComponent::ReplicatedPredictionKey will catch up and the predictive effects will be removed. When they are removed in this case, we again check PredicitonKey and decide 
+ *	6. At the same time, FReplicatedPredictionKeyItem::OnRep will catch up and the predictive effects will be removed. When they are removed in this case, we again check PredictionKey and decide 
  *		if we should not do the 'On Remove' logic / GameplayCue.
  *		
  *	At this point, we have effectively predicted a gameplay effect as a side effect and handled the 'Undo' and 'Redo' problems.
  *		
- *	
+ *  See FActiveGameplayEffectsContainer::ApplyGameplayEffectSpec where it registers what to do when caught-up (RemoveActiveGameplayEffect_NoReturn).
+ *	See FActiveGameplayEffect::PostReplicatedAdd, FActiveGameplayEffect::PreReplicatedRemove, and FActiveGameplayCue::PostReplicatedAdd for examples of how FPredictionKey is associated with the GE's and GC's.
+ *
  *	*** Attribute Prediction ***
  *	
  *	Since attributes are replicated as standard uproperties, predicting modification to them can be tricky ("Override" problem). Instantaneous modification can be even harder since these are non stateful by nature.
@@ -126,7 +118,7 @@ DECLARE_DELEGATE(FPredictionKeyEvent);
  *	reaggregate our 'final value' after a replication happens.
  *	
  *	
- *	1. We treat predictive instant gameplay effects as infinite duration gamepaly effects. See UAbilitySystemComponent::ApplyGameplayEffectSpecToSelf.
+ *	1. We treat predictive instant gameplay effects as infinite duration gameplay effects. See UAbilitySystemComponent::ApplyGameplayEffectSpecToSelf.
  *	2. We have to *always* receive RepNotify calls on our attributes (not just when there is a change from last local value, since we will predict the change ahead of time). Done with REPNOTIFY_Always.
  *	3. In the attribute RepNotify, we call into the AbilitySystemComponent::ActiveGameplayEffects to update our 'final value' give the new 'base value'. the GAMEPLAYATTRIBUTE_REPNOTIFY can do this.
  *	4. Everything else will work like above (GameplayEffect prediction) : when the prediction key is caught up, the predictive GameplayEffect is removed and we will return to the server given value.
@@ -163,7 +155,8 @@ DECLARE_DELEGATE(FPredictionKeyEvent);
  *	There are some nuances to however, since the server will also run the code that triggers events. The server won't just wait to hear from the client. The server will keep a list of triggered abilities that have been
  *	activated from a predictive ability. When receiving a TryActivate from a triggered ability, the server will look to see if /it/ has already run this ability, and respond with that information.
  *	
- *	There is work left to do on Triggered Events and replication. (explained at the end).
+ *	The issue is we do not properly rollback these operations. There is work left to do on Triggered Events and replication. (explained at the end).
+ * 
  *	
  *	---------------------------------------------------------	
  *	
@@ -181,17 +174,20 @@ DECLARE_DELEGATE(FPredictionKeyEvent);
  *	sent to it from the client, whereas the client will generate a new PredictionKey each time it calls TryActivateAbility. The client *has* to generate a new PredictionKey for each ability activate, since each activate
  *	is not logically atomic. Each side effect produced in the chain of events has to have a unique PredictionKey. We cannot have GameplayEffects produced in X have the same PredictionKey produced in Z.
  *	
- *	To get around this, The prediction key of X is considered the Base key for Y and Z. The dependancy from Y to Z is kept completely client side, which is done in by FPredictionKeyDelegates::AddDependancy. We add delegates
+ *	To get around this, The prediction key of X is considered the Base key for Y and Z. The dependency from Y to Z is kept completely client side, which is done in by FPredictionKeyDelegates::AddDependancy. We add delegates
  *	to reject/catchup Z if Y rejected/confirmed.
  *	
  *	This dependency system allows us to have multiple predictive actions that are not logically atomic within a single prediction window/scope.
- *
+ *  
+ *  There is a problem though: because the dependencies are kept client side, the server does not actually know if it had previously rejected a dependent action. You can design around this issue by using activation tags
+ *  in your gameplay abilities. For instance, when predicting dependents GA_Combo1 -> GA_Combo2, you could make GA_Combo2 only activate if it has a GameplayTag given by GA_Combo1. Thus a rejection of GA_Combo1 would also
+ *  cause the server to reject the activation of GA_Combo2.
  *	
  *	
  *	*** Additional Prediction Windows (within an Ability) ***
  *	
- *	As stated, A prediction key is only usable during a single logical scope. Once ActivateAbility returns, we are essentially done with that key. If the ability is waiting on an external event or timer, by the time it
- *	returns, we will have gotten a confirm/reject from the server. Any side effects produced after this will no longer be tied to the lifespan of the original key.
+ *	As stated, A prediction key is only usable during a single logical scope. Once ActivateAbility returns, we are essentially done with that key. If the ability is waiting on an external event or timer, it's possible
+ *	we will have already received a confirm/reject from the server by the time we're ready to continue execution. Thus any additional side effects produced after the initial activation can no longer be tied to the lifespan of the original key.
  *	
  *	This isn't that bad, except that abilities will sometimes want to react to player input. For example, 'a hold down and charge' ability wants to instantly predict some stuff when the button is released. It is possible
  *	to create a new prediction window within an ability with FScopedPredictionWindow. 
@@ -219,7 +215,11 @@ DECLARE_DELEGATE(FPredictionKeyEvent);
  *	Triggered events do not explicitly replicate. E.g., if a triggered event only runs on the server, the client will never hear about it. This also prevents us from doing cross player/AI etc events. Support
  *	for this should eventually be added and it should follow the same pattern that GameplayEffect and GameplayCues follow (predict triggered event with a prediction key, ignore the RPC event if it has a prediction key).
  *	
- *		 
+ *	Big caveat with this whole system: Rollback of any chained activations (including triggered events) is currently not possible out of the box. The reason for this is because each ServerTryActivateAbility is going to be responded to in order.
+ *  Let's chain dependent GA's as an example: GA_Mispredict -> GA_Predict1. In this example, when GA_Mispredict is activated and predicted locally it will immediately also activate GA_Predict1. The client sends ServerTryActivateAbility for GA_Mispredict,
+ *  and the server rejects it (sending back a ClientActivateAbilityFailed). As it stands, we don't have any delegates that reject dependent abilities on the client (and the server isn't even aware there are dependencies).  On the server, it also
+ *  receives a ServerTryActivateAbility for GA_Predict1. Assuming that succeeds, the Client and Server are now both executing GA_Predict1 even though GA_Mispredict never happened. You can design around this by using the tag system to ensure GA_Mispredict succeeded.
+ *
  *	*** Predicting "Meta" Attributes such as Damage/Healing vs "real" attributes such as Health ***
  *	
  *	We are unable to apply meta attributes predictively. Meta attributes only work on instant effects, in the back end of GameplayEffect (Pre/Post Modify Attribute on the UAttributeSet). These events are not called when
@@ -229,7 +229,7 @@ DECLARE_DELEGATE(FPredictionKeyEvent);
  *	to the backend (UAttributeSet::PostModifyAttribute). 
  * 
  *			 
- *	*** Predicting ongoing multiplicitive GameplayEffects ***
+ *	*** Predicting ongoing multiplicative GameplayEffects ***
  *	
  *	There are also limitations when predicting % based gameplay effects. Since the server replicates down the 'final value' of an attribute, but not the entire aggregator chain of what is modifying it, we may run into cases where
  *	the client cannot accurately predict new gameplay effects.
@@ -293,13 +293,8 @@ struct GAMEPLAYABILITIES_API FPredictionKey
 	GENERATED_USTRUCT_BODY()
 
 	typedef int16 KeyType;
+	FPredictionKey() = default;
 
-	FPredictionKey()
-	: Current(0), Base(0), bIsStale(false), bIsServerInitiated(false)
-	{
-
-	}
-	
 	/** On the server, what network connection this was serialized on. */
 	UE_DEPRECATED(5.0, "No longer used in favor of GetPredictiveConnectionKey, to avoid holding a direct object reference.")
 	UPROPERTY(NotReplicated)
@@ -307,19 +302,19 @@ struct GAMEPLAYABILITIES_API FPredictionKey
 
 	/** The unique ID of this prediction key */
 	UPROPERTY()
-	int16	Current;
+	int16	Current = 0;
 
-	/** If non 0, the prediction key this was created from */
-	UPROPERTY()
-	int16	Base;
+	/** If non 0, the original prediction key this was created from (in a dependency chain) */
+	UPROPERTY(NotReplicated)
+	int16	Base = 0;
 
 	/** If stale, this key cannot be used for more prediction */
 	UPROPERTY(NotReplicated)
-	bool bIsStale;
+	bool bIsStale = false;
 
 	/** True if this was created as a server initiated activation key, used to identify server activations but cannot be used for prediction */
 	UPROPERTY()
-	bool bIsServerInitiated;
+	bool bIsServerInitiated = false;
 
 	/** Construct a new prediction key with no dependencies */
 	static FPredictionKey CreateNewPredictionKey(class UAbilitySystemComponent*);
@@ -376,14 +371,10 @@ struct GAMEPLAYABILITIES_API FPredictionKey
 		return (Current > 0) && (PredictiveConnectionKey == 0);
 	}
 
-	bool DependsOn(KeyType Key)
-	{
-		return (Current == Key || Base == Key);
-	}
-
 	bool operator==(const FPredictionKey& Other) const
 	{
-		return Current == Other.Current && Base == Other.Base && bIsServerInitiated == Other.bIsServerInitiated;
+		// we're omitting Base here because it's not replicated
+		return Current == Other.Current && bIsServerInitiated == Other.bIsServerInitiated;
 	}
 
 	FString ToString() const
@@ -393,8 +384,8 @@ struct GAMEPLAYABILITIES_API FPredictionKey
 
 	friend uint32 GetTypeHash(const FPredictionKey& InKey)
 	{
-		// A bit of overlap, but the == should cover it
-		return ((InKey.Current << 16) | (InKey.Base << 1) | (InKey.bIsServerInitiated & 1));
+		// we're omitting Base here because it's not replicated
+		return ((InKey.Current << 1) | (InKey.bIsServerInitiated & 1));
 	}
 
 	UPTRINT GetPredictiveConnectionKey() const { return PredictiveConnectionKey; }
@@ -404,16 +395,10 @@ private:
 
 	void GenerateNewPredictionKey();
 
-	FPredictionKey(int32 Key)
-		: Current(static_cast<KeyType>(Key)), Base(0), bIsStale(false), bIsServerInitiated(false)
-	{
-		check(Key >= std::numeric_limits<KeyType>::min() && Key <= std::numeric_limits<KeyType>::max());
-	}
-
-	FPredictionKey(int16 InKey, int16 PreviousKey)
-		: Current(InKey), Base(PreviousKey), bIsStale(false), bIsServerInitiated(false)
+	explicit FPredictionKey(int32 Key)
+		: Current(static_cast<KeyType>(Key))
 	{
-
+		check(Key >= 0 && Key <= std::numeric_limits<KeyType>::max());
 	}
 
 	/** On the server, uniquely identifies network connection this was serialized on/from.  See NetSerialize for additional information. */
diff --git a/Engine/Source/Runtime/Core/Public/Misc/EngineNetworkCustomVersion.h b/Engine/Source/Runtime/Core/Public/Misc/EngineNetworkCustomVersion.h
--- a/Engine/Source/Runtime/Core/Public/Misc/EngineNetworkCustomVersion.h
+++ b/Engine/Source/Runtime/Core/Public/Misc/EngineNetworkCustomVersion.h
@@ -86,6 +86,7 @@ struct FEngineNetworkCustomVersion
 		GameStateReplicatedTimeAsDouble = 31,	// Bump version to support AGameStateBase::ReplicatedWorldTimeSeconds as double instead of float.
 		CustomVersions = 32,                    // Bump version to switch to using custom versions
 		DynamicMontageSerialization = 33,		// Bump version to support dynamic montage serialization in the Gameplay Ability System
+		PredictionKeyBaseNotReplicated = 34,	// Bump version to stop FPredictionKey::Base from being replicated (it was unused).
 
 		// -----<new versions can be added above this line>-------------------------------------------------
 		VersionPlusOne,
