PR #11182: [GAS] Removed ensure that is triggered when a cue plays a particle system that gets preculled

diff --git a/Engine/Plugins/Runtime/GameplayAbilities/Source/GameplayAbilities/Private/GameplayCueNotifyTypes.cpp b/Engine/Plugins/Runtime/GameplayAbilities/Source/GameplayAbilities/Private/GameplayCueNotifyTypes.cpp
--- a/Engine/Plugins/Runtime/GameplayAbilities/Source/GameplayAbilities/Private/GameplayCueNotifyTypes.cpp
+++ b/Engine/Plugins/Runtime/GameplayAbilities/Source/GameplayAbilities/Private/GameplayCueNotifyTypes.cpp
@@ -426,7 +426,7 @@ bool FGameplayCueNotify_ParticleInfo::PlayParticleEffect(const FGameplayCueNotif
 						SpawnLocation, SpawnRotation, SpawnScale, bAutoDestroy, bAutoActivate, ENCPoolMethod::AutoRelease);
 				}
 
-				if (ensure(SpawnedFXSC))
+				if (SpawnedFXSC)
 				{
 					SpawnedFXSC->SetCastShadow(bCastShadow);
 				}
