Added serialization operator<< to FGameplayAbilitySpecHandle so users can invoke serialization on a FGameplayAbilitySpecHandle if needed. Enables compatibility with game project top-level structs that override NetSerialize().

diff --git a/Engine/Plugins/Runtime/GameplayAbilities/Source/GameplayAbilities/Public/GameplayAbilitySpecHandle.h b/Engine/Plugins/Runtime/GameplayAbilities/Source/GameplayAbilities/Public/GameplayAbilitySpecHandle.h
--- a/Engine/Plugins/Runtime/GameplayAbilities/Source/GameplayAbilities/Public/GameplayAbilitySpecHandle.h
+++ b/Engine/Plugins/Runtime/GameplayAbilities/Source/GameplayAbilities/Public/GameplayAbilitySpecHandle.h
@@ -40,6 +40,14 @@ struct FGameplayAbilitySpecHandle
 		return Handle != Other.Handle;
 	}
 
+	/** Operator to expose FGameplayAbilitySpecHandle serialization to custom serialization functions like NetSerialize overrides. */
+	friend FArchive& operator<<(FArchive& Ar, FGameplayAbilitySpecHandle& Value)
+	{
+		static_assert(sizeof(FGameplayAbilitySpecHandle) == 4, "If properties of FGameplayAbilitySpecHandle change, consider updating this operator implementation.");
+		Ar << Value.Handle;
+		return Ar;
+	}
+
 	friend uint32 GetTypeHash(const FGameplayAbilitySpecHandle& SpecHandle)
 	{
 		return ::GetTypeHash(SpecHandle.Handle);
