When using ForceReplication on the AbilitySystemComponent, also call ForceReplication if using a ReplicationProxy Actor to ensure the prompt delivery of GameplayCues and GameplayTags.

This is expected to increase bandwidth usage somewhat, but has been negligible in my tests (heavily drown out by noise).
The rollback CVar is "AbilitySystem.Fix.ForceReplicationAlsoUpdatesReplicatedProxyInterface 0" (the default is now 1).

diff --git a/Engine/Plugins/Runtime/GameplayAbilities/Source/GameplayAbilities/Private/AbilitySystemComponent.cpp b/Engine/Plugins/Runtime/GameplayAbilities/Source/GameplayAbilities/Private/AbilitySystemComponent.cpp
--- a/Engine/Plugins/Runtime/GameplayAbilities/Source/GameplayAbilities/Private/AbilitySystemComponent.cpp
+++ b/Engine/Plugins/Runtime/GameplayAbilities/Source/GameplayAbilities/Private/AbilitySystemComponent.cpp
@@ -42,6 +42,9 @@ static FAutoConsoleVariableRef CVarUseReplicationConditionForActiveGameplayEffec
 static bool bReplicateAbilitiesToSimulatedProxies = false;
 static FAutoConsoleVariableRef CVarReplicateGameplayAbilitiesToOwnerOnly(TEXT("AbilitySystem.Fix.ReplicateAbilitiesToSimulatedProxies"), bReplicateAbilitiesToSimulatedProxies, TEXT("Default: False.  When false, Gameplay Abilities replicate to AutonomousProxies only, not SimulatedProxies (surmised to be a bug)"));
 
+static bool bForceReplicationAlsoUpdatesReplicatedProxyInterface = true;
+static FAutoConsoleVariableRef CVarForceReplicationAlsoUpdatesReplicatedProxyInterface(TEXT("AbilitySystem.Fix.ForceReplicationAlsoUpdatesReplicatedProxyInterface"), bForceReplicationAlsoUpdatesReplicatedProxyInterface, TEXT("Default: True.  When true, Calling ForceReplication() on the AbilitySystemComponent will also call ForceReplication() on the ReplicationProxy to ensure prompt replication of Cues and Tags"));
+
 UAbilitySystemComponent::UAbilitySystemComponent(const FObjectInitializer& ObjectInitializer)
 	: Super(ObjectInitializer)
 	, GameplayTagCountContainer()
@@ -1658,6 +1661,15 @@ void UAbilitySystemComponent::ForceReplication()
 	if (OwningActor)
 	{
 		OwningActor->ForceNetUpdate();
+
+		if (ReplicationProxyEnabled && bForceReplicationAlsoUpdatesReplicatedProxyInterface)
+		{
+			IAbilitySystemReplicationProxyInterface* ReplicationProxy = GetReplicationInterface();
+			if (ReplicationProxy && ReplicationProxy != this)
+			{
+				ReplicationProxy->ForceReplication();
+			}
+		}
 	}
 }
 
