Improving Gameplay Ability System Debuggers.

Ability System Debug Commands have now been moved to a CheatManagerExtension.  This allows multiple clients to use the debug commands and they will be applied to the correct client.  Adding more comments and logging around why things can be rejected or work a certain way (e.g. AbilitySystem.Ability.Activate).

Improving the GameplayDebuggerCategory_Abilities so that it's easier to use to inspect and interpret the current GAS state.

GameplayCueManager now attempts to display the Gameplay Cues in non-overlapping fashion, so it's easier to see when multiple events occur for the same GameplayCue.

AbilitySystem debug commands are prefixed with AbilitySystem.Ability or AbilitySystem.Effect depending on if you're targeting Gameplay Abilities or Gameplay Effects.

Additional changes to the Visual Logger are coming.  The plan is for the Visual Logger to become the "event history" and this Gameplay Debugger to be a live-only view.

diff --git a/Engine/Plugins/Runtime/GameplayAbilities/Source/GameplayAbilities/Private/AbilitySystemCheatManagerExtension.cpp b/Engine/Plugins/Runtime/GameplayAbilities/Source/GameplayAbilities/Private/AbilitySystemCheatManagerExtension.cpp
new file mode 100644
--- /dev/null
+++ b/Engine/Plugins/Runtime/GameplayAbilities/Source/GameplayAbilities/Private/AbilitySystemCheatManagerExtension.cpp
@@ -0,0 +1,658 @@
+// Copyright Epic Games, Inc. All Rights Reserved.
+
+#include "AbilitySystemCheatManagerExtension.h"
+#include "AbilitySystemComponent.h"
+#include "AbilitySystemGlobals.h"
+#include "AbilitySystemLog.h"
+
+#include "Engine/Blueprint.h"
+#include "Engine/World.h"
+#include "GameFramework/Pawn.h"
+#include "GameFramework/PlayerController.h"
+#include "Logging/LogScopedVerbosityOverride.h"
+#include "UObject/Package.h"
+#include "UObject/UObjectIterator.h"
+
+namespace  UE::AbilitySystem::Private
+{
+	/**
+	 * Common logic used to fuzzy-find a requested class (or alternatively, a passed-in asset path).
+	 */
+	template<typename ClassToFind>
+	TSubclassOf<ClassToFind> FuzzyFindClass(FString SearchString)
+	{
+		TSubclassOf<ClassToFind> FoundClass;
+
+		// See if we passed in a class name of a Class that already exists in memory.
+		// If we passed-in Default__, just remove that part since we're looking for Classes, not CDO names.
+		SearchString.RemoveFromStart(TEXT("Default__"));
+		const int SearchStringLen = SearchString.Len();
+		int BestClassMatchLen = INT_MAX;
+		for (TObjectIterator<UClass> ClassIt; ClassIt; ++ClassIt)
+		{
+			const bool bClassMatches = ClassIt->IsChildOf(ClassToFind::StaticClass());
+			if (!bClassMatches)
+			{
+				continue;
+			}
+
+			// Class name search
+			const FString ClassName = ClassIt->GetName();
+			const int ClassNameLen = ClassName.Len();
+			if (ClassNameLen < BestClassMatchLen && ClassNameLen >= SearchStringLen)
+			{
+				bool bContains = ClassName.Contains(SearchString);
+				if (bContains)
+				{
+					FoundClass = *ClassIt;
+					BestClassMatchLen = ClassNameLen;
+				}
+			}
+		}
+
+		// If it wasn't a class name, then perhaps it was the path to a specific asset
+		if (!FoundClass)
+		{
+			FSoftObjectPath SoftObjectPath{ SearchString };
+			if (UObject* ReferencedObject = SoftObjectPath.ResolveObject())
+			{
+				if (UPackage* ReferencedPackage = Cast<UPackage>(ReferencedObject))
+				{
+					ReferencedObject = ReferencedPackage->FindAssetInPackage();
+				}
+
+				if (UBlueprint* ReferencedBlueprint = Cast<UBlueprint>(ReferencedObject))
+				{
+					FoundClass = ReferencedBlueprint->GeneratedClass;
+				}
+				else if (ClassToFind* ReferencedGA = Cast<ClassToFind>(ReferencedObject))
+				{
+					FoundClass = ReferencedGA->GetClass();
+				}
+			}
+		}
+
+		return FoundClass;
+	}
+}
+
+struct FScopedCanActivateAbilityLogGatherer : public FOutputDevice
+{
+	FScopedCanActivateAbilityLogGatherer(const FLogCategoryBase& InLogCategoryToCapture)
+		: LogCategoryToCapture{ InLogCategoryToCapture.GetCategoryName() }
+	{
+		GLog->AddOutputDevice(this);
+	}
+
+	~FScopedCanActivateAbilityLogGatherer()
+	{
+		GLog->RemoveOutputDevice(this);
+	}
+
+	/** Structure for the Log Entries */
+	struct FLogEntry
+	{
+		const FName Category;
+		const ELogVerbosity::Type Verbosity;
+		const FString Text;
+	};
+
+	TArray<FLogEntry>&& GetCapturedLogs()
+	{
+		FScopeLock _(&CriticalSection);
+		return MoveTemp(Logs);
+	}
+
+protected:
+	virtual void Serialize(const TCHAR* V, ELogVerbosity::Type Verbosity, const FName& InCategory) override
+	{
+		if (InCategory != LogCategoryToCapture)
+		{
+			return;
+		}
+
+		FScopeLock _(&CriticalSection);
+		Logs.Emplace(FLogEntry{ InCategory, Verbosity, V });
+	}
+
+	FCriticalSection	CriticalSection;
+	FName				LogCategoryToCapture;
+	TArray<FLogEntry>	Logs;
+};
+
+
+UAbilitySystemCheatManagerExtension::UAbilitySystemCheatManagerExtension()
+{
+#if WITH_SERVER_CODE && UE_WITH_CHEAT_MANAGER
+	if (HasAnyFlags(RF_ClassDefaultObject))
+	{
+		UCheatManager::RegisterForOnCheatManagerCreated(FOnCheatManagerCreated::FDelegate::CreateLambda(
+			[](UCheatManager* CheatManager)
+			{
+				CheatManager->AddCheatManagerExtension(NewObject<ThisClass>(CheatManager));
+			}));
+
+	}
+#endif
+}
+
+
+void UAbilitySystemCheatManagerExtension::AbilityGrant(const FString& AssetSearchString) const
+{
+	UAbilitySystemGlobals& AbilitySystemGlobals = UAbilitySystemGlobals::Get();
+
+	APlayerController* PC = GetPlayerController();
+	UAbilitySystemComponent* ASC = PC ? AbilitySystemGlobals.GetAbilitySystemComponentFromActor(PC->GetPawn()) : nullptr;
+	if (!ASC)
+	{
+		UE_LOG(LogConsoleResponse, Log, TEXT("%s did not have AbilitySystemComponent"), *GetNameSafe(PC));
+		return;
+	}
+
+	if (AssetSearchString.IsEmpty())
+	{
+		AbilityListGranted();
+		return;
+	}
+
+	// We couldn't find anything the user was searching for, so early out
+	TSubclassOf<UGameplayAbility> ActivateAbilityClass = UE::AbilitySystem::Private::FuzzyFindClass<UGameplayAbility>(AssetSearchString);
+	if (!ActivateAbilityClass)
+	{
+		UE_LOG(LogConsoleResponse, Log, TEXT("Could not find a valid Gameplay Ability based on Search String '%s'"), *AssetSearchString);
+		if (PC->GetWorld()->IsPlayInEditor())
+		{
+			UE_LOG(LogConsoleResponse, Log, TEXT("Reminder: If it's a non-native class, make sure it's loaded in the Editor."));
+		}
+
+		return;
+	}
+
+	// Check if it's already granted
+	if (const FGameplayAbilitySpec* ExistingSpec = ASC->FindAbilitySpecFromClass(ActivateAbilityClass))
+	{
+		UE_LOG(LogConsoleResponse, Log, TEXT("Existing Ability Spec '%s' on Player '%s' (It is already granted)."), *GetNameSafe(*ActivateAbilityClass), *GetNameSafe(PC));
+		return;
+	}
+
+	// If we're not the authority, we need to send the command to the server because we can't grant locally.
+	if (!ASC->IsOwnerActorAuthoritative())
+	{
+		const FString ServerCommand = FString::Printf(TEXT("%hs %s"), __func__, *AssetSearchString);
+		PC->ServerExec(ServerCommand);
+
+		UE_LOG(LogConsoleResponse, Log, TEXT("Sent Command 'ServerExec %s' from Player '%s' to Server because only the authority can grant abilities."), *ServerCommand, *GetNameSafe(PC));
+		return;
+	}
+
+	// It wasn't granted, let's grant it now.
+	FGameplayAbilitySpec AbilitySpec{ ActivateAbilityClass };
+	FGameplayAbilitySpecHandle SpecHandle = ASC->GiveAbility(AbilitySpec);
+	if (SpecHandle.IsValid())
+	{
+		UE_LOG(LogConsoleResponse, Log, TEXT("Successfully Granted '%s' on Player '%s'."), *GetNameSafe(ActivateAbilityClass.Get()), *GetNameSafe(PC));
+	}
+	else
+	{
+		UE_LOG(LogConsoleResponse, Log, TEXT("Failed to Grant '%s' on Player '%s'."), *GetNameSafe(ActivateAbilityClass.Get()), *GetNameSafe(PC));
+	}
+}
+
+void UAbilitySystemCheatManagerExtension::AbilityListGranted() const
+{
+	APlayerController* PC = GetPlayerController();
+	UAbilitySystemComponent* ASC = PC ? UAbilitySystemGlobals::Get().GetAbilitySystemComponentFromActor(PC->GetPawn()) : nullptr;
+	if (!ASC)
+	{
+		UE_LOG(LogConsoleResponse, Log, TEXT("%s did not have AbilitySystemComponent"), *GetNameSafe(PC));
+		return;
+	}
+
+	const UEnum* ExecutionEnumPtr = FindObject<UEnum>(nullptr, TEXT("/Script/GameplayAbilities.EGameplayAbilityNetExecutionPolicy"), true);
+	check(ExecutionEnumPtr && TEXT("Couldn't locate EGameplayAbilityNetExecutionPolicy enum!"));
+
+	const UEnum* SecurityEnumPtr = FindObject<UEnum>(nullptr, TEXT("/Script/GameplayAbilities.EGameplayAbilityNetSecurityPolicy"), true);
+	check(SecurityEnumPtr && TEXT("Couldn't locate EGameplayAbilityNetSecurityPolicy enum!"));
+
+	UE_LOG(LogConsoleResponse, Log, TEXT("Granted abilities to %s (ASC: '%s'):"), *PC->GetName(), *ASC->GetFullName());
+
+	for (FGameplayAbilitySpec& Activatable : ASC->GetActivatableAbilities())
+	{
+		const TCHAR* ActiveText = Activatable.IsActive() ? TEXT("**ACTIVE**") : TEXT("");
+		UE_LOG(LogConsoleResponse, Log, TEXT("   %s (%s - %s) %s"), *Activatable.Ability->GetName(), *ExecutionEnumPtr->GetDisplayNameTextByIndex(Activatable.Ability->GetNetExecutionPolicy()).ToString(), *SecurityEnumPtr->GetDisplayNameTextByIndex(Activatable.Ability->GetNetSecurityPolicy()).ToString(), ActiveText);
+	}
+}
+
+void UAbilitySystemCheatManagerExtension::AbilityActivate(const FString& PartialName) const
+{
+	UAbilitySystemGlobals& AbilitySystemGlobals = UAbilitySystemGlobals::Get();
+
+	APlayerController* PC = GetPlayerController();
+	UAbilitySystemComponent* ASC = PC ? AbilitySystemGlobals.GetAbilitySystemComponentFromActor(PC->GetPawn()) : nullptr;
+	if (!ASC)
+	{
+		UE_LOG(LogConsoleResponse, Log, TEXT("%s did not have AbilitySystemComponent"), *GetNameSafe(PC));
+		return;
+	}
+
+	if (PartialName.IsEmpty())
+	{
+		AbilityListGranted();
+		return;
+	}
+
+	// Start by figuring out if we're trying to execute by GameplayTag.
+	const FName SearchName{ TCHAR_TO_ANSI(*PartialName), FNAME_Find };
+	if (SearchName != NAME_None)
+	{
+		const FGameplayTag GameplayTag = FGameplayTag::RequestGameplayTag(SearchName, false);
+		if (GameplayTag.IsValid())
+		{
+			FGameplayTagContainer GameplayTagContainer{ GameplayTag };
+			TArray<FGameplayAbilitySpec*> MatchingSpecs;
+			constexpr bool bTagRequirementsMustMatch = false;
+			ASC->GetActivatableGameplayAbilitySpecsByAllMatchingTags(GameplayTagContainer, MatchingSpecs, bTagRequirementsMustMatch);
+			if (MatchingSpecs.Num() > 0)
+			{
+				FString MatchingAbilitiesString = FString::JoinBy(MatchingSpecs, TEXT(", "), [](FGameplayAbilitySpec* Item) { return Item->GetDebugString(); });
+
+				bool bSuccess = ASC->TryActivateAbilitiesByTag(GameplayTagContainer);
+				if (bSuccess)
+				{
+					UE_LOG(LogConsoleResponse, Log, TEXT("Requested Tag '%s' successfully executed one of: %s."), *SearchName.ToString(), *MatchingAbilitiesString);
+				}
+				else
+				{
+					UE_LOG(LogConsoleResponse, Log, TEXT("Requested Tag '%s' was expected to execute one of: %s. But it failed to do so due to tag requirements."), *SearchName.ToString(), *MatchingAbilitiesString);
+				}
+			}
+			else
+			{
+				UE_LOG(LogConsoleResponse, Log, TEXT("Requested Tag '%s' matched no given Gameplay Abilities to %s."), *SearchName.ToString(), *GetNameSafe(PC));
+			}
+
+			return;
+		}
+	}
+
+	// We couldn't find anything the user was searching for, so early out
+	TSubclassOf<UGameplayAbility> ActivateAbilityClass = UE::AbilitySystem::Private::FuzzyFindClass<UGameplayAbility>(PartialName);
+	if (!ActivateAbilityClass)
+	{
+		UE_LOG(LogConsoleResponse, Log, TEXT("Could not find a valid Gameplay Ability based on Search String '%s'"), *PartialName);
+		return;
+	}
+
+	// Check if we already have this ability granted, because it affects how we can execute it.
+	const FGameplayAbilitySpec* ExistingSpec = ASC->FindAbilitySpecFromClass(ActivateAbilityClass);
+
+	// If we're not the authority, we should check if we need to send the command to the server
+	if (!ASC->IsOwnerActorAuthoritative())
+	{
+		const UGameplayAbility* Ability = ActivateAbilityClass.GetDefaultObject();
+		EGameplayAbilityNetExecutionPolicy::Type NetExecutionPolicy = Ability->GetNetExecutionPolicy();
+
+		const bool bMustActivateLocally = (NetExecutionPolicy == EGameplayAbilityNetExecutionPolicy::LocalOnly || NetExecutionPolicy == EGameplayAbilityNetExecutionPolicy::LocalPredicted);
+		if (!ExistingSpec && bMustActivateLocally)
+		{
+			UE_LOG(LogConsoleResponse, Log, TEXT("Cannot Activate '%s' on '%s' as the ability is not granted.  Since its NetExecutionPolicy must activate locally, you must first grant it (use the AbilityGrant command)."), *GetNameSafe(ActivateAbilityClass.Get()), *GetNameSafe(PC));
+			return;
+		}
+
+		const bool bSendToServer = (NetExecutionPolicy >= EGameplayAbilityNetExecutionPolicy::ServerInitiated);
+		if (bSendToServer)
+		{
+			const FString ServerCommand = FString::Printf(TEXT("%hs %s"), __func__, *PartialName);
+			PC->ServerExec(ServerCommand);
+
+			UE_LOG(LogConsoleResponse, Log, TEXT("Sent Command '%s' from Player '%s' to Server (Reason: Net Execution Policy)."), *ServerCommand, *GetNameSafe(PC));
+			return;
+		}
+	}
+
+
+	// We found what the user was searching for, so let's try to activate it.  First, let's assume the Ability was already granted.
+	LOG_SCOPE_VERBOSITY_OVERRIDE(LogAbilitySystem, ELogVerbosity::VeryVerbose);
+	FScopedCanActivateAbilityLogGatherer LogGatherer{ LogAbilitySystem };
+	FScopedCanActivateAbilityLogEnabler LogEnabler;
+	bool bSuccess = ASC->TryActivateAbilityByClass(ActivateAbilityClass);
+	if (bSuccess)
+	{
+		UE_LOG(LogConsoleResponse, Log, TEXT("Successfully Activated Granted Ability '%s' on Player '%s'."), *GetNameSafe(ActivateAbilityClass.Get()), *GetNameSafe(PC));
+		return;
+	}
+
+	// It wasn't granted (or we failed to activate it even though it was granted).  We can't really differentiate those two, so let's grant it, then activate it.
+	FGameplayAbilitySpec AbilitySpec{ ActivateAbilityClass };
+	FGameplayAbilitySpecHandle SpecHandle = ASC->GiveAbilityAndActivateOnce(AbilitySpec);
+	if (SpecHandle.IsValid())
+	{
+		UE_LOG(LogConsoleResponse, Log, TEXT("Successfully Granted, then Activated '%s' on Player '%s'."), *GetNameSafe(ActivateAbilityClass.Get()), *GetNameSafe(PC));
+	}
+	else
+	{
+		UE_LOG(LogConsoleResponse, Log, TEXT("Failed to Grant & Activate '%s' on Player '%s'. Logs:"), *GetNameSafe(ActivateAbilityClass.Get()), *GetNameSafe(PC));
+
+		TArray<FScopedCanActivateAbilityLogGatherer::FLogEntry> CapturedLogs = LogGatherer.GetCapturedLogs();
+		for (const FScopedCanActivateAbilityLogGatherer::FLogEntry& LogEntry : CapturedLogs)
+		{
+			FMsg::LogV(__FILE__, __LINE__, LogConsoleResponse.GetCategoryName(), LogEntry.Verbosity, *LogEntry.Text, {});
+		}
+	}
+}
+
+void UAbilitySystemCheatManagerExtension::AbilityCancel(const FString& PartialName) const
+{
+	APlayerController* PC = GetPlayerController();
+	UAbilitySystemComponent* ASC = PC ? UAbilitySystemGlobals::Get().GetAbilitySystemComponentFromActor(PC->GetPawn()) : nullptr;
+	if (!ASC)
+	{
+		UE_LOG(LogConsoleResponse, Log, TEXT("PlayerController %s did not have an AbilitySystemComponent"), *GetNameSafe(PC));
+		return;
+	}
+
+	FString SearchString = PartialName;
+	SearchString.RemoveFromStart(TEXT("Default__"));
+
+	bool bFound = false;
+	for (FGameplayAbilitySpec& GASpec : ASC->GetActivatableAbilities())
+	{
+		if (GASpec.IsActive())
+		{
+			if (!GASpec.Ability || !GASpec.Ability->GetName().Contains(SearchString))
+			{
+				continue;
+			}
+
+			TArray<UGameplayAbility*> ActiveAbilities = GASpec.GetAbilityInstances();
+			if (ActiveAbilities.Num() > 0)
+			{
+				for (UGameplayAbility* Instance : ActiveAbilities)
+				{
+					if (Instance)
+					{
+						UE_LOG(LogConsoleResponse, Log, TEXT("%s (%s): Cancelling (instanced) %s"), *PC->GetName(), *ASC->GetName(), *Instance->GetName());
+						Instance->CancelAbility(GASpec.Handle, ASC->AbilityActorInfo.Get(), GASpec.ActivationInfo, true);
+						bFound = true;
+					}
+				}
+			}
+			else
+			{
+				UE_LOG(LogConsoleResponse, Log, TEXT("%s (%s): Cancelling (non-instanced) %s"), *PC->GetName(), *ASC->GetName(), *GASpec.Ability->GetName());
+				GASpec.Ability->CancelAbility(GASpec.Handle, ASC->AbilityActorInfo.Get(), GASpec.ActivationInfo, true);
+				bFound = true;
+			}
+		}
+	}
+
+	if (!bFound)
+	{
+		UE_LOG(LogConsoleResponse, Log, TEXT("Found no Active Gameplay Abilities on %s (%s) that matched '%s'"), *PC->GetName(), *ASC->GetName(), *SearchString);
+	}
+}
+
+void UAbilitySystemCheatManagerExtension::EffectListActive() const
+{
+	APlayerController* PC = GetPlayerController();
+	UAbilitySystemComponent* ASC = PC ? UAbilitySystemGlobals::Get().GetAbilitySystemComponentFromActor(PC->GetPawn()) : nullptr;
+	if (!ASC)
+	{
+		UE_LOG(LogConsoleResponse, Log, TEXT("PlayerController %s did not have an AbilitySystemComponent"), *GetNameSafe(PC));
+		return;
+	}
+
+	UE_LOG(LogConsoleResponse, Log, TEXT("Active Gameplay Effects on %s (%s):"), *PC->GetName(), *ASC->GetName());
+
+	const FActiveGameplayEffectsContainer& ActiveGEContainer = ASC->GetActiveGameplayEffects();
+	for (const FActiveGameplayEffectHandle& ActiveGEHandle : ActiveGEContainer.GetAllActiveEffectHandles())
+	{
+		if (const FActiveGameplayEffect* ActiveGE = ASC->GetActiveGameplayEffect(ActiveGEHandle))
+		{
+			const UGameplayEffect* GESpecDef = ActiveGE->Spec.Def;
+			const FString ActiveState = ActiveGE->bIsInhibited ? FString(TEXT("Inhibited")) : FString::Printf(TEXT("Active (Stack: %d)"), ActiveGE->Spec.GetStackCount());
+
+			UE_LOG(LogConsoleResponse, Log, TEXT(" [%u] %s: %s"), GetTypeHash(ActiveGEHandle), *GetNameSafe(GESpecDef), *ActiveState);
+		}
+	}
+}
+
+void UAbilitySystemCheatManagerExtension::EffectApply(const FString& PartialName, float EffectLevel /*= FGameplayEffectConstants::INVALID_LEVEL*/) const
+{
+	APlayerController* PC = GetPlayerController();
+	UAbilitySystemComponent* ASC = PC ? UAbilitySystemGlobals::Get().GetAbilitySystemComponentFromActor(PC->GetPawn()) : nullptr;
+	if (!ASC)
+	{
+		UE_LOG(LogConsoleResponse, Log, TEXT("PlayerController %s did not have an AbilitySystemComponent"), *GetNameSafe(PC));
+		return;
+	}
+
+	// We couldn't find anything the user was searching for, so early out
+	TSubclassOf<UGameplayEffect> GameplayEffectClass = UE::AbilitySystem::Private::FuzzyFindClass<UGameplayEffect>(PartialName);
+	if (!GameplayEffectClass)
+	{
+		UE_LOG(LogConsoleResponse, Log, TEXT("Could not find a valid Gameplay Effect based on Search String '%s'"), *PartialName);
+		return;
+	}
+
+	// Create a GameplayEffectSpec that executes the passed-in parameters
+	FGameplayEffectContextHandle GEContextHandle = ASC->MakeEffectContext();
+	APawn* Pawn = PC->GetPawn();
+	GEContextHandle.AddInstigator(Pawn, Pawn);
+	GEContextHandle.AddOrigin(Pawn->GetActorLocation());
+	FGameplayEffectSpec GESpec{ GameplayEffectClass.GetDefaultObject(), GEContextHandle, EffectLevel };
+
+	// We need to create a new valid prediction key (as if we just started activating an ability) so the GE will fire even if we are not the authority.
+	FPredictionKey FakePredictionKey = (PC->GetLocalRole() == ENetRole::ROLE_Authority) ? FPredictionKey() : FPredictionKey::CreateNewPredictionKey(ASC);
+	FActiveGameplayEffectHandle ActiveGEHandle = ASC->ApplyGameplayEffectSpecToSelf(GESpec, FakePredictionKey);
+	if (ActiveGEHandle.IsValid())
+	{
+		UE_LOG(LogConsoleResponse, Log, TEXT("Successfully Applied Gameplay Effect '%s' on Player '%s'."), *GameplayEffectClass->GetName(), *PC->GetName());
+	}
+	else
+	{
+		UE_LOG(LogConsoleResponse, Log, TEXT("Failed to Apply Gameplay Effect '%s' on Player '%s'."), *GameplayEffectClass->GetName(), *PC->GetName());
+	}
+}
+
+void UAbilitySystemCheatManagerExtension::EffectRemove(const FString& NameOrHandle) const
+{
+	UAbilitySystemGlobals& AbilitySystemGlobals = UAbilitySystemGlobals::Get();
+
+	APlayerController* PC = GetPlayerController();
+	UAbilitySystemComponent* ASC = PC ? AbilitySystemGlobals.GetAbilitySystemComponentFromActor(PC->GetPawn()) : nullptr;
+	if (!ASC)
+	{
+		UE_LOG(LogConsoleResponse, Log, TEXT("PlayerController %s did not have an AbilitySystemComponent"), *GetNameSafe(PC));
+		return;
+	}
+
+	uint32 SearchHandleValue = 0;
+	if (NameOrHandle.IsNumeric())
+	{
+		int64 SearchHash64 = FCString::Atoi64(*NameOrHandle);
+		if (SearchHash64 > 0)
+		{
+			SearchHandleValue = static_cast<uint32>(SearchHash64 & UINT_MAX);
+		}
+	}
+
+	bool bAnyMatch = false;
+	const FActiveGameplayEffectsContainer& ActiveGEContainer = ASC->GetActiveGameplayEffects();
+	for (const FActiveGameplayEffectHandle& ActiveGEHandle : ActiveGEContainer.GetAllActiveEffectHandles())
+	{
+		if (const FActiveGameplayEffect* ActiveGE = ASC->GetActiveGameplayEffect(ActiveGEHandle))
+		{
+			const UGameplayEffect* GESpecDef = ActiveGE->Spec.Def;
+			const uint32 ActiveGEHandleValue = GetTypeHash(ActiveGEHandle);
+
+			const bool bMatch = (SearchHandleValue > 0) ? (ActiveGEHandleValue == SearchHandleValue) : GESpecDef->GetName().Contains(NameOrHandle);
+			if (bMatch)
+			{
+				UE_LOG(LogConsoleResponse, Log, TEXT("Removing Active GE Handle %u (%s)"), ActiveGEHandleValue, *GESpecDef->GetName());
+				ASC->RemoveActiveGameplayEffect(ActiveGEHandle);
+				bAnyMatch = true;
+			}
+		}
+	}
+
+	if (!bAnyMatch)
+	{
+		UE_LOG(LogConsoleResponse, Log, TEXT("Found no Active Gameplay Effects on %s (%s) that matched '%s'"), *PC->GetName(), *ASC->GetName(), *NameOrHandle);
+	}
+}
+
+
+/**
+ * Handy wrapper function that executes a named exec function with the given Arguments, but through the CheatManager.  This allows the function to execute with context
+ * if RPC'd to the server, thus routing to the proper PlayerController.  All of the functions support -server as an optional parameter, routing the function directly to the server.
+ */
+void ForwardToAbilitySystemCheatManagerExtension(const FName& ExecFunctionName, const TArray<FString>& Args, UWorld* World, FOutputDevice& OutputDevice)
+{
+	if (APlayerController* PlayerController = World->GetFirstPlayerController())
+	{
+		APawn* PCPawn = PlayerController->GetPawnOrSpectator();
+		if (UCheatManager* CheatManager = PlayerController->CheatManager)
+		{
+			// All of these commands allow -server to force to be sent to the server
+			TArray<FString> MutableArgs = Args;
+			const bool bExecuteOnServer = MutableArgs.RemoveSingle(TEXT("-server")) > 0;
+
+			const FString Cmd = FString::Printf(TEXT("%s %s"), *ExecFunctionName.ToString(), *FString::Join(MutableArgs, TEXT(" ")));
+			if (bExecuteOnServer)
+			{
+				OutputDevice.Logf(TEXT("Executing 'ServerExec %s':"), *Cmd);
+				PlayerController->ServerExec(Cmd);
+			}
+			else
+			{
+				const bool bSuccess = CheatManager->ProcessConsoleExec(*Cmd, OutputDevice, PCPawn);
+				if (!bSuccess)
+				{
+					OutputDevice.Logf(TEXT("CheatManager %s on PlayerController %s did not process command '%s'"), *GetNameSafe(CheatManager), *GetNameSafe(PlayerController), *Cmd);
+				}
+			}
+		}
+		else
+		{
+			OutputDevice.Logf(TEXT("PlayerController %s did not have a CheatManager"), *GetNameSafe(PlayerController));
+		}
+	}
+	else
+	{
+		OutputDevice.Logf(TEXT("Could not find Player Controller in %s"), *GetNameSafe(World));
+	}
+}
+
+//
+// AbilitySystem.Ability Debug Commands forward to the CheatManager
+//
+
+FAutoConsoleCommand DebugAbilitySystemAbilityGrantCommand(TEXT("AbilitySystem.Ability.Grant"), TEXT("<ClassName/AssetName>. Grants an Ability to the Player. Granting only happens on the Authority, so this command will be sent to the Server."),
+	FConsoleCommandWithWorldArgsAndOutputDeviceDelegate::CreateLambda([](const TArray<FString>& Args, UWorld* World, FOutputDevice& OutputDevice)
+		{
+			if (Args.Num() < 1)
+			{
+				OutputDevice.Logf(TEXT("Missing Arguments: <ClassName/AssetName>"));
+				OutputDevice.Logf(TEXT(" ClassName: If supplied, the ClassName of the loaded GameplayAbility to Grant."));
+				OutputDevice.Logf(TEXT(" AssetName: If supplied, the AssetName of the loaded GameplayAbility to Grant."));
+
+				ForwardToAbilitySystemCheatManagerExtension(GET_FUNCTION_NAME_CHECKED(UAbilitySystemCheatManagerExtension, AbilityListGranted), Args, World, OutputDevice);
+				return;
+			}
+
+			ForwardToAbilitySystemCheatManagerExtension(GET_FUNCTION_NAME_CHECKED(UAbilitySystemCheatManagerExtension, AbilityGrant), Args, World, OutputDevice);
+		}), ECVF_Cheat);
+
+FAutoConsoleCommand DebugAbilitySystemAbilityCancelCommand(TEXT("AbilitySystem.Ability.Cancel"), TEXT("[-Server] <PartialName>. Cancels (prematurely Ends) a currently executing Gameplay Ability"),
+	FConsoleCommandWithWorldArgsAndOutputDeviceDelegate::CreateLambda([](const TArray<FString>& Args, UWorld* World, FOutputDevice& OutputDevice)
+		{
+			if (Args.Num() < 1)
+			{
+				OutputDevice.Logf(TEXT("Missing Arguments: <PartialName>"));
+				OutputDevice.Logf(TEXT(" -Server: Indicates this command should be sent to the server, rather than executed locally."));
+				OutputDevice.Logf(TEXT(" PartialName: A substring of the currently active abilities to cancel. All matching abilities will be cancelled."));
+
+				ForwardToAbilitySystemCheatManagerExtension(GET_FUNCTION_NAME_CHECKED(UAbilitySystemCheatManagerExtension, AbilityListGranted), Args, World, OutputDevice);
+				return;
+			}
+
+			ForwardToAbilitySystemCheatManagerExtension(GET_FUNCTION_NAME_CHECKED(UAbilitySystemCheatManagerExtension, AbilityCancel), Args, World, OutputDevice);
+		}), ECVF_Cheat);
+
+FAutoConsoleCommand DebugAbilitySystemAbilityActivateCommand(TEXT("AbilitySystem.Ability.Activate"), TEXT("[-Server] <TagName/ClassName/AssetName>. Activate a Gameplay Ability.\nSubstring name matching works for Activation Tags (on already granted abilities), Asset Paths (on non-granted abilities), or Class Names on both."),
+	FConsoleCommandWithWorldArgsAndOutputDeviceDelegate::CreateLambda([](const TArray<FString>& Args, UWorld* World, FOutputDevice& OutputDevice)
+		{
+			if (Args.Num() < 1)
+			{
+				OutputDevice.Logf(TEXT("This command Activates an Ability.  If the Ability is not already granted, then it is activated on the server using GiveAbilityAndActivateOnce, which will remove it after completion."));
+				OutputDevice.Logf(TEXT("Missing Arguments: <TagName or ClassName or AssetName>"));
+				OutputDevice.Logf(TEXT(" -Server: Indicates this command should be sent to the server, rather than executed locally."));
+				OutputDevice.Logf(TEXT(" TagName: If the argument is a GameplayTag, it will be activated using TryActivateAbilitiesByTag which requires the Gameplay Ability be previously granted."));
+				OutputDevice.Logf(TEXT(" ClassName or AssetName: A substring of class/asset.  Only the best match that is currently loaded will be considered."));
+
+				ForwardToAbilitySystemCheatManagerExtension(GET_FUNCTION_NAME_CHECKED(UAbilitySystemCheatManagerExtension, AbilityListGranted), Args, World, OutputDevice);
+				return;
+			}
+
+			ForwardToAbilitySystemCheatManagerExtension(GET_FUNCTION_NAME_CHECKED(UAbilitySystemCheatManagerExtension, AbilityActivate), Args, World, OutputDevice);
+		}), ECVF_Cheat);
+
+
+FAutoConsoleCommand DebugAbilitySystemAbilityListGrantedCommand(TEXT("AbilitySystem.Ability.ListGranted"), TEXT("List the Gameplay Abilities that are granted to the local player"),
+	FConsoleCommandWithWorldArgsAndOutputDeviceDelegate::CreateLambda([](const TArray<FString>& Args, UWorld* World, FOutputDevice& OutputDevice)
+		{
+			ForwardToAbilitySystemCheatManagerExtension(GET_FUNCTION_NAME_CHECKED(UAbilitySystemCheatManagerExtension, AbilityListGranted), Args, World, OutputDevice);
+		}), ECVF_Cheat);
+
+//
+// AbilitySystem.Effect Debug Commands forward to the CheatManager
+//
+
+FAutoConsoleCommand DebugAbilitySystemEffectListActiveCommand(TEXT("AbilitySystem.Effect.ListActive"), TEXT("[-Server] Lists all of the Gameplay Effects currently active on the Player"),
+	FConsoleCommandWithWorldArgsAndOutputDeviceDelegate::CreateLambda([](const TArray<FString>& Args, UWorld* World, FOutputDevice& OutputDevice)
+		{
+			ForwardToAbilitySystemCheatManagerExtension(GET_FUNCTION_NAME_CHECKED(UAbilitySystemCheatManagerExtension, EffectListActive), Args, World, OutputDevice);
+		}), ECVF_Cheat);
+
+FAutoConsoleCommand DebugAbilitySystemEffectRemoveCommand(TEXT("AbilitySystem.Effect.Remove"), TEXT("[-Server] <Handle/Name>. Remove a Gameplay Effect that is currently active on the Player"),
+	FConsoleCommandWithWorldArgsAndOutputDeviceDelegate::CreateLambda([](const TArray<FString>& Args, UWorld* World, FOutputDevice& OutputDevice)
+		{
+			if (Args.Num() < 1)
+			{
+				OutputDevice.Logf(TEXT("Missing Arguments: [-Server] <Handle or Name>"));
+				OutputDevice.Logf(TEXT(" -Server: Indicates this command should be sent to the server, rather than executed locally."));
+				OutputDevice.Logf(TEXT(" Handle: The numeric handle of the Gameplay Effect to remove (found through EffectListActive). This is not compatible with -server."));
+				OutputDevice.Logf(TEXT(" Name: Partial name of the Gameplay Effect to remove."));
+				return;
+			}
+
+			// Let's just make sure the user isn't about to do something that's going to be profoundly confusing
+			const bool bExecuteOnServer = Args.Contains("-server");
+			if (bExecuteOnServer)
+			{
+				const bool bArgIsHandle = Args.ContainsByPredicate([](const FString& Value) { return FCString::Atoi64(*Value) > 0; });
+				if (bArgIsHandle)
+				{
+					UE_LOG(LogConsoleResponse, Log, TEXT("Error: Search by Handle Value is not permitted with -server (because ActiveGE Handles are not replicated to the same handle value)"));
+					return;
+				}
+			}
+
+			ForwardToAbilitySystemCheatManagerExtension(GET_FUNCTION_NAME_CHECKED(UAbilitySystemCheatManagerExtension, EffectRemove), Args, World, OutputDevice);
+		}), ECVF_Cheat);
+
+FAutoConsoleCommand DebugAbilitySystemEffectApply(TEXT("AbilitySystem.Effect.Apply"), TEXT("[-Server] <Class/AssetName> [Level]. Apply a Gameplay Effect on the Player.  Substring name matching works for Asset Tags, Asset Paths, or Class Names.  Use -Server to send to the server (default is apply locally)."),
+	FConsoleCommandWithWorldArgsAndOutputDeviceDelegate::CreateLambda([](const TArray<FString>& Args, UWorld* World, FOutputDevice& OutputDevice)
+		{
+			if (Args.Num() < 1)
+			{
+				OutputDevice.Logf(TEXT("Missing Arguments: [-Server] <MatchString> [Level]"));
+				OutputDevice.Logf(TEXT(" -Server: Indicates this command should be sent to the server, rather than predicted locally."));
+				OutputDevice.Logf(TEXT(" MatchString: Supply a substring of a class or asset name to match. Only loaded GameplayEffects are searched."));
+				OutputDevice.Logf(TEXT(" Level: Optionally supply a Level such as 4.  Can be omitted (in which case the GE is level-less)."));
+				return;
+			}
+
+			ForwardToAbilitySystemCheatManagerExtension(GET_FUNCTION_NAME_CHECKED(UAbilitySystemCheatManagerExtension, EffectApply), Args, World, OutputDevice);
+		}), ECVF_Cheat);
diff --git a/Engine/Plugins/Runtime/GameplayAbilities/Source/GameplayAbilities/Private/AbilitySystemCheatManagerExtension.h b/Engine/Plugins/Runtime/GameplayAbilities/Source/GameplayAbilities/Private/AbilitySystemCheatManagerExtension.h
new file mode 100644
--- /dev/null
+++ b/Engine/Plugins/Runtime/GameplayAbilities/Source/GameplayAbilities/Private/AbilitySystemCheatManagerExtension.h
@@ -0,0 +1,49 @@
+// Copyright Epic Games, Inc. All Rights Reserved.
+
+#pragma once
+
+#include "GameFramework/CheatManager.h"
+
+#include "AbilitySystemCheatManagerExtension.generated.h"
+
+/** Cheats related to GAS */
+UCLASS(NotBlueprintable)
+class UAbilitySystemCheatManagerExtension final : public UCheatManagerExtension
+{
+	GENERATED_BODY()
+
+public:
+	UAbilitySystemCheatManagerExtension();
+
+	// Gameplay Abilities
+
+	/** List all of the Abilities Granted to the owning Player */
+	UFUNCTION(Exec)
+	void AbilityListGranted() const;
+
+	/** Grant a specified GameplayAbility to the owning Player */
+	UFUNCTION(Exec)
+	void AbilityGrant(const FString& AssetSearchString) const;
+
+	/** Activate a previously granted GameplayAbility on the owning Player */
+	UFUNCTION(Exec)
+	void AbilityActivate(const FString& PartialName) const;
+
+	/** Cancel a previously activated GameplayAbility on the owning Player */
+	UFUNCTION(Exec)
+	void AbilityCancel(const FString& PartialName) const;
+
+	// Gameplay Effects
+
+	/** List the Active Gameplay Effects on the owning Player */
+	UFUNCTION(Exec)
+	void EffectListActive() const;
+
+	/** Apply a Gameplay Effect on the owning Player */
+	UFUNCTION(Exec)
+	void EffectApply(const FString& PartialName, float EffectLevel = -1.0f) const;
+
+	/** Remove an Active Gameplay Effect on the owning Player */
+	UFUNCTION(Exec)
+	void EffectRemove(const FString& NameOrHandle) const;
+};
diff --git a/Engine/Plugins/Runtime/GameplayAbilities/Source/GameplayAbilities/Private/AbilitySystemComponent_Abilities.cpp b/Engine/Plugins/Runtime/GameplayAbilities/Source/GameplayAbilities/Private/AbilitySystemComponent_Abilities.cpp
--- a/Engine/Plugins/Runtime/GameplayAbilities/Source/GameplayAbilities/Private/AbilitySystemComponent_Abilities.cpp
+++ b/Engine/Plugins/Runtime/GameplayAbilities/Source/GameplayAbilities/Private/AbilitySystemComponent_Abilities.cpp
@@ -1626,7 +1626,7 @@ bool UAbilitySystemComponent::InternalTryActivateAbility(FGameplayAbilitySpecHan
 		{
 			// If we have a valid prediction key, the ability was started on the local client so it's okay
 
-			ABILITY_LOG(Warning, TEXT("Can't activate LocalOnly or LocalPredicted ability %s when not local! Net Execution Policy is %d."), *Ability->GetName(), (int32)Ability->GetNetExecutionPolicy());
+			ABILITY_LOG(Warning, TEXT("Can't activate %s ability %s when not local."), *UEnum::GetValueAsString<EGameplayAbilityNetExecutionPolicy::Type>(Ability->GetNetExecutionPolicy()), *Ability->GetName());
 
 			if (NetworkFailTag.IsValid())
 			{
@@ -1640,7 +1640,7 @@ bool UAbilitySystemComponent::InternalTryActivateAbility(FGameplayAbilitySpecHan
 
 	if (NetMode != ROLE_Authority && (Ability->GetNetExecutionPolicy() == EGameplayAbilityNetExecutionPolicy::ServerOnly || Ability->GetNetExecutionPolicy() == EGameplayAbilityNetExecutionPolicy::ServerInitiated))
 	{
-		ABILITY_LOG(Warning, TEXT("Can't activate ServerOnly or ServerInitiated ability %s when not the server! Net Execution Policy is %d."), *Ability->GetName(), (int32)Ability->GetNetExecutionPolicy());
+		ABILITY_LOG(Warning, TEXT("Can't activate %s ability %s when not the server."), *UEnum::GetValueAsString<EGameplayAbilityNetExecutionPolicy::Type>(Ability->GetNetExecutionPolicy()), *Ability->GetName());
 
 		if (NetworkFailTag.IsValid())
 		{
diff --git a/Engine/Plugins/Runtime/GameplayAbilities/Source/GameplayAbilities/Private/AbilitySystemGlobals.cpp b/Engine/Plugins/Runtime/GameplayAbilities/Source/GameplayAbilities/Private/AbilitySystemGlobals.cpp
--- a/Engine/Plugins/Runtime/GameplayAbilities/Source/GameplayAbilities/Private/AbilitySystemGlobals.cpp
+++ b/Engine/Plugins/Runtime/GameplayAbilities/Source/GameplayAbilities/Private/AbilitySystemGlobals.cpp
@@ -36,9 +36,6 @@ namespace UE::AbilitySystemGlobals
 
 	FAutoConsoleVariableRef CVarAbilitySystemIgnoreCooldowns(TEXT("AbilitySystem.IgnoreCooldowns"), bIgnoreAbilitySystemCooldowns, TEXT("Ignore cooldowns for all Gameplay Abilities."), ECVF_Cheat);
 	FAutoConsoleVariableRef CVarAbilitySystemIgnoreCosts(TEXT("AbilitySystem.IgnoreCosts"), bIgnoreAbilitySystemCosts, TEXT("Ignore costs for all Gameplay Abilities."), ECVF_Cheat);
-
-	template<typename ClassToFind>
-	TSubclassOf<ClassToFind> FuzzyFindClass(FString SearchString);
 }
 
 UAbilitySystemGlobals::UAbilitySystemGlobals(const FObjectInitializer& ObjectInitializer)
@@ -618,499 +615,3 @@ bool FNetSerializeScriptStructCache::NetSerialize(FArchive& Ar, UScriptStruct*&
 		return false;
 	}
 }
-
-namespace UE::AbilitySystemGlobals
-{
-	/**
-	 * Common logic used to fuzzy-find a requested class (or alternatively, a passed-in asset path).
-	 */
-	template<typename ClassToFind>
-	TSubclassOf<ClassToFind> FuzzyFindClass(FString SearchString)
-	{
-		TSubclassOf<ClassToFind> FoundClass;
-
-		// See if we passed in a class name of a Class that already exists in memory.
-		// If we passed-in Default__, just remove that part since we're looking for Classes, not CDO names.
-		SearchString.RemoveFromStart(TEXT("Default__"));
-		const int SearchStringLen = SearchString.Len();
-		int BestClassMatchLen = INT_MAX;
-		for (TObjectIterator<UClass> ClassIt; ClassIt; ++ClassIt)
-		{
-			const bool bClassMatches = ClassIt->IsChildOf(ClassToFind::StaticClass());
-			if (!bClassMatches)
-			{
-				continue;
-			}
-
-			// Class name search
-			const FString ClassName = ClassIt->GetName();
-			const int ClassNameLen = ClassName.Len();
-			if (ClassNameLen < BestClassMatchLen && ClassNameLen >= SearchStringLen)
-			{
-				bool bContains = ClassName.Contains(SearchString);
-				if (bContains)
-				{
-					FoundClass = *ClassIt;
-					BestClassMatchLen = ClassNameLen;
-				}
-			}
-		}
-
-		// If it wasn't a class name, then perhaps it was the path to a specific asset
-		if (!FoundClass)
-		{
-			FSoftObjectPath SoftObjectPath{ SearchString };
-			if (UObject* ReferencedObject = SoftObjectPath.ResolveObject())
-			{
-				if (UPackage* ReferencedPackage = Cast<UPackage>(ReferencedObject))
-				{
-					ReferencedObject = ReferencedPackage->FindAssetInPackage();
-				}
-
-				if (UBlueprint* ReferencedBlueprint = Cast<UBlueprint>(ReferencedObject))
-				{
-					FoundClass = ReferencedBlueprint->GeneratedClass;
-				}
-				else if (ClassToFind* ReferencedGA = Cast<ClassToFind>(ReferencedObject))
-				{
-					FoundClass = ReferencedGA->GetClass();
-				}
-			}
-		}
-
-		return FoundClass;
-	}
-}
-
-//
-// AbilitySystem.Ability Debug Commands
-//
-
-FAutoConsoleCommand DebugAbilitySystemAbilityListGrantedCommand(TEXT("AbilitySystem.Ability.Grant"), TEXT("Include a param [ClassName/AssetName] to Grant an Ability to the Player.  Omit parameters to lists all of the Gameplay Abilities currently granted to the Player."),
-	FConsoleCommandWithWorldArgsAndOutputDeviceDelegate::CreateLambda([](const TArray<FString>& Args, UWorld* World, FOutputDevice& OutputDevice)
-	{
-		UAbilitySystemGlobals& AbilitySystemGlobals = UAbilitySystemGlobals::Get();
-
-		APlayerController* PC = World->GetFirstPlayerController();
-		UAbilitySystemComponent* ASC = PC ? AbilitySystemGlobals.GetAbilitySystemComponentFromActor(PC->GetPawn()) : nullptr;
-		if (!ASC)
-		{
-			OutputDevice.Logf(TEXT("Could not find Player (%s) with AbilitySystemComponent in World (%s)"), *GetNameSafe(PC), *GetNameSafe(World));
-			return;
-		}
-
-		if (Args.Num() > 1)
-		{
-			OutputDevice.Logf(TEXT("Expected a single parameter (ClassName or AssetName).  Use no parameters to list all granted abilities."));
-			return;
-		}
-
-		// List already granted
-		if (Args.Num() < 1)
-		{
-			const UEnum* ExecutionEnumPtr = FindObject<UEnum>(nullptr, TEXT("/Script/GameplayAbilities.EGameplayAbilityNetExecutionPolicy"), true);
-			check(ExecutionEnumPtr&& TEXT("Couldn't locate EGameplayAbilityNetExecutionPolicy enum!"));
-
-			const UEnum* SecurityEnumPtr = FindObject<UEnum>(nullptr, TEXT("/Script/GameplayAbilities.EGameplayAbilityNetSecurityPolicy"), true);
-			check(SecurityEnumPtr&& TEXT("Couldn't locate EGameplayAbilityNetSecurityPolicy enum!"));
-
-			OutputDevice.Logf(TEXT("Granted abilities to %s (ASC: '%s'):"), *PC->GetName(), *ASC->GetFullName());
-
-			for (FGameplayAbilitySpec& Activatable : ASC->GetActivatableAbilities())
-			{
-				const TCHAR* ActiveText = Activatable.IsActive() ? TEXT("**ACTIVE**") : TEXT("");
-				OutputDevice.Logf(TEXT("   %s (%s - %s) %s"), *Activatable.Ability->GetName(), *ExecutionEnumPtr->GetDisplayNameTextByIndex(Activatable.Ability->GetNetExecutionPolicy()).ToString(), *SecurityEnumPtr->GetDisplayNameTextByIndex(Activatable.Ability->GetNetSecurityPolicy()).ToString(), ActiveText);
-			}
-
-			return;
-		}
-
-		FString SearchString = Args[0];
-
-		// We couldn't find anything the user was searching for, so early out
-		TSubclassOf<UGameplayAbility> ActivateAbilityClass = UE::AbilitySystemGlobals::FuzzyFindClass<UGameplayAbility>(SearchString);
-		if (!ActivateAbilityClass)
-		{
-			OutputDevice.Logf(TEXT("Could not find a valid Gameplay Ability based on Search String '%s'"), *SearchString);
-			return;
-		}
-
-		// Check if it's already granted
-		if (const FGameplayAbilitySpec* ExistingSpec = ASC->FindAbilitySpecFromClass(ActivateAbilityClass))
-		{
-			OutputDevice.Logf(TEXT("Existing Ability Spec '%s' on Player '%s' (It is already granted)."), *GetNameSafe(*ActivateAbilityClass), *GetNameSafe(PC));
-			return;
-		}
-
-		// If we're not the authority, we need to send the command to the server because we can't grant locally.
-		if (!ASC->IsOwnerActorAuthoritative())
-		{
-			const FString ServerCommand = FString::Printf(TEXT("AbilitySystem.Ability.Grant %s"), *SearchString);
-			PC->ServerExec(ServerCommand);
-
-			OutputDevice.Logf(TEXT("Sent Command '%s' from Player '%s' to Server (Reason: Cannot Grant if not Authority)."), *ServerCommand, *GetNameSafe(PC));
-			return;
-		}
-
-		// It wasn't granted, let's grant it now.
-		FGameplayAbilitySpec AbilitySpec{ ActivateAbilityClass };
-		FGameplayAbilitySpecHandle SpecHandle = ASC->GiveAbility(AbilitySpec);
-		if (SpecHandle.IsValid())
-		{
-			OutputDevice.Logf(TEXT("Successfully Granted '%s' on Player '%s'."), *GetNameSafe(ActivateAbilityClass.Get()), *GetNameSafe(PC));
-		}
-		else
-		{
-			OutputDevice.Logf(TEXT("Failed to Grant '%s' on Player '%s'."), *GetNameSafe(ActivateAbilityClass.Get()), *GetNameSafe(PC));
-		}
-		}), ECVF_Cheat);
-
-FAutoConsoleCommand DebugAbilitySystemAbilityCancelCommand(TEXT("AbilitySystem.Ability.Cancel"), TEXT("<Name>. Cancels (prematurely Ends) a currently executing Gameplay Ability"),
-	FConsoleCommandWithWorldArgsAndOutputDeviceDelegate::CreateLambda([](const TArray<FString>& Args, UWorld* World, FOutputDevice& OutputDevice)
-		{
-			UAbilitySystemGlobals& AbilitySystemGlobals = UAbilitySystemGlobals::Get();
-
-			APlayerController* PC = World->GetFirstPlayerController();
-			UAbilitySystemComponent* ASC = PC ? AbilitySystemGlobals.GetAbilitySystemComponentFromActor(PC->GetPawn()) : nullptr;
-			if (!ASC)
-			{
-				OutputDevice.Logf(TEXT("Could not find Player (%s) with AbilitySystemComponent in World (%s)"), *GetNameSafe(PC), *GetNameSafe(World));
-				return;
-			}
-
-			if (Args.Num() < 1)
-			{
-				if (IConsoleObject* GetGrantedCommand = IConsoleManager::Get().FindConsoleObject(TEXT("AbilitySystem.Ability.ListGranted"), false))
-				{
-					GetGrantedCommand->AsCommand()->Execute(Args, World, OutputDevice);
-				}
-				return;
-			}
-
-			FString SearchString = Args[0];
-			SearchString.RemoveFromStart(TEXT("Default__"));
-
-			bool bFound = false;
-			for (FGameplayAbilitySpec& GASpec : ASC->GetActivatableAbilities())
-			{
-				if (GASpec.IsActive())
-				{
-					if (!GASpec.Ability || !GASpec.Ability->GetName().Contains(SearchString))
-					{
-						continue;
-					}
-
-					TArray<UGameplayAbility*> ActiveAbilities = GASpec.GetAbilityInstances();
-					if (ActiveAbilities.Num() > 0)
-					{
-						for (UGameplayAbility* Instance : ActiveAbilities)
-						{
-							if (Instance)
-							{
-								OutputDevice.Logf(TEXT("%s (%s): Cancelling (instanced) %s"), *PC->GetName(), *ASC->GetName(), *Instance->GetName());
-								Instance->CancelAbility(GASpec.Handle, ASC->AbilityActorInfo.Get(), GASpec.ActivationInfo, true);
-								bFound = true;
-							}
-						}
-					}
-					else
-					{
-						OutputDevice.Logf(TEXT("%s (%s): Cancelling (non-instanced) %s"), *PC->GetName(), *ASC->GetName(), *GASpec.Ability->GetName());
-						GASpec.Ability->CancelAbility(GASpec.Handle, ASC->AbilityActorInfo.Get(), GASpec.ActivationInfo, true);
-						bFound = true;
-					}
-				}
-			}
-
-			if (!bFound)
-			{
-				OutputDevice.Logf(TEXT("Found no Active Gameplay Abilities on %s (%s) that matched '%s'"), *PC->GetName(), *ASC->GetName(), *SearchString);
-			}
-	}), ECVF_Cheat);
-
-FAutoConsoleCommand DebugAbilitySystemAbilityActivateCommand(TEXT("AbilitySystem.Ability.Activate"), TEXT("<TagName/ClassName/AssetName>. Activate a Gameplay Ability.\nSubstring name matching works for Activation Tags (on already granted abilities), Asset Paths (on non-granted abilities), or Class Names on both."),
-	FConsoleCommandWithWorldArgsAndOutputDeviceDelegate::CreateLambda([](const TArray<FString>& Args, UWorld* World, FOutputDevice& OutputDevice)
-	{
-		UAbilitySystemGlobals& AbilitySystemGlobals = UAbilitySystemGlobals::Get();
-
-		APlayerController* PC = World->GetFirstPlayerController();
-		UAbilitySystemComponent* ASC = PC ? AbilitySystemGlobals.GetAbilitySystemComponentFromActor(PC->GetPawn()) : nullptr;
-		if (!ASC)
-		{
-			OutputDevice.Logf(TEXT("Could not find Player (%s) with AbilitySystemComponent in World (%s)"), *GetNameSafe(PC), *GetNameSafe(World));
-			return;
-		}
-
-		if (Args.Num() < 1)
-		{
-			if (IConsoleObject* GetGrantedCommand = IConsoleManager::Get().FindConsoleObject(TEXT("AbilitySystem.Ability.Grant"), false))
-			{
-				GetGrantedCommand->AsCommand()->Execute(Args, World, OutputDevice);
-			}
-
-			return;
-		}
-
-		const FString SearchString = Args[0];
-
-		// Start by figuring out if we're trying to execute by GameplayTag.
-		const FName SearchName{ TCHAR_TO_ANSI(*SearchString), FNAME_Find };
-		if (SearchName != NAME_None)
-		{
-			const FGameplayTag GameplayTag = FGameplayTag::RequestGameplayTag(SearchName, false);
-			if (GameplayTag.IsValid())
-			{
-				FGameplayTagContainer GameplayTagContainer{ GameplayTag };
-				TArray<FGameplayAbilitySpec*> MatchingSpecs;
-				constexpr bool bTagRequirementsMustMatch = false;
-				ASC->GetActivatableGameplayAbilitySpecsByAllMatchingTags(GameplayTagContainer, MatchingSpecs, bTagRequirementsMustMatch);
-				if (MatchingSpecs.Num() > 0)
-				{
-					FString MatchingAbilitiesString = FString::JoinBy(MatchingSpecs, TEXT(", "), [](FGameplayAbilitySpec* Item) { return Item->GetDebugString(); });
-
-					bool bSuccess = ASC->TryActivateAbilitiesByTag(GameplayTagContainer);
-					if (bSuccess)
-					{
-						OutputDevice.Logf(TEXT("Requested Tag '%s' successfully executed one of: %s."), *SearchName.ToString(), *MatchingAbilitiesString);
-					}
-					else
-					{
-						OutputDevice.Logf(TEXT("Requested Tag '%s' was expected to execute one of: %s. But it failed to do so due to tag requirements."), *SearchName.ToString(), *MatchingAbilitiesString);
-					}
-				}
-				else
-				{
-					OutputDevice.Logf(TEXT("Requested Tag '%s' matched no given Gameplay Abilities to %s."), *SearchName.ToString(), *GetNameSafe(PC));
-				}
-
-				return;
-			}
-		}
-
-		// We couldn't find anything the user was searching for, so early out
-		TSubclassOf<UGameplayAbility> ActivateAbilityClass = UE::AbilitySystemGlobals::FuzzyFindClass<UGameplayAbility>(SearchString);
-		if (!ActivateAbilityClass)
-		{
-			OutputDevice.Logf(TEXT("Could not find a valid Gameplay Ability based on Search String '%s'"), *SearchString);
-			return;
-		}
-
-		// If we're not the authority, we should check if we need to send the command to the server
-		if (!ASC->IsOwnerActorAuthoritative())
-		{
-			const UGameplayAbility* Ability = ActivateAbilityClass.GetDefaultObject();
-			const bool bSendToServer = Ability && (Ability->GetNetExecutionPolicy() >= EGameplayAbilityNetExecutionPolicy::ServerInitiated);
-			if (bSendToServer)
-			{
-				const FString ServerCommand = FString::Printf(TEXT("AbilitySystem.Ability.Activate %s"), *SearchString);
-				PC->ServerExec(ServerCommand);
-
-				OutputDevice.Logf(TEXT("Sent Command '%s' from Player '%s' to Server (Reason: Net Execution Policy)."), *ServerCommand, *GetNameSafe(PC), *GetNameSafe(ActivateAbilityClass.Get()));
-				return;
-			}
-		}
-
-		// We found what the user was searching for, so let's try to activate it.  First, let's assume the Ability was already granted.
-		bool bSuccess = ASC->TryActivateAbilityByClass(ActivateAbilityClass);
-		if (bSuccess)
-		{
-			OutputDevice.Logf(TEXT("Successfully Activated Granted Ability '%s' on Player '%s'."), *GetNameSafe(ActivateAbilityClass.Get()), *GetNameSafe(PC));
-			return;
-		}
-
-		// It wasn't granted (or we failed to activate it even though it was granted).  We can't really differentiate those two, so let's grant it, then activate it.
-		FGameplayAbilitySpec AbilitySpec{ ActivateAbilityClass };
-		FGameplayAbilitySpecHandle SpecHandle = ASC->GiveAbilityAndActivateOnce(AbilitySpec);
-		if (SpecHandle.IsValid())
-		{
-			OutputDevice.Logf(TEXT("Successfully Granted, then Activated '%s' on Player '%s'."), *GetNameSafe(ActivateAbilityClass.Get()), *GetNameSafe(PC));
-		}
-		else
-		{
-			OutputDevice.Logf(TEXT("Failed to Grant & Activate '%s' on Player '%s'."), *GetNameSafe(ActivateAbilityClass.Get()), *GetNameSafe(PC));
-		}
-	}), ECVF_Cheat);
-
-
-//
-// AbilitySystem.Effect Debug Commands
-//
-
-FAutoConsoleCommand DebugAbilitySystemEffectListActiveCommand(TEXT("AbilitySystem.Effect.ListActive"), TEXT("Lists all of the Gameplay Effects currently active on the Player"),
-	FConsoleCommandWithWorldArgsAndOutputDeviceDelegate::CreateLambda([](const TArray<FString>& Args, UWorld* World, FOutputDevice& OutputDevice)
-	{
-		UAbilitySystemGlobals& AbilitySystemGlobals = UAbilitySystemGlobals::Get();
-
-		APlayerController* PC = World->GetFirstPlayerController();
-		UAbilitySystemComponent* ASC = PC ? AbilitySystemGlobals.GetAbilitySystemComponentFromActor(PC->GetPawn()) : nullptr;
-		if (!ASC)
-		{
-			OutputDevice.Logf(TEXT("Could not find Player (%s) with AbilitySystemComponent in World (%s)"), *GetNameSafe(PC), *GetNameSafe(World));
-			return;
-		}
-
-		OutputDevice.Logf(TEXT("Active Gameplay Effects on %s (%s):"), *PC->GetName(), *ASC->GetName());
-
-		const FActiveGameplayEffectsContainer& ActiveGEContainer = ASC->GetActiveGameplayEffects();
-		for (const FActiveGameplayEffectHandle& ActiveGEHandle : ActiveGEContainer.GetAllActiveEffectHandles())
-		{
-			if (const FActiveGameplayEffect* ActiveGE = ASC->GetActiveGameplayEffect(ActiveGEHandle))
-			{
-				const UGameplayEffect* GESpecDef = ActiveGE->Spec.Def;
-				const FString ActiveState = ActiveGE->bIsInhibited ? FString(TEXT("Inhibited")) : FString::Printf(TEXT("Active (Stack: %d)"), ActiveGE->Spec.GetStackCount());
-
-				OutputDevice.Logf(TEXT(" [%u] %s: %s"), GetTypeHash(ActiveGEHandle), *GetNameSafe(GESpecDef), *ActiveState);
-			}
-		}
-	}), ECVF_Cheat);
-
-FAutoConsoleCommand DebugAbilitySystemEffectRemoveCommand(TEXT("AbilitySystem.Effect.Remove"), TEXT("[-Server] <Handle/Name>. Remove a Gameplay Effect that is currently active on the Player"),
-	FConsoleCommandWithWorldArgsAndOutputDeviceDelegate::CreateLambda([](const TArray<FString>& Args, UWorld* World, FOutputDevice& OutputDevice)
-		{
-			TArray<FString> MutableArgs = Args;
-			const bool bExecuteOnServer = MutableArgs.RemoveSingle(TEXT("-server")) > 0;
-
-			UAbilitySystemGlobals& AbilitySystemGlobals = UAbilitySystemGlobals::Get();
-
-			APlayerController* PC = World->GetFirstPlayerController();
-			UAbilitySystemComponent* ASC = PC ? AbilitySystemGlobals.GetAbilitySystemComponentFromActor(PC->GetPawn()) : nullptr;
-			if (!ASC)
-			{
-				OutputDevice.Logf(TEXT("Could not find Player (%s) with AbilitySystemComponent in World (%s)"), *GetNameSafe(PC), *GetNameSafe(World));
-				return;
-			}
-
-			if (MutableArgs.Num() < 1)
-			{
-				OutputDevice.Logf(TEXT("Missing Arguments: [-Server] <Handle/Name>"));
-				OutputDevice.Logf(TEXT(" -Server: Indicates this command should be sent to the server, rather than predicted locally (NOTE: Handles are not the same on the server)."));
-				OutputDevice.Logf(TEXT(" Handle/Name: Supply a substring of a class name to match.  Handles can also be used if executed locally (useful in the case of multiple instances of the same GE)."));
-
-				if (IConsoleObject* ListActiveCommand = IConsoleManager::Get().FindConsoleObject(TEXT("AbilitySystem.Effect.ListActive"), false))
-				{
-					OutputDevice.Logf(TEXT("\nPerforming AbilitySystem.Effect.ListActive:"));
-					ListActiveCommand->AsCommand()->Execute(Args, World, OutputDevice);
-				}
-
-				return;
-			}
-
-			uint32 SearchHandleValue = 0;
-			FString SearchString = MutableArgs[0];
-			if (SearchString.IsNumeric())
-			{
-				int64 SearchHash64 = FCString::Atoi64(*SearchString);
-				if (SearchHash64 > 0)
-				{
-					SearchHandleValue = static_cast<uint32>(SearchHash64 & UINT_MAX);
-				}
-			}
-
-			if (!ASC->IsOwnerActorAuthoritative() && bExecuteOnServer)
-			{
-				if (SearchHandleValue > 0)
-				{
-					OutputDevice.Logf(TEXT("Error: Search by Handle Value is not permitted with -server (because ActiveGE Handles are not replicated to the same handle value)"));
-					return;
-				}
-				else
-				{
-					const FString AllArgs = FString::Join(MutableArgs, TEXT(" "));
-					const FString ServerCommand = FString::Printf(TEXT("AbilitySystem.Effect.Remove %s"), *AllArgs);
-					PC->ServerExec(ServerCommand);
-
-					OutputDevice.Logf(TEXT("Sent Command '%s' from Player '%s' to Server (as requested)."), *ServerCommand, *GetNameSafe(PC));
-					return;
-				}
-			}
-
-			bool bAnyMatch = false;
-			const FActiveGameplayEffectsContainer& ActiveGEContainer = ASC->GetActiveGameplayEffects();
-			for (const FActiveGameplayEffectHandle& ActiveGEHandle : ActiveGEContainer.GetAllActiveEffectHandles())
-			{
-				if (const FActiveGameplayEffect* ActiveGE = ASC->GetActiveGameplayEffect(ActiveGEHandle))
-				{
-					const UGameplayEffect* GESpecDef = ActiveGE->Spec.Def;
-					const uint32 ActiveGEHandleValue = GetTypeHash(ActiveGEHandle);
-
-					const bool bMatch = (SearchHandleValue > 0) ? (ActiveGEHandleValue == SearchHandleValue) : GESpecDef->GetName().Contains(SearchString);
-					if (bMatch)
-					{
-						OutputDevice.Logf(TEXT("Removing Active GE Handle %u (%s)"), ActiveGEHandleValue, *GESpecDef->GetName());
-						ASC->RemoveActiveGameplayEffect(ActiveGEHandle);
-						bAnyMatch = true;
-					}
-				}
-			}
-
-			if (!bAnyMatch)
-			{
-				OutputDevice.Logf(TEXT("Found no Active Gameplay Effects on %s (%s) that matched '%s'"), *PC->GetName(), *ASC->GetName(), *SearchString);
-			}
-		}), ECVF_Cheat);
-
-FAutoConsoleCommand DebugAbilitySystemEffectApply(TEXT("AbilitySystem.Effect.Apply"), TEXT("[-Server] <Class/AssetName> [Level]. Apply a Gameplay Effect on the Player.  Substring name matching works for Asset Tags, Asset Paths, or Class Names.  Use -Server to send to the server (default is apply locally)."),
-	FConsoleCommandWithWorldArgsAndOutputDeviceDelegate::CreateLambda([](const TArray<FString>& Args, UWorld* World, FOutputDevice& OutputDevice)
-	{
-		TArray<FString> MutableArgs = Args;
-		const bool bExecuteOnServer = MutableArgs.RemoveSingle(TEXT("-server")) > 0;
-
-		if (MutableArgs.Num() < 1)
-		{
-			OutputDevice.Logf(TEXT("Missing Arguments: [-Server] <MatchString> [Level]"));
-			OutputDevice.Logf(TEXT(" -Server: Indicates this command should be sent to the server, rather than predicted locally (search for assets is still verified locally)."));
-			OutputDevice.Logf(TEXT(" MatchString: Supply a substring of a class or asset name to match. Only loaded GameplayEffects are searched."));
-			OutputDevice.Logf(TEXT(" Level: Optionally supply a Level such as 4.  Can be omitted (in which case the GE is level-less)."));
-			return;
-		}
-
-		UAbilitySystemGlobals& AbilitySystemGlobals = UAbilitySystemGlobals::Get();
-
-		APlayerController* PC = World->GetFirstPlayerController();
-		APawn* Pawn = PC ? PC->GetPawn() : nullptr;
-		UAbilitySystemComponent* ASC = Pawn ? AbilitySystemGlobals.GetAbilitySystemComponentFromActor(Pawn) : nullptr;
-		if (!ASC)
-		{
-			OutputDevice.Logf(TEXT("Could not find Player (%s) with AbilitySystemComponent in World (%s)"), *GetNameSafe(PC), *GetNameSafe(World));
-			return;
-		}
-
-		// We couldn't find anything the user was searching for, so early out
-		const FString SearchString = MutableArgs[0];
-		TSubclassOf<UGameplayEffect> GameplayEffectClass = UE::AbilitySystemGlobals::FuzzyFindClass<UGameplayEffect>(SearchString);
-		if (!GameplayEffectClass)
-		{
-			OutputDevice.Logf(TEXT("Could not find a valid Gameplay Effect based on Search String '%s'"), *SearchString);
-			return;
-		}
-
-		// If we're not the authority, we need to send the command to the server because we can't grant locally.
-		if (!ASC->IsOwnerActorAuthoritative() && bExecuteOnServer)
-		{
-			const FString AllArgs = FString::Join(MutableArgs, TEXT(" "));
-			const FString ServerCommand = FString::Printf(TEXT("AbilitySystem.Effect.Apply %s"), *AllArgs);
-			PC->ServerExec(ServerCommand);
-
-			OutputDevice.Logf(TEXT("Sent Command '%s' from Player '%s' to Server (as requested)."), *ServerCommand, *GetNameSafe(PC));
-			return;
-		}
-
-		// Create a GameplayEffectSpec that executes the passed-in parameters
-		FGameplayEffectContextHandle GEContextHandle = ASC->MakeEffectContext();
-		GEContextHandle.AddInstigator(Pawn, Pawn);
-		GEContextHandle.AddOrigin(Pawn->GetActorLocation());
-		const float Level = FMath::Max(FGameplayEffectConstants::INVALID_LEVEL, (Args.Num() > 1) ? FPlatformString::Atof(*MutableArgs[1]) : FGameplayEffectConstants::INVALID_LEVEL);
-		FGameplayEffectSpec GESpec{ GameplayEffectClass.GetDefaultObject(), GEContextHandle, Level };
-
-		// We need to create a new valid prediction key (as if we just started activating an ability) so the GE will fire even if we are not the authority.
-		FPredictionKey FakePredictionKey = (PC->GetLocalRole() == ENetRole::ROLE_Authority) ? FPredictionKey() : FPredictionKey::CreateNewPredictionKey(ASC);
-		FActiveGameplayEffectHandle ActiveGEHandle = ASC->ApplyGameplayEffectSpecToSelf(GESpec, FakePredictionKey);
-		if (ActiveGEHandle.IsValid())
-		{
-			OutputDevice.Logf(TEXT("Successfully Applied Gameplay Effect '%s' on Player '%s'."), *GameplayEffectClass->GetName(), *PC->GetName());
-		}
-		else
-		{
-			OutputDevice.Logf(TEXT("Failed to Apply Gameplay Effect '%s' on Player '%s'."), *GameplayEffectClass->GetName(), *PC->GetName());
-		}
-	}), ECVF_Cheat);
diff --git a/Engine/Plugins/Runtime/GameplayAbilities/Source/GameplayAbilities/Private/GameplayCueManager.cpp b/Engine/Plugins/Runtime/GameplayAbilities/Source/GameplayAbilities/Private/GameplayCueManager.cpp
--- a/Engine/Plugins/Runtime/GameplayAbilities/Source/GameplayAbilities/Private/GameplayCueManager.cpp
+++ b/Engine/Plugins/Runtime/GameplayAbilities/Source/GameplayAbilities/Private/GameplayCueManager.cpp
@@ -195,8 +195,8 @@ void UGameplayCueManager::RouteGameplayCue(AActor* TargetActor, FGameplayTag Gam
 	if (DisplayGameplayCues && !(Options & EGameplayCueExecutionOptions::IgnoreDebug))
 	{
 		FString DebugStr = FString::Printf(TEXT("[%s] %s - %s"), *GetNameSafe(TargetActor), *GameplayCueTag.ToString(), *EGameplayCueEventToString(EventType) );
-		FColor DebugColor = FColor::Green;
-		DrawDebugString(TargetActor->GetWorld(), FVector(0.f, 0.f, 100.f), DebugStr, TargetActor, DebugColor, DisplayGameplayCueDuration);
+		FColor DebugColor = EventType == EGameplayCueEvent::Removed ? FColor::Red : FColor::Green;
+		DrawDebugString(TargetActor->GetWorld(), FVector(0.f, 0.f, -10.0f * static_cast<int>(EventType)), DebugStr, TargetActor, DebugColor, DisplayGameplayCueDuration);
 		ABILITY_LOG(Display, TEXT("%s"), *DebugStr);
 	}
 #endif // ENABLE_DRAW_DEBUG
diff --git a/Engine/Plugins/Runtime/GameplayAbilities/Source/GameplayAbilities/Private/GameplayDebuggerCategory_Abilities.cpp b/Engine/Plugins/Runtime/GameplayAbilities/Source/GameplayAbilities/Private/GameplayDebuggerCategory_Abilities.cpp
--- a/Engine/Plugins/Runtime/GameplayAbilities/Source/GameplayAbilities/Private/GameplayDebuggerCategory_Abilities.cpp
+++ b/Engine/Plugins/Runtime/GameplayAbilities/Source/GameplayAbilities/Private/GameplayDebuggerCategory_Abilities.cpp
@@ -12,11 +12,22 @@
 #include "AbilitySystemGlobals.h"
 #include "AbilitySystemComponent.h"
 #include "Engine/Canvas.h"
+#include "Engine/NetConnection.h"
+#include "GameFramework/PlayerController.h"
 
 
 FGameplayDebuggerCategory_Abilities::FGameplayDebuggerCategory_Abilities()
 {
 	SetDataPackReplication<FRepData>(&DataPack);
+
+	// Hard coding these to avoid needing to import InputCore just for EKeys::GetFName().
+	const FName KeyNameOne{ "One" };
+	const FName KeyNameTwo{ "Two" };
+	const FName KeyNameThree{ "Three" };
+
+	BindKeyPress(KeyNameOne, FGameplayDebuggerInputModifier::Shift, this, &FGameplayDebuggerCategory_Abilities::OnShowGameplayTagsToggle, EGameplayDebuggerInputMode::Local);
+	BindKeyPress(KeyNameTwo, FGameplayDebuggerInputModifier::Shift, this, &FGameplayDebuggerCategory_Abilities::OnShowGameplayAbilitiesToggle, EGameplayDebuggerInputMode::Local);
+	BindKeyPress(KeyNameThree, FGameplayDebuggerInputModifier::Shift, this, &FGameplayDebuggerCategory_Abilities::OnShowGameplayEffectsToggle, EGameplayDebuggerInputMode::Local);
 }
 
 TSharedRef<FGameplayDebuggerCategory> FGameplayDebuggerCategory_Abilities::MakeInstance()
@@ -24,9 +35,25 @@ TSharedRef<FGameplayDebuggerCategory> FGameplayDebuggerCategory_Abilities::MakeI
 	return MakeShareable(new FGameplayDebuggerCategory_Abilities());
 }
 
+void FGameplayDebuggerCategory_Abilities::OnShowGameplayTagsToggle()
+{
+	bShowGameplayTags = !bShowGameplayTags;
+}
+
+void FGameplayDebuggerCategory_Abilities::OnShowGameplayAbilitiesToggle()
+{
+	bShowGameplayAbilities = !bShowGameplayAbilities;
+}
+
+void FGameplayDebuggerCategory_Abilities::OnShowGameplayEffectsToggle()
+{
+	bShowGameplayEffects = !bShowGameplayEffects;
+}
+
 void FGameplayDebuggerCategory_Abilities::FRepData::Serialize(FArchive& Ar)
 {
-	Ar << OwnedTags;
+	bool bSuccess;
+	OwnedTags.NetSerialize(Ar, ClientPackageMap.Get(), bSuccess);
 
 	int32 NumAbilities = Abilities.Num();
 	Ar << NumAbilities;
@@ -66,10 +93,11 @@ void FGameplayDebuggerCategory_Abilities::CollectData(APlayerController* OwnerPC
 	UAbilitySystemComponent* AbilityComp = UAbilitySystemGlobals::GetAbilitySystemComponentFromActor(DebugActor);
 	if (AbilityComp)
 	{
-		static FGameplayTagContainer OwnerTags;
-		OwnerTags.Reset();
-		AbilityComp->GetOwnedGameplayTags(OwnerTags);
-		DataPack.OwnedTags = OwnerTags.ToStringSimple();
+		// Save off the package map for serialization over the network
+		UNetConnection* NetConnection = OwnerPC->GetNetConnection();
+		DataPack.ClientPackageMap = NetConnection ? NetConnection->PackageMap : nullptr;
+
+		AbilityComp->GetOwnedGameplayTags(DataPack.OwnedTags);
 
 		TArray<FGameplayEffectSpec> ActiveEffectSpecs;
 		AbilityComp->GetAllActiveGameplayEffectSpecs(ActiveEffectSpecs);
@@ -112,7 +140,7 @@ void FGameplayDebuggerCategory_Abilities::CollectData(APlayerController* OwnerPC
 	}
 }
 
-bool FGameplayDebuggerCategory_Abilities::WrapStringAccordingToViewport(const FString& StrIn, FString& StrOut, FGameplayDebuggerCanvasContext& CanvasContext, float ViewportWitdh)
+bool FGameplayDebuggerCategory_Abilities::WrapStringAccordingToViewport(const FString& StrIn, FString& StrOut, FGameplayDebuggerCanvasContext& CanvasContext, float ViewportWitdh) const
 {
 	if (!StrIn.IsEmpty())
 	{
@@ -137,6 +165,10 @@ bool FGameplayDebuggerCategory_Abilities::WrapStringAccordingToViewport(const FS
 			}
 			return true;
 		}
+		else
+		{
+			StrOut = StrIn;
+		}
 	}
 	// No need to wrap the text 
 	return false;
@@ -144,103 +176,196 @@ bool FGameplayDebuggerCategory_Abilities::WrapStringAccordingToViewport(const FS
 
 void FGameplayDebuggerCategory_Abilities::DrawData(APlayerController* OwnerPC, FGameplayDebuggerCanvasContext& CanvasContext)
 {
-	FVector2D ViewPortSize;
-	GEngine->GameViewport->GetViewportSize( /*out*/ViewPortSize);
+	// Draw the sub-category bindings inline with the category header
+	{
+		CanvasContext.CursorX += 250.0f;
+		CanvasContext.CursorY -= CanvasContext.GetLineHeight();
+		const TCHAR* Active = TEXT("{green}");
+		const TCHAR* Inactive = TEXT("{yellow}");
+		CanvasContext.Printf(TEXT("Tags [%s%s{white}]\tAbilities [%s%s{white}]\tEffects [%s%s{white}]"), bShowGameplayTags ? Active : Inactive, *GetInputHandlerDescription(0), bShowGameplayAbilities ? Active : Inactive, *GetInputHandlerDescription(1), bShowGameplayEffects ? Active : Inactive, *GetInputHandlerDescription(2));
+	}
+
+	static float LastDrawDataEndSize = CanvasContext.Canvas->SizeY - CanvasContext.CursorY - CanvasContext.CursorX;
+	float ThisDrawDataStartPos = CanvasContext.CursorY;
 
-	const float BackgroundPadding = 5.0f;
-	const FVector2D BackgroundSize(ViewPortSize.X - 2 * BackgroundPadding, ViewPortSize.Y);
 	const FLinearColor BackgroundColor(0.1f, 0.1f, 0.1f, 0.8f);
+	const FVector2D BackgroundPos{ CanvasContext.CursorX, CanvasContext.CursorY };
+	const FVector2D BackgroundSize(CanvasContext.Canvas->SizeX - (2.0f * CanvasContext.CursorX), LastDrawDataEndSize);
 
-	// Draw a transparent background so that the text is easier to look at
-	FCanvasTileItem Background(FVector2D(0.0f, 0.0f), BackgroundSize, BackgroundColor);
+	// Draw a transparent dark background so that the text is easier to look at
+	FCanvasTileItem Background(FVector2D(0.0f), BackgroundSize, BackgroundColor);
 	Background.BlendMode = SE_BLEND_Translucent;
-	CanvasContext.DrawItem(Background, CanvasContext.DefaultX - BackgroundPadding, CanvasContext.DefaultY - BackgroundPadding);
 
-	FString WrappedOwnedTagsStr;
-	// If need to wrap string, use the wrapped string, else use the DataPack one, avoid string copying.
-	const FString& OwnedTagsRef = WrapStringAccordingToViewport(DataPack.OwnedTags, WrappedOwnedTagsStr, CanvasContext, BackgroundSize.X) ? WrappedOwnedTagsStr : DataPack.OwnedTags;
+	CanvasContext.DrawItem(Background, BackgroundPos.X, BackgroundPos.Y);
 
-	CanvasContext.Printf(TEXT("Owned Tags: \n{yellow}%s"), *OwnedTagsRef);
+	if (bShowGameplayTags)
+	{
+		DrawGameplayTags(CanvasContext, OwnerPC);
+	}
+
+	if (bShowGameplayAbilities)
+	{
+		DrawGameplayAbilities(CanvasContext, OwnerPC);
+	}
+
+	if (bShowGameplayEffects)
+	{
+		DrawGameplayEffects(CanvasContext, OwnerPC);
+	}
+
+	LastDrawDataEndSize = CanvasContext.CursorY - ThisDrawDataStartPos;
+}
+
+void FGameplayDebuggerCategory_Abilities::DrawGameplayTags(FGameplayDebuggerCanvasContext& CanvasContext, const APlayerController* OwnerPC) const
+{
+	const float CanvasWidth = CanvasContext.Canvas->SizeX;
 
 	AActor* LocalDebugActor = FindLocalDebugActor();
-	UAbilitySystemComponent* AbilityComp = UAbilitySystemGlobals::GetAbilitySystemComponentFromActor(LocalDebugActor);
-	if (AbilityComp)
+	if (UAbilitySystemComponent* AbilityComp = UAbilitySystemGlobals::GetAbilitySystemComponentFromActor(LocalDebugActor))
 	{
-		static FGameplayTagContainer OwnerTags;
-		OwnerTags.Reset();
-		AbilityComp->GetOwnedGameplayTags(OwnerTags);
+		FGameplayTagContainer ServerOnlyTags = DataPack.OwnedTags;
 
-		TArray<FString> OwnerTagsStrArray = OwnerTags.ToStringsMaxLen(1024);
-		FString OwnerTagsStr = OwnerTagsStrArray.Num() > 0 ? OwnerTagsStrArray[0] : TEXT("");
-		FString WrappedOwnerTagsStr;
+		// If we're not the authority, we should represent the tags in such a way that the user can see the difference between agreed upon server/client tags
+		// and ones where they disagree.
+		if (!AbilityComp->IsOwnerActorAuthoritative())
+		{
+			static FGameplayTagContainer LocalOnlyTags;
+			AbilityComp->GetOwnedGameplayTags(LocalOnlyTags);
 
-		const FString& OwnerTagsStrRef = WrapStringAccordingToViewport(OwnerTagsStr, WrappedOwnerTagsStr, CanvasContext, BackgroundSize.X) ? WrappedOwnerTagsStr : OwnerTagsStr;
-		CanvasContext.Printf(TEXT("Local Tags: \n{cyan}%s"), *OwnerTagsStrRef);
+			FGameplayTagContainer MatchingTags = ServerOnlyTags.FilterExact(LocalOnlyTags);
+			ServerOnlyTags.RemoveTags(MatchingTags);
+			LocalOnlyTags.RemoveTags(MatchingTags);
+
+			// Wrap the strings to the viewport
+			FString MatchingTagsStr, ServerOnlyTagsStr, LocalOnlyTagsStr;
+			WrapStringAccordingToViewport(MatchingTags.ToStringSimple(), MatchingTagsStr, CanvasContext, CanvasWidth);
+			WrapStringAccordingToViewport(ServerOnlyTags.ToStringSimple(), ServerOnlyTagsStr, CanvasContext, CanvasWidth);
+			WrapStringAccordingToViewport(LocalOnlyTags.ToStringSimple(), LocalOnlyTagsStr, CanvasContext, CanvasWidth);
+
+			MatchingTagsStr = MatchingTagsStr.Len() > 0 ? FString::Printf(TEXT("{cyan}%s  "), *MatchingTagsStr) : FString{};
+			ServerOnlyTagsStr = ServerOnlyTagsStr.Len() > 0 ? FString::Printf(TEXT("{yellow}%s  "), *ServerOnlyTagsStr) : FString{};
+			LocalOnlyTagsStr = LocalOnlyTagsStr.Len() > 0 ? FString::Printf(TEXT("{green}%s  "), *LocalOnlyTagsStr) : FString{};
+
+			CanvasContext.Printf(TEXT("Owned Tags Legend:  {cyan}Both  {yellow}Server  {green}Local \n%s%s%s"), *MatchingTagsStr, *ServerOnlyTagsStr, *LocalOnlyTagsStr);
+		}
+		else
+		{
+			FString ServerOnlyTagsStr;
+			WrapStringAccordingToViewport(ServerOnlyTags.ToStringSimple(), ServerOnlyTagsStr, CanvasContext, CanvasWidth);
+			CanvasContext.Printf(TEXT("Owned Tags: \n{cyan}%s"), *ServerOnlyTagsStr);
+		}
 	}
 
+	// End with a newline to separate from the other categories
+	CanvasContext.Print(TEXT(""));
+}
+
+void FGameplayDebuggerCategory_Abilities::DrawGameplayEffects(FGameplayDebuggerCanvasContext& CanvasContext, const APlayerController* OwnerPC) const
+{
 	CanvasContext.Printf(TEXT("Gameplay Effects: {yellow}%d"), DataPack.GameplayEffects.Num());
 	for (int32 Idx = 0; Idx < DataPack.GameplayEffects.Num(); Idx++)
 	{
 		const FRepData::FGameplayEffectDebug& ItemData = DataPack.GameplayEffects[Idx];
 
-		FString Desc = FString::Printf(TEXT("\t{yellow}%s {grey}source:{white}%s {grey}duration:{white}"), *ItemData.Effect, *ItemData.Context);
-		Desc += (ItemData.Duration > 0.0f) ? FString::Printf(TEXT("%.2f"), ItemData.Duration) : FString(TEXT("INF"));
+		FStringBuilderBase Desc;
+		Desc.Appendf(TEXT("\t{yellow}%s {grey}source:{white}%s {grey}duration:{white}"), *ItemData.Effect, *ItemData.Context);
+		if (ItemData.Duration > 0.0f)
+		{
+			Desc.Appendf(TEXT("%.3f"), ItemData.Duration);
+		}
+		else
+		{
+			Desc.Appendf(TEXT("INF"));
+		}
 
 		if (ItemData.Period > 0.0f)
 		{
-			Desc += FString::Printf(TEXT(" {grey}period:{white}%.2f"), ItemData.Period);
+			Desc.Appendf(TEXT(" {grey}period:{white}%.3f"), ItemData.Period);
 		}
 
 		if (ItemData.Stacks > 1)
 		{
-			Desc += FString::Printf(TEXT(" {grey}stacks:{white}%d"), ItemData.Stacks);
+			Desc.Appendf(TEXT(" {grey}stacks:{white}%d"), ItemData.Stacks);
 		}
 
 		if (ItemData.Level > 1.0f)
 		{
-			Desc += FString::Printf(TEXT(" {grey}level:{white}%.2f"), ItemData.Level);
+			Desc.Appendf(TEXT(" {grey}level:{white}%.2f"), ItemData.Level);
 		}
 
-		CanvasContext.Print(Desc);
+		CanvasContext.Print(Desc.ToString());
 	}
+	
+	// End with a newline to separate from the other categories
+	CanvasContext.MoveToNewLine();
+}
 
-	CanvasContext.Printf(TEXT("Gameplay Abilities: {yellow}%d"), DataPack.Abilities.Num());
-	int32 HalfNum = FMath::CeilToInt(DataPack.Abilities.Num() / 2.f);
-	for (int32 Idx = 0; Idx < HalfNum; Idx++)
+void FGameplayDebuggerCategory_Abilities::DrawGameplayAbilities(FGameplayDebuggerCanvasContext& CanvasContext, const APlayerController* OwnerPC) const
+{
+	const float CanvasWidth = CanvasContext.Canvas->SizeX;
+
+	// Let's do something semi-smart and stable to resize the columns to readable yet wide enough
+	int32 NumActive = 0;
+	FString LongestLengthObjectName;
+	for (const FRepData::FGameplayAbilityDebug& ItemData : DataPack.Abilities)
 	{
-		if (2 * Idx + 1 < DataPack.Abilities.Num())
+		if (ItemData.Ability.Len() > LongestLengthObjectName.Len())
 		{
-			const FRepData::FGameplayAbilityDebug* ItemData[2] = { &DataPack.Abilities[2 * Idx], &DataPack.Abilities[2 * Idx + 1] };
-			FString Abilities[2];
-			float WidthSum = 0.0f;
-			for (size_t j = 0; j < 2; ++j)
-			{
-				Abilities[j] = FString::Printf(TEXT("\t{yellow}%s {grey}source:{white}%s {grey}level:{white}%d {grey}active:{white}%s"),
-					*ItemData[j]->Ability, *ItemData[j]->Source, ItemData[j]->Level, ItemData[j]->bIsActive ? TEXT("YES") : TEXT("NO"));
-				float TempWidth = 0.0f; 
-				float TempHeight = 0.0f;
-				CanvasContext.MeasureString(Abilities[j], TempWidth, TempHeight);
-				WidthSum += TempWidth;
-			}
+			LongestLengthObjectName = ItemData.Ability;
+		}
 
-			if (WidthSum < BackgroundSize.X)
-			{
-				CanvasContext.Print(Abilities[0].Append(Abilities[1]));
-			}
-			else
-			{
-				CanvasContext.Print(Abilities[0]);
-				CanvasContext.Print(Abilities[1]);
-			}
+		if (ItemData.Source.Len() > LongestLengthObjectName.Len())
+		{
+			LongestLengthObjectName = ItemData.Source;
 		}
-		else
+
+		NumActive += ItemData.bIsActive;
+	}
+
+	// Measure the individual string sizes, so that we can size the columns properly
+	constexpr float Padding = 10.0f;
+	static float ObjNameSize = 150.0f, SourceNameSize = 100.0f, LevelNameSize = 100.0f, TempSizeY = 0.0f;
+	static int32 CachedLen = 0;
+	if (LongestLengthObjectName.Len() > CachedLen)
+	{
+		CanvasContext.MeasureString(LongestLengthObjectName, ObjNameSize, TempSizeY);
+		CanvasContext.MeasureString(TEXT("source: "), SourceNameSize, TempSizeY);
+		CanvasContext.MeasureString(TEXT("level: 00"), LevelNameSize, TempSizeY);
+		ObjNameSize += Padding;
+	}
+	const float ColumnWidth = ObjNameSize * 2 + SourceNameSize + LevelNameSize;
+	const int NumColumns = FMath::Max(1, FMath::FloorToInt(CanvasWidth / ColumnWidth));
+
+	CanvasContext.Printf(TEXT("Gameplay Abilities: \t{yellow}Granted[%d] \t{cyan}Active[%d]"), DataPack.Abilities.Num(), NumActive);
+	for (const FRepData::FGameplayAbilityDebug& ItemData : DataPack.Abilities)
+	{
+		float CursorX = CanvasContext.CursorX;
+		float CursorY = CanvasContext.CursorY;
+
+		// Print positions manually to align them properly
+		CanvasContext.PrintAt(CursorX + ObjNameSize * 0, CursorY, ItemData.bIsActive ? FColor::Cyan : FColor::Yellow, ItemData.Ability);
+		CanvasContext.PrintAt(CursorX + ObjNameSize * 1, CursorY, FString::Printf(TEXT("{grey}source:{white}%s"), *ItemData.Source));
+		CanvasContext.PrintAt(CursorX + ObjNameSize * 2 + SourceNameSize, CursorY, FString::Printf(TEXT("{grey}level:{white}%02d"), ItemData.Level));
+
+		// PrintAt would have reset these values, restore them.
+		CanvasContext.CursorX = CursorX + (CanvasWidth / NumColumns);
+		CanvasContext.CursorY = CursorY;
+
+		// If we're going to overflow, go to the next line...
+		if (CanvasContext.CursorX + ColumnWidth >= CanvasWidth)
 		{
-			const FRepData::FGameplayAbilityDebug& ItemData = DataPack.Abilities[2 * Idx];
-			// Only display one
-			CanvasContext.Printf(TEXT("\t{yellow}%s {grey}source:{white}%s {grey}level:{white}%d {grey}active:{white}%s"),
-				*ItemData.Ability, *ItemData.Source, ItemData.Level, ItemData.bIsActive ? TEXT("YES") : TEXT("NO"));
+			CanvasContext.MoveToNewLine();
 		}
 	}
+
+	// End the row with a newline
+	if (CanvasContext.CursorX != CanvasContext.DefaultX)
+	{
+		CanvasContext.MoveToNewLine();
+	}
+
+	// End the category with a newline to separate from the other categories
+	CanvasContext.MoveToNewLine();
 }
 
 #endif // WITH_GAMEPLAY_DEBUGGER_MENU
diff --git a/Engine/Plugins/Runtime/GameplayAbilities/Source/GameplayAbilities/Private/GameplayDebuggerCategory_Abilities.h b/Engine/Plugins/Runtime/GameplayAbilities/Source/GameplayAbilities/Private/GameplayDebuggerCategory_Abilities.h
--- a/Engine/Plugins/Runtime/GameplayAbilities/Source/GameplayAbilities/Private/GameplayDebuggerCategory_Abilities.h
+++ b/Engine/Plugins/Runtime/GameplayAbilities/Source/GameplayAbilities/Private/GameplayDebuggerCategory_Abilities.h
@@ -5,10 +5,12 @@
 #include "CoreMinimal.h"
 
 #if WITH_GAMEPLAY_DEBUGGER_MENU
+#include "GameplayTagContainer.h"
 #include "GameplayDebuggerCategory.h"
 
 class AActor;
 class APlayerController;
+class UPackageMap;
 
 class FGameplayDebuggerCategory_Abilities : public FGameplayDebuggerCategory
 {
@@ -20,10 +22,23 @@ public:
 
 	static TSharedRef<FGameplayDebuggerCategory> MakeInstance();
 
+	void OnShowGameplayTagsToggle();
+	void OnShowGameplayAbilitiesToggle();
+	void OnShowGameplayEffectsToggle();
+		
 protected:
+
+	void DrawGameplayTags(FGameplayDebuggerCanvasContext& CanvasContext, const APlayerController* OwnerPC) const;
+	void DrawGameplayAbilities(FGameplayDebuggerCanvasContext& CanvasContext, const APlayerController* OwnerPC) const;
+	void DrawGameplayEffects(FGameplayDebuggerCanvasContext& CanvasContext, const APlayerController* OwnerPC) const;
+
+
 	struct FRepData
 	{
-		FString OwnedTags;
+		// to aid in NetSerialize
+		TWeakObjectPtr<UPackageMap>	ClientPackageMap;
+
+		FGameplayTagContainer OwnedTags;
 
 		struct FGameplayAbilityDebug
 		{
@@ -49,7 +64,12 @@ protected:
 	};
 	FRepData DataPack;
 
-	bool WrapStringAccordingToViewport(const FString& iStr, FString& oStr, FGameplayDebuggerCanvasContext& CanvasContext, float ViewportWitdh);
+	bool WrapStringAccordingToViewport(const FString& iStr, FString& oStr, FGameplayDebuggerCanvasContext& CanvasContext, float ViewportWitdh) const;
+
+private:
+	bool bShowGameplayTags = true;
+	bool bShowGameplayAbilities = true;
+	bool bShowGameplayEffects = true;
 };
 
 #endif // WITH_GAMEPLAY_DEBUGGER_MENU
