Fix warning V1078: An empty container is iterated. The loop will not be executed.

diff --git a/Engine/Plugins/Runtime/GameplayAbilities/Source/GameplayAbilities/Private/AbilitySystemDebugHUD.cpp b/Engine/Plugins/Runtime/GameplayAbilities/Source/GameplayAbilities/Private/AbilitySystemDebugHUD.cpp
--- a/Engine/Plugins/Runtime/GameplayAbilities/Source/GameplayAbilities/Private/AbilitySystemDebugHUD.cpp
+++ b/Engine/Plugins/Runtime/GameplayAbilities/Source/GameplayAbilities/Private/AbilitySystemDebugHUD.cpp
@@ -470,7 +470,7 @@ void AAbilitySystemDebugHUD::DisplayDebugStrings(UCanvas* InCanvas, const AActor
 		InCanvas->TextSize(Font, DrawString, TextWidth, TextHeight);
 		MaxWidth = FMath::Max(MaxWidth, TextWidth);
 		TotalHeight += TextHeight;
-		new(StringsToDraw)FStringAndLineHeight(DrawString, TextHeight);
+		StringsToDraw.Emplace(FStringAndLineHeight(DrawString, TextHeight));
 	}
 
 	FVector ThisOrigin;
diff --git a/Engine/Source/Editor/FoliageEdit/Private/FoliageEdMode.cpp b/Engine/Source/Editor/FoliageEdit/Private/FoliageEdMode.cpp
--- a/Engine/Source/Editor/FoliageEdit/Private/FoliageEdMode.cpp
+++ b/Engine/Source/Editor/FoliageEdit/Private/FoliageEdMode.cpp
@@ -2963,7 +2963,7 @@ void FEdModeFoliage::ApplyPaintBucket_Add(AActor* Actor)
 						continue;
 					}
 
-					new(InstancesToPlace)FPotentialInstance(InstLocation, Triangle.WorldNormal, Component, 1.f);
+					InstancesToPlace.Emplace(FPotentialInstance(InstLocation, Triangle.WorldNormal, Component, 1.f));
 				}
 			}
 		}
