Fixed crash in UAbilitySystemComponent::MakeEffectContext that can happen if AbilityActorInfo is not valid.

diff --git a/Engine/Plugins/Runtime/GameplayAbilities/Source/GameplayAbilities/Private/AbilitySystemComponent.cpp b/Engine/Plugins/Runtime/GameplayAbilities/Source/GameplayAbilities/Private/AbilitySystemComponent.cpp
--- a/Engine/Plugins/Runtime/GameplayAbilities/Source/GameplayAbilities/Private/AbilitySystemComponent.cpp
+++ b/Engine/Plugins/Runtime/GameplayAbilities/Source/GameplayAbilities/Private/AbilitySystemComponent.cpp
@@ -469,10 +469,13 @@ FGameplayEffectSpecHandle UAbilitySystemComponent::MakeOutgoingSpec(TSubclassOf<
 FGameplayEffectContextHandle UAbilitySystemComponent::MakeEffectContext() const
 {
 	FGameplayEffectContextHandle Context = FGameplayEffectContextHandle(UAbilitySystemGlobals::Get().AllocGameplayEffectContext());
+	
 	// By default use the owner and avatar as the instigator and causer
-	check(AbilityActorInfo.IsValid());
+	if (ensureMsgf(AbilityActorInfo.IsValid(), TEXT("Unable to make effect context because AbilityActorInfo is not valid.")))
+	{
+		Context.AddInstigator(AbilityActorInfo->OwnerActor.Get(), AbilityActorInfo->AvatarActor.Get());
+	}
 	
-	Context.AddInstigator(AbilityActorInfo->OwnerActor.Get(), AbilityActorInfo->AvatarActor.Get());
 	return Context;
 }
 
