Added nullptr check in FActiveGameplayEffectsContainer::SetAttributeBaseValue to prevent crashes when the owner is invalid.

#tests Tested putting two players in split screen and eliminating them both in PIE. Also tried stepping into the function and making sure it doesn't ensure when owner is valid

#localization none

diff --git a/Engine/Plugins/Runtime/GameplayAbilities/Source/GameplayAbilities/Private/GameplayEffect.cpp b/Engine/Plugins/Runtime/GameplayAbilities/Source/GameplayAbilities/Private/GameplayEffect.cpp
--- a/Engine/Plugins/Runtime/GameplayAbilities/Source/GameplayAbilities/Private/GameplayEffect.cpp
+++ b/Engine/Plugins/Runtime/GameplayAbilities/Source/GameplayAbilities/Private/GameplayEffect.cpp
@@ -3460,6 +3460,11 @@ void FActiveGameplayEffectsContainer::InternalUpdateNumericalAttribute(FGameplay
 
 void FActiveGameplayEffectsContainer::SetAttributeBaseValue(FGameplayAttribute Attribute, float NewBaseValue)
 {
+	if (!ensureMsgf(Owner, TEXT("%hs: This ActiveGameplayEffectsContainer has an invalid owner. Unable to set attribute %s"), __FUNCTION__, *Attribute.AttributeName))
+	{
+		return;
+	}
+
 	const UAttributeSet* Set = Owner->GetAttributeSubobject(Attribute.GetAttributeSetClass());
 	if (!ensureMsgf(Set, TEXT("FActiveGameplayEffectsContainer::SetAttributeBaseValue: Unable to get attribute set for attribute %s"), *Attribute.AttributeName))
 	{
