Added AbilityActorInfo nullcheck to UAbilitySystemComponent::InvokeGameplayCueEvent

diff --git a/Engine/Plugins/Runtime/GameplayAbilities/Source/GameplayAbilities/Private/AbilitySystemComponent.cpp b/Engine/Plugins/Runtime/GameplayAbilities/Source/GameplayAbilities/Private/AbilitySystemComponent.cpp
--- a/Engine/Plugins/Runtime/GameplayAbilities/Source/GameplayAbilities/Private/AbilitySystemComponent.cpp
+++ b/Engine/Plugins/Runtime/GameplayAbilities/Source/GameplayAbilities/Private/AbilitySystemComponent.cpp
@@ -1257,11 +1257,14 @@ void UAbilitySystemComponent::InvokeGameplayCueEvent(const FGameplayTag Gameplay
 
 void UAbilitySystemComponent::InvokeGameplayCueEvent(const FGameplayTag GameplayCueTag, EGameplayCueEvent::Type EventType, const FGameplayCueParameters& GameplayCueParameters)
 {
-	AActor* ActorAvatar = AbilityActorInfo->AvatarActor.Get();
-	
-	if (ActorAvatar != nullptr && !bSuppressGameplayCues)
+	if(ensureMsgf(AbilityActorInfo != nullptr, TEXT("AbilityActorInfo is null for %s. Probably OnRegister of the AbilitySystemComponent was not called yet."), *this->GetOwner()->GetName()))
 	{
-		UAbilitySystemGlobals::Get().GetGameplayCueManager()->HandleGameplayCue(ActorAvatar, GameplayCueTag, EventType, GameplayCueParameters);
+		AActor* ActorAvatar = AbilityActorInfo->AvatarActor.Get();
+		
+		if (ActorAvatar != nullptr && !bSuppressGameplayCues)
+		{
+			UAbilitySystemGlobals::Get().GetGameplayCueManager()->HandleGameplayCue(ActorAvatar, GameplayCueTag, EventType, GameplayCueParameters);
+		}
 	}
 }
 
