Improving VisLogging Capabilities of the Gameplay Ability System.  Using the VisLog is now viable for figuring out where Gameplay Effects came from.

diff --git a/Engine/Plugins/Runtime/GameplayAbilities/Source/GameplayAbilities/Private/Abilities/GameplayAbility.cpp b/Engine/Plugins/Runtime/GameplayAbilities/Source/GameplayAbilities/Private/Abilities/GameplayAbility.cpp
--- a/Engine/Plugins/Runtime/GameplayAbilities/Source/GameplayAbilities/Private/Abilities/GameplayAbility.cpp
+++ b/Engine/Plugins/Runtime/GameplayAbilities/Source/GameplayAbilities/Private/Abilities/GameplayAbility.cpp
@@ -353,6 +353,13 @@ bool UGameplayAbility::CanActivateAbility(const FGameplayAbilitySpecHandle Handl
 		return false;
 	}
 
+	FGameplayAbilitySpec* Spec = AbilitySystemComponent->FindAbilitySpecFromHandle(Handle);
+	if (!Spec)
+	{
+		ABILITY_LOG(Warning, TEXT("CanActivateAbility %s failed, called with invalid Handle"), *GetName());
+		return false;
+	}
+
 	if (AbilitySystemComponent->GetUserAbilityActivationInhibited())
 	{
 		/**
@@ -364,6 +371,12 @@ bool UGameplayAbility::CanActivateAbility(const FGameplayAbilitySpecHandle Handl
 		 *	Basically: CanActivateAbility is only used by user activated abilities now. If triggered abilities need to check costs/cooldowns, then we may
 		 *	want to split this function up and change the calling API to distinguish between 'can I initiate an ability activation' and 'can this ability be activated'.
 		 */ 
+
+		if (FScopedCanActivateAbilityLogEnabler::IsLoggingEnabled())
+		{
+			UE_LOG(LogAbilitySystem, Verbose, TEXT("%s: %s could not be activated due to GetUserAbilityActivationInhibited"), *GetNameSafe(ActorInfo->OwnerActor.Get()), *GetNameSafe(Spec->Ability));
+			UE_VLOG(ActorInfo->OwnerActor.Get(), VLogAbilitySystem, Verbose, TEXT("%s could not be activated due to GetUserAbilityActivationInhibited"), *GetNameSafe(Spec->Ability));
+		}
 		return false;
 	}
 	
@@ -373,7 +386,8 @@ bool UGameplayAbility::CanActivateAbility(const FGameplayAbilitySpecHandle Handl
 	{
 		if (FScopedCanActivateAbilityLogEnabler::IsLoggingEnabled())
 		{
-			ABILITY_VLOG(ActorInfo->OwnerActor.Get(), Verbose, TEXT("Ability could not be activated due to Cooldown: %s"), *GetName());
+			UE_LOG(LogAbilitySystem, Verbose, TEXT("%s: %s could not be activated due to Cooldown"), *GetNameSafe(ActorInfo->OwnerActor.Get()), *GetNameSafe(Spec->Ability));
+			UE_VLOG(ActorInfo->OwnerActor.Get(), VLogAbilitySystem, Verbose, TEXT("%s could not be activated due to Cooldown"), *GetNameSafe(Spec->Ability));
 		}
 		return false;
 	}
@@ -382,7 +396,8 @@ bool UGameplayAbility::CanActivateAbility(const FGameplayAbilitySpecHandle Handl
 	{
 		if (FScopedCanActivateAbilityLogEnabler::IsLoggingEnabled())
 		{
-			ABILITY_VLOG(ActorInfo->OwnerActor.Get(), Verbose, TEXT("Ability could not be activated due to Cost: %s"), *GetName());
+			UE_LOG(LogAbilitySystem, Verbose, TEXT("%s: %s could not be activated due to Cost"), *GetNameSafe(ActorInfo->OwnerActor.Get()), *GetNameSafe(Spec->Ability));
+			UE_VLOG(ActorInfo->OwnerActor.Get(), VLogAbilitySystem, Verbose, TEXT("%s could not be activated due to Cost"), *GetNameSafe(Spec->Ability));
 		}
 		return false;
 	}
@@ -391,24 +406,19 @@ bool UGameplayAbility::CanActivateAbility(const FGameplayAbilitySpecHandle Handl
 	{	// If the ability's tags are blocked, or if it has a "Blocking" tag or is missing a "Required" tag, then it can't activate.
 		if (FScopedCanActivateAbilityLogEnabler::IsLoggingEnabled())
 		{
-			ABILITY_VLOG(ActorInfo->OwnerActor.Get(), Verbose, TEXT("Ability could not be activated due to Blocking Tags or Missing Required Tags: %s"), *GetName());
+			UE_LOG(LogAbilitySystem, Verbose, TEXT("%s: %s could not be activated due to Blocking Tags or Missing Required Tags"), *GetNameSafe(ActorInfo->OwnerActor.Get()), *GetNameSafe(Spec->Ability));
+			UE_VLOG(ActorInfo->OwnerActor.Get(), VLogAbilitySystem, Verbose, TEXT("%s could not be activated due to Blocking Tags or Missing Required Tags"), *GetNameSafe(Spec->Ability));
 		}
 		return false;
 	}
 
-	FGameplayAbilitySpec* Spec = AbilitySystemComponent->FindAbilitySpecFromHandle(Handle);
-	if (!Spec)
-	{
-		ABILITY_LOG(Warning, TEXT("CanActivateAbility %s failed, called with invalid Handle"), *GetName());
-		return false;
-	}
-
 	// Check if this ability's input binding is currently blocked
 	if (AbilitySystemComponent->IsAbilityInputBlocked(Spec->InputID))
 	{
 		if (FScopedCanActivateAbilityLogEnabler::IsLoggingEnabled())
 		{
-			ABILITY_VLOG(ActorInfo->OwnerActor.Get(), Verbose, TEXT("Ability could not be activated due to blocked input ID %i: %s"), Spec->InputID, *GetName());
+			UE_LOG(LogAbilitySystem, Verbose, TEXT("%s: %s could not be activated due to blocked input ID %d"), *GetNameSafe(ActorInfo->OwnerActor.Get()), *GetNameSafe(Spec->Ability), Spec->InputID);
+			UE_VLOG(ActorInfo->OwnerActor.Get(), VLogAbilitySystem, Verbose, TEXT("%s could not be activated due to blocked input ID %d"), *GetNameSafe(Spec->Ability), Spec->InputID);
 		}
 		return false;
 	}
@@ -417,7 +427,12 @@ bool UGameplayAbility::CanActivateAbility(const FGameplayAbilitySpecHandle Handl
 	{
 		if (K2_CanActivateAbility(*ActorInfo, Handle, OutTags) == false)
 		{
-			ABILITY_LOG(Log, TEXT("CanActivateAbility %s failed, blueprint refused"), *GetName());
+			if (FScopedCanActivateAbilityLogEnabler::IsLoggingEnabled())
+			{
+				UE_LOG(LogAbilitySystem, Verbose, TEXT("%s: CanActivateAbility on %s failed, Blueprint override returned false"), *GetNameSafe(ActorInfo->OwnerActor.Get()), *GetNameSafe(Spec->Ability));
+				UE_VLOG(ActorInfo->OwnerActor.Get(), VLogAbilitySystem, Verbose, TEXT("CanActivateAbility on %s failed, Blueprint override returned false"), *GetNameSafe(Spec->Ability));
+			}
+
 			return false;
 		}
 	}
@@ -1358,14 +1373,16 @@ void UGameplayAbility::OnGameplayTaskInitialized(UGameplayTask& Task)
 
 void UGameplayAbility::OnGameplayTaskActivated(UGameplayTask& Task)
 {
-	ABILITY_VLOG(CastChecked<AActor>(GetOuter()), Log, TEXT("Task Started %s"), *Task.GetName());
+	UE_VLOG(CastChecked<AActor>(GetOuter()), LogGameplayTasks, VeryVerbose, TEXT("GameplayAbility Task Started %s"), *Task.GetName());
+	UE_LOG(LogGameplayTasks, VeryVerbose, TEXT("GameplayAbility Task Started %s"), *Task.GetName());
 
 	ActiveTasks.Add(&Task);
 }
 
 void UGameplayAbility::OnGameplayTaskDeactivated(UGameplayTask& Task)
 {
-	ABILITY_VLOG(CastChecked<AActor>(GetOuter()), Log, TEXT("Task Ended %s"), *Task.GetName());
+	UE_VLOG(CastChecked<AActor>(GetOuter()), LogGameplayTasks, VeryVerbose, TEXT("GameplayAbility Task Ended %s"), *Task.GetName());
+	UE_LOG(LogGameplayTasks, VeryVerbose, TEXT("GameplayAbility Task Ended %s"), *Task.GetName());
 
 	ActiveTasks.Remove(&Task);
 
diff --git a/Engine/Plugins/Runtime/GameplayAbilities/Source/GameplayAbilities/Private/AbilitySystemComponent.cpp b/Engine/Plugins/Runtime/GameplayAbilities/Source/GameplayAbilities/Private/AbilitySystemComponent.cpp
--- a/Engine/Plugins/Runtime/GameplayAbilities/Source/GameplayAbilities/Private/AbilitySystemComponent.cpp
+++ b/Engine/Plugins/Runtime/GameplayAbilities/Source/GameplayAbilities/Private/AbilitySystemComponent.cpp
@@ -880,13 +880,13 @@ FActiveGameplayEffectHandle UAbilitySystemComponent::ApplyGameplayEffectSpecToSe
 			// Log results of applied GE spec
 			if (UE_LOG_ACTIVE(VLogAbilitySystem, Log))
 			{
-				ABILITY_VLOG(GetOwnerActor(), Log, TEXT("Applied %s"), *OurCopyOfSpec->Def->GetFName().ToString());
+				UE_VLOG(GetOwnerActor(), VLogAbilitySystem, Log, TEXT("Applied %s"), *OurCopyOfSpec->Def->GetFName().ToString());
 
 				for (const FGameplayModifierInfo& Modifier : Spec.Def->Modifiers)
 				{
 					float Magnitude = 0.f;
 					Modifier.ModifierMagnitude.AttemptCalculateMagnitude(Spec, Magnitude);
-					ABILITY_VLOG(GetOwnerActor(), Log, TEXT("         %s: %s %f"), *Modifier.Attribute.GetName(), *EGameplayModOpToString(Modifier.ModifierOp), Magnitude);
+					UE_VLOG(GetOwnerActor(), VLogAbilitySystem, Log, TEXT("         %s: %s %f"), *Modifier.Attribute.GetName(), *EGameplayModOpToString(Modifier.ModifierOp), Magnitude);
 				}
 			}
 		}
@@ -1002,13 +1002,13 @@ void UAbilitySystemComponent::ExecuteGameplayEffect(FGameplayEffectSpec &Spec, F
 
 	if (UE_LOG_ACTIVE(VLogAbilitySystem, Log))
 	{
-		ABILITY_VLOG(GetOwnerActor(), Log, TEXT("Executed %s"), *Spec.Def->GetFName().ToString());
+		UE_VLOG(GetOwnerActor(), VLogAbilitySystem, Log, TEXT("Executed %s"), *Spec.Def->GetFName().ToString());
 		
 		for (const FGameplayModifierInfo& Modifier : Spec.Def->Modifiers)
 		{
 			float Magnitude = 0.f;
 			Modifier.ModifierMagnitude.AttemptCalculateMagnitude(Spec, Magnitude);
-			ABILITY_VLOG(GetOwnerActor(), Log, TEXT("         %s: %s %f"), *Modifier.Attribute.GetName(), *EGameplayModOpToString(Modifier.ModifierOp), Magnitude);
+			UE_VLOG(GetOwnerActor(), VLogAbilitySystem, Log, TEXT("         %s: %s %f"), *Modifier.Attribute.GetName(), *EGameplayModOpToString(Modifier.ModifierOp), Magnitude);
 		}
 	}
 
@@ -1803,6 +1803,100 @@ void UAbilitySystemComponent::PrintAllGameplayEffects() const
 	ActiveGameplayEffects.PrintAllGameplayEffects();
 }
 
+#if ENABLE_VISUAL_LOG
+static FVisualLogStatusCategory GrabDebugSnapshot_GameplayAbilities(const UAbilitySystemComponent* ASC)
+{
+	FVisualLogStatusCategory AllAbilitiesStatus;
+	AllAbilitiesStatus.Category = TEXT("Gameplay Abilities");
+
+	for (const FGameplayAbilitySpec& Spec : ASC->GetActivatableAbilities())
+	{
+		FVisualLogStatusCategory AbilityStatus;
+		const UGameplayAbility* AbilitySource = Spec.GetPrimaryInstance() ? Spec.GetPrimaryInstance() : Spec.Ability.Get();
+		AbilityStatus.Category = FString::Printf(TEXT("%s[%s] %s"), Spec.IsActive() ? TEXT("**") : TEXT(""), *Spec.Handle.ToString(), *GetNameSafe(AbilitySource));
+
+		AbilityStatus.Add(TEXT("ActiveCount"), FString::Printf(TEXT("%d"), Spec.ActiveCount));
+		AbilityStatus.Add(TEXT("Level"), FString::Printf(TEXT("%d"), Spec.Level));
+
+		if (UObject* SourceObject = Spec.SourceObject.Get())
+		{
+			AbilityStatus.Add(TEXT("SourceObject"), *GetNameSafe(SourceObject));
+		}
+
+		FActiveGameplayEffectHandle AGEHandle = ASC->FindActiveGameplayEffectHandle(Spec.Handle);
+		if (AGEHandle.IsValid())
+		{
+			if (const FActiveGameplayEffect* ActiveGE = ASC->GetActiveGameplayEffect(AGEHandle))
+			{
+				AbilityStatus.Add(TEXT("Granting Effect"), FString::Printf(TEXT("[%s] %s"), *AGEHandle.ToString(), *ActiveGE->Spec.ToSimpleString()));
+			}
+			else
+			{
+				AbilityStatus.Add(TEXT("Granting Effect"), FString::Printf(TEXT("[%s] NOT FOUND"), *AGEHandle.ToString()));
+			}
+		}
+
+		if (FGameplayEventData* GameplayEventData = Spec.GameplayEventData.Get())
+		{
+			FVisualLogStatusCategory EventDataStatus;
+			EventDataStatus.Category = TEXT("GameplayEventData");
+
+			EventDataStatus.Add(TEXT("EventTag"), GameplayEventData->EventTag.ToString());
+
+#define EventDataStatus_AddOptional(x) if (GameplayEventData-> x) { EventDataStatus.Add(TEXT("##x"), *GetNameSafe(GameplayEventData-> x)); }
+			EventDataStatus_AddOptional(Instigator);
+			EventDataStatus_AddOptional(Target);
+			EventDataStatus_AddOptional(OptionalObject);
+			EventDataStatus_AddOptional(OptionalObject2);
+#undef EventDataStatus_AddOptional
+
+			EventDataStatus.Add(TEXT("Context"), GameplayEventData->ContextHandle.ToString());
+
+			if (GameplayEventData->InstigatorTags.Num())
+			{
+				EventDataStatus.Add(TEXT("InstigatorTags"), GameplayEventData->InstigatorTags.ToStringSimple());
+			}
+
+			if (GameplayEventData->TargetTags.Num())
+			{
+				EventDataStatus.Add(TEXT("TargetTags"), GameplayEventData->TargetTags.ToStringSimple());
+			}
+
+			if (GameplayEventData->EventMagnitude != 0.0f)
+			{
+				EventDataStatus.Add(TEXT("EventMagnitude"), FString::Printf(TEXT("%.3f"), GameplayEventData->EventMagnitude));
+			}
+
+			for (int32 Index = 0; Index < GameplayEventData->TargetData.Num(); ++Index)
+			{
+				FGameplayAbilityTargetData* TargetData = GameplayEventData->TargetData.Get(Index);
+				EventDataStatus.Add(FString::Printf(TEXT("TargetData[%d]"), Index), TargetData ? *TargetData->ToString() : TEXT("null"));
+			}
+
+			AbilityStatus.AddChild(EventDataStatus);
+		}
+
+		AllAbilitiesStatus.AddChild(AbilityStatus);
+	}
+
+	return AllAbilitiesStatus;
+}
+
+void UAbilitySystemComponent::GrabDebugSnapshot(FVisualLogEntry* Snapshot) const
+{
+	Super::GrabDebugSnapshot(Snapshot);
+
+	if (ActivatableAbilities.Items.Num() > 0)
+	{
+		FVisualLogStatusCategory AbilitiesStatus = GrabDebugSnapshot_GameplayAbilities(this);
+		Snapshot->Status.Add(AbilitiesStatus);
+	}
+
+	ActiveGameplayEffects.DescribeSelfToVisLog(Snapshot);
+}
+#endif
+
+
 bool UAbilitySystemComponent::IsOwnerActorAuthoritative() const
 {
 	return !bCachedIsNetSimulated;
diff --git a/Engine/Plugins/Runtime/GameplayAbilities/Source/GameplayAbilities/Private/AbilitySystemComponent_Abilities.cpp b/Engine/Plugins/Runtime/GameplayAbilities/Source/GameplayAbilities/Private/AbilitySystemComponent_Abilities.cpp
--- a/Engine/Plugins/Runtime/GameplayAbilities/Source/GameplayAbilities/Private/AbilitySystemComponent_Abilities.cpp
+++ b/Engine/Plugins/Runtime/GameplayAbilities/Source/GameplayAbilities/Private/AbilitySystemComponent_Abilities.cpp
@@ -34,6 +34,7 @@
 #include "ProfilingDebugging/CsvProfiler.h"
 #include "GameFramework/Character.h"
 #include "GameFramework/CharacterMovementComponent.h"
+#include "VisualLogger/VisualLogger.h"
 
 #define LOCTEXT_NAMESPACE "AbilitySystemComponent"
 
@@ -282,6 +283,7 @@ FGameplayAbilitySpecHandle UAbilitySystemComponent::GiveAbility(const FGameplayA
 	// If locked, add to pending list. The Spec.Handle is not regenerated when we receive, so returning this is ok.
 	if (AbilityScopeLockCount > 0)
 	{
+		UE_LOG(LogAbilitySystem, Verbose, TEXT("%s: GiveAbility %s delayed (ScopeLocked)"), *GetNameSafe(GetOwner()), *GetNameSafe(Spec.Ability));
 		AbilityPendingAdds.Add(Spec);
 		return Spec.Handle;
 	}
@@ -298,6 +300,8 @@ FGameplayAbilitySpecHandle UAbilitySystemComponent::GiveAbility(const FGameplayA
 	OnGiveAbility(OwnedSpec);
 	MarkAbilitySpecDirty(OwnedSpec, true);
 
+	UE_LOG(LogAbilitySystem, Log, TEXT("%s: GiveAbility %s [%s] Level: %d Source: %s"), *GetNameSafe(GetOwner()), *GetNameSafe(Spec.Ability), *Spec.Handle.ToString(), Spec.Level, *GetNameSafe(Spec.SourceObject.Get()));
+	UE_VLOG(GetOwner(), VLogAbilitySystem, Log, TEXT("GiveAbility %s [%s] Level: %d Source: %s"), *GetNameSafe(Spec.Ability), *Spec.Handle.ToString(), Spec.Level, *GetNameSafe(Spec.SourceObject.Get()));
 	return OwnedSpec.Handle;
 }
 
@@ -682,6 +686,9 @@ void UAbilitySystemComponent::OnRemoveAbility(FGameplayAbilitySpec& Spec)
 
 	Spec.ReplicatedInstances.Empty();
 	Spec.NonReplicatedInstances.Empty();
+
+	UE_LOG(LogAbilitySystem, Log, TEXT("%s: Removed Ability %s [%s] Level: %d"), *GetNameSafe(GetOwner()), *GetNameSafe(Spec.Ability), *Spec.Handle.ToString(), Spec.Level);
+	UE_VLOG(GetOwner(), VLogAbilitySystem, Log, TEXT("Removed Ability %s [%s] Level: %d"), *GetNameSafe(Spec.Ability), *Spec.Handle.ToString(), Spec.Level);
 }
 
 void UAbilitySystemComponent::CheckForClearedAbilities()
@@ -1108,6 +1115,9 @@ void UAbilitySystemComponent::NotifyAbilityEnded(FGameplayAbilitySpecHandle Hand
 		return;
 	}
 
+	UE_LOG(LogAbilitySystem, Log, TEXT("%s: Ended [%s] %s. Level: %d. WasCancelled: %d."), *GetNameSafe(GetOwner()), *Handle.ToString(), Spec->GetPrimaryInstance() ? *Spec->GetPrimaryInstance()->GetName() : *Ability->GetName(), Spec->Level, bWasCancelled);
+	UE_VLOG(GetOwner(), VLogAbilitySystem, Log, TEXT("Ended [%s] %s. Level: %d. WasCancelled: %d."), *Handle.ToString(), Spec->GetPrimaryInstance() ? *Spec->GetPrimaryInstance()->GetName() : *Ability->GetName(), Spec->Level, bWasCancelled);
+
 	ENetRole OwnerRole = GetOwnerRole();
 
 	// If AnimatingAbility ended, clear the pointer
@@ -1625,8 +1635,8 @@ bool UAbilitySystemComponent::InternalTryActivateAbility(FGameplayAbilitySpecHan
 		if (Ability->GetNetExecutionPolicy() == EGameplayAbilityNetExecutionPolicy::LocalOnly || (Ability->GetNetExecutionPolicy() == EGameplayAbilityNetExecutionPolicy::LocalPredicted && !InPredictionKey.IsValidKey()))
 		{
 			// If we have a valid prediction key, the ability was started on the local client so it's okay
-
-			ABILITY_LOG(Warning, TEXT("Can't activate %s ability %s when not local."), *UEnum::GetValueAsString<EGameplayAbilityNetExecutionPolicy::Type>(Ability->GetNetExecutionPolicy()), *Ability->GetName());
+			UE_LOG(LogAbilitySystem, Warning, TEXT("%s: Can't activate %s ability %s when not local"), *GetNameSafe(GetOwner()), *UEnum::GetValueAsString<EGameplayAbilityNetExecutionPolicy::Type>(Ability->GetNetExecutionPolicy()), *Ability->GetName());
+			UE_VLOG(GetOwner(), VLogAbilitySystem, Warning, TEXT("Can't activate %s ability %s when not local"), *UEnum::GetValueAsString<EGameplayAbilityNetExecutionPolicy::Type>(Ability->GetNetExecutionPolicy()), *Ability->GetName());
 
 			if (NetworkFailTag.IsValid())
 			{
@@ -1640,7 +1650,8 @@ bool UAbilitySystemComponent::InternalTryActivateAbility(FGameplayAbilitySpecHan
 
 	if (NetMode != ROLE_Authority && (Ability->GetNetExecutionPolicy() == EGameplayAbilityNetExecutionPolicy::ServerOnly || Ability->GetNetExecutionPolicy() == EGameplayAbilityNetExecutionPolicy::ServerInitiated))
 	{
-		ABILITY_LOG(Warning, TEXT("Can't activate %s ability %s when not the server."), *UEnum::GetValueAsString<EGameplayAbilityNetExecutionPolicy::Type>(Ability->GetNetExecutionPolicy()), *Ability->GetName());
+		UE_LOG(LogAbilitySystem, Warning, TEXT("%s: Can't activate %s ability %s when not the server"), *GetNameSafe(GetOwner()), *UEnum::GetValueAsString<EGameplayAbilityNetExecutionPolicy::Type>(Ability->GetNetExecutionPolicy()), *Ability->GetName());
+		UE_VLOG(GetOwner(), VLogAbilitySystem, Warning, TEXT("Can't activate %s ability %s when not the server"), *UEnum::GetValueAsString<EGameplayAbilityNetExecutionPolicy::Type>(Ability->GetNetExecutionPolicy()), *Ability->GetName());
 
 		if (NetworkFailTag.IsValid())
 		{
@@ -1653,12 +1664,15 @@ bool UAbilitySystemComponent::InternalTryActivateAbility(FGameplayAbilitySpecHan
 
 	// If it's an instanced one, the instanced ability will be set, otherwise it will be null
 	UGameplayAbility* InstancedAbility = Spec->GetPrimaryInstance();
+	UGameplayAbility* AbilitySource = InstancedAbility ? InstancedAbility : Ability;
 
 	if (TriggerEventData)
 	{
-		UGameplayAbility* AbilitySource = InstancedAbility ? InstancedAbility : Ability;
 		if (!AbilitySource->ShouldAbilityRespondToEvent(ActorInfo, TriggerEventData))
 		{
+			UE_LOG(LogAbilitySystem, Verbose, TEXT("%s: Can't activate %s because ShouldAbilityRespondToEvent was false."), *GetNameSafe(GetOwner()), *Ability->GetName());
+			UE_VLOG(GetOwner(), VLogAbilitySystem, Verbose, TEXT("Can't activate %s because ShouldAbilityRespondToEvent was false."), *Ability->GetName());
+
 			NotifyAbilityFailed(Handle, AbilitySource, InternalTryActivateAbilityFailureTags);
 			return false;
 		}
@@ -1668,14 +1682,11 @@ bool UAbilitySystemComponent::InternalTryActivateAbility(FGameplayAbilitySpecHan
 		const FGameplayTagContainer* SourceTags = TriggerEventData ? &TriggerEventData->InstigatorTags : nullptr;
 		const FGameplayTagContainer* TargetTags = TriggerEventData ? &TriggerEventData->TargetTags : nullptr;
 
-		// If we have an instanced ability, call CanActivateAbility on it.
-		// Otherwise we always do a non instanced CanActivateAbility check using the CDO of the Ability.
-		UGameplayAbility* const CanActivateAbilitySource = InstancedAbility ? InstancedAbility : Ability;
 		FScopedCanActivateAbilityLogEnabler LogEnabler;
-
-		if (!CanActivateAbilitySource->CanActivateAbility(Handle, ActorInfo, SourceTags, TargetTags, &InternalTryActivateAbilityFailureTags))
+		if (!AbilitySource->CanActivateAbility(Handle, ActorInfo, SourceTags, TargetTags, &InternalTryActivateAbilityFailureTags))
 		{
-			NotifyAbilityFailed(Handle, CanActivateAbilitySource, InternalTryActivateAbilityFailureTags);
+			// CanActivateAbility with LogEnabler will have UE_LOG/UE_VLOG so don't add more failure logs here
+			NotifyAbilityFailed(Handle, AbilitySource, InternalTryActivateAbilityFailureTags);
 			return false;
 		}
 	}
@@ -1687,13 +1698,16 @@ bool UAbilitySystemComponent::InternalTryActivateAbility(FGameplayAbilitySpecHan
 		{
 			if (Ability->bRetriggerInstancedAbility && InstancedAbility)
 			{
+				UE_LOG(LogAbilitySystem, Verbose, TEXT("%s: Ending %s prematurely to retrigger."), *GetNameSafe(GetOwner()), *Ability->GetName());
+				UE_VLOG(GetOwner(), VLogAbilitySystem, Verbose, TEXT("Ending %s prematurely to retrigger."), *Ability->GetName());
+
 				bool bReplicateEndAbility = true;
 				bool bWasCancelled = false;
 				InstancedAbility->EndAbility(Handle, ActorInfo, Spec->ActivationInfo, bReplicateEndAbility, bWasCancelled);
 			}
 			else
 			{
-				ABILITY_LOG(Verbose, TEXT("Can't activate instanced per actor ability %s when their is already a currently active instance for this actor."), *Ability->GetName());
+				UE_LOG(LogAbilitySystem, Verbose, TEXT("Can't activate instanced per actor ability %s when their is already a currently active instance for this actor."), *Ability->GetName());
 				return false;
 			}
 		}
@@ -1702,7 +1716,7 @@ bool UAbilitySystemComponent::InternalTryActivateAbility(FGameplayAbilitySpecHan
 	// Make sure we have a primary
 	if (Ability->GetInstancingPolicy() == EGameplayAbilityInstancingPolicy::InstancedPerActor && !InstancedAbility)
 	{
-		ABILITY_LOG(Warning, TEXT("InternalTryActivateAbility called but instanced ability is missing! NetMode: %d. Ability: %s"), (int32)NetMode, *Ability->GetName());
+		UE_LOG(LogAbilitySystem, Warning, TEXT("InternalTryActivateAbility called but instanced ability is missing! NetMode: %d. Ability: %s"), (int32)NetMode, *Ability->GetName());
 		return false;
 	}
 
@@ -1759,13 +1773,9 @@ bool UAbilitySystemComponent::InternalTryActivateAbility(FGameplayAbilitySpecHan
 			InstancedAbility = CreateNewInstanceOfAbility(*Spec, Ability);
 			InstancedAbility->CallActivateAbility(Handle, ActorInfo, ActivationInfo, OnGameplayAbilityEndedDelegate, TriggerEventData);
 		}
-		else if (InstancedAbility)
-		{
-			InstancedAbility->CallActivateAbility(Handle, ActorInfo, ActivationInfo, OnGameplayAbilityEndedDelegate, TriggerEventData);
-		}
 		else
 		{
-			Ability->CallActivateAbility(Handle, ActorInfo, ActivationInfo, OnGameplayAbilityEndedDelegate, TriggerEventData);
+			AbilitySource->CallActivateAbility(Handle, ActorInfo, ActivationInfo, OnGameplayAbilityEndedDelegate, TriggerEventData);
 		}
 	}
 	else if (Ability->GetNetExecutionPolicy() == EGameplayAbilityNetExecutionPolicy::LocalPredicted)
@@ -1819,13 +1829,9 @@ bool UAbilitySystemComponent::InternalTryActivateAbility(FGameplayAbilitySpecHan
 				ABILITY_LOG(Error, TEXT("InternalTryActivateAbility called on ability %s that is InstancePerExecution and Replicated. This is an invalid configuration."), *Ability->GetName() );
 			}
 		}
-		else if (InstancedAbility)
-		{
-			InstancedAbility->CallActivateAbility(Handle, ActorInfo, ActivationInfo, OnGameplayAbilityEndedDelegate, TriggerEventData);
-		}
-		else 
+		else
 		{
-			Ability->CallActivateAbility(Handle, ActorInfo, ActivationInfo, OnGameplayAbilityEndedDelegate, TriggerEventData);
+			AbilitySource->CallActivateAbility(Handle, ActorInfo, ActivationInfo, OnGameplayAbilityEndedDelegate, TriggerEventData);
 		}
 	}
 	
@@ -1847,6 +1853,8 @@ bool UAbilitySystemComponent::InternalTryActivateAbility(FGameplayAbilitySpecHan
 
 	AbilityLastActivatedTime = GetWorld()->GetTimeSeconds();
 
+	UE_LOG(LogAbilitySystem, Log, TEXT("%s: Activated [%s] %s. Level: %d. PredictionKey: %s."), *GetNameSafe(GetOwner()), *Spec->Handle.ToString(), *GetNameSafe(AbilitySource), Spec->Level, *ActivationInfo.GetActivationPredictionKey().ToString());
+	UE_VLOG(GetOwner(), VLogAbilitySystem, Log, TEXT("Activated [%s] %s. Level: %d. PredictionKey: %s."), *Spec->Handle.ToString(), *GetNameSafe(AbilitySource), Spec->Level, *ActivationInfo.GetActivationPredictionKey().ToString());
 	return true;
 }
 
diff --git a/Engine/Plugins/Runtime/GameplayAbilities/Source/GameplayAbilities/Private/GameplayEffect.cpp b/Engine/Plugins/Runtime/GameplayAbilities/Source/GameplayAbilities/Private/GameplayEffect.cpp
--- a/Engine/Plugins/Runtime/GameplayAbilities/Source/GameplayAbilities/Private/GameplayEffect.cpp
+++ b/Engine/Plugins/Runtime/GameplayAbilities/Source/GameplayAbilities/Private/GameplayEffect.cpp
@@ -799,7 +799,7 @@ bool UGameplayEffect::OnAddedToActiveContainer(FActiveGameplayEffectsContainer&
 		}
 	}
 
-	UE_VLOG_UELOG(ActiveGEContainer.Owner->GetOwnerActor(), LogGameplayEffects, VeryVerbose, TEXT("%s: %s. ShouldBeActive: %d"), ANSI_TO_TCHAR(__func__), *ActiveGE.GetDebugString(), bShouldBeActive);
+	UE_VLOG_UELOG(ActiveGEContainer.Owner->GetOwnerActor(), LogGameplayEffects, Log, TEXT("Added: %s. Auth: %d. ReplicationID: %d. ShouldBeActive: %d"), *ActiveGE.GetDebugString(), ActiveGEContainer.IsNetAuthority(), ActiveGE.ReplicationID, bShouldBeActive);
 
 	return bShouldBeActive;
 }
@@ -814,7 +814,7 @@ void UGameplayEffect::OnExecuted(FActiveGameplayEffectsContainer& ActiveGEContai
 		}
 	}
 
-	UE_VLOG_UELOG(ActiveGEContainer.Owner->GetOwnerActor(), LogGameplayEffects, VeryVerbose, TEXT("%s: %s"), ANSI_TO_TCHAR(__func__), *GetNameSafe(GESpec.Def));
+	UE_VLOG_UELOG(ActiveGEContainer.Owner->GetOwnerActor(), LogGameplayEffects, Log, TEXT("Executed: %s"), *GetNameSafe(GESpec.Def));
 }
 
 void UGameplayEffect::OnApplied(FActiveGameplayEffectsContainer& ActiveGEContainer, FGameplayEffectSpec& GESpec, FPredictionKey& PredictionKey) const
@@ -827,7 +827,7 @@ void UGameplayEffect::OnApplied(FActiveGameplayEffectsContainer& ActiveGEContain
 		}
 	}
 
-	UE_VLOG_UELOG(ActiveGEContainer.Owner->GetOwnerActor(), LogGameplayEffects, VeryVerbose, TEXT("%s: %s"), ANSI_TO_TCHAR(__func__), *GetNameSafe(GESpec.Def));
+	UE_LOG(LogGameplayEffects, Verbose, TEXT("Applied: %s"), *GetNameSafe(GESpec.Def));
 }
 
 int32 UGameplayEffect::GetStackLimitCount() const
@@ -1835,6 +1835,135 @@ void FGameplayEffectSpec::RecaptureAttributeDataForClone(UAbilitySystemComponent
 	}
 }
 
+#if ENABLE_VISUAL_LOG
+FVisualLogStatusCategory FGameplayEffectSpec::GrabVisLogStatus() const
+{
+	FVisualLogStatusCategory SpecStatus;
+	SpecStatus.Category = FString::Printf(TEXT("Spec: %s"), *GetNameSafe(Def));
+
+	SpecStatus.Add(TEXT("Level"), FString::Printf(TEXT("%.2f"), Level));
+
+	if (Duration == UGameplayEffect::INSTANT_APPLICATION)
+	{
+		SpecStatus.Add(TEXT("Duration"), TEXT("Instant"));
+	}
+	else if (Duration == UGameplayEffect::INFINITE_DURATION)
+	{
+		SpecStatus.Add(TEXT("Duration"), TEXT("Infinite"));
+	}
+	else
+	{
+		SpecStatus.Add(TEXT("Duration"), FString::Printf(TEXT("%.3f"), Duration));
+	}
+
+	if (Period > 0.0f)
+	{
+		SpecStatus.Add(TEXT("Period"), FString::Printf(TEXT("%.3f"), Period));
+	}
+
+	int32 LocalStackCount = GetStackCount();
+	if (LocalStackCount > 0)
+	{
+		SpecStatus.Add(TEXT("StackCount"), FString::Printf(TEXT("%d"), LocalStackCount));
+	}
+
+	if (DynamicGrantedTags.Num() > 0)
+	{
+		SpecStatus.Add(TEXT("DynamicGrantedTags"), DynamicGrantedTags.ToStringSimple());
+	}
+
+	if (GetDynamicAssetTags().Num() > 0)
+	{
+		SpecStatus.Add(TEXT("DynamicAssetTags"), GetDynamicAssetTags().ToStringSimple());
+	}
+
+	auto AddTagContainerAggregator = [&](FString&& InName, const FTagContainerAggregator& InContainer)
+		{
+			const FGameplayTagContainer* AggregatedTags = InContainer.GetAggregatedTags();
+			if (InContainer.GetActorTags().Num() || InContainer.GetSpecTags().Num() || (AggregatedTags && AggregatedTags->Num()))
+			{
+				FVisualLogStatusCategory Status;
+				Status.Category = MoveTemp(InName);
+
+				Status.Add(TEXT("ActorTags"), InContainer.GetActorTags().ToStringSimple());
+				Status.Add(TEXT("AggregatedTags"), AggregatedTags ? AggregatedTags->ToStringSimple() : FString{});
+				Status.Add(TEXT("SpecTags"), InContainer.GetSpecTags().ToStringSimple());
+
+				SpecStatus.AddChild(Status);
+			}
+		};
+	AddTagContainerAggregator(TEXT("CapturedSourceTags"), CapturedSourceTags);
+	AddTagContainerAggregator(TEXT("CapturedTargetTags"), CapturedTargetTags);
+
+	// Handle the SetByCallers
+	{
+		FVisualLogStatusCategory SetByCallersStatus;
+		SetByCallersStatus.Category = TEXT("SetByCallers");
+		for (const auto& Entry : SetByCallerNameMagnitudes)
+		{
+			SetByCallersStatus.Add(Entry.Key.ToString(), FString::Printf(TEXT("%.3f"), Entry.Value));
+		}
+
+		for (const auto& Entry : SetByCallerTagMagnitudes)
+		{
+			SetByCallersStatus.Add(Entry.Key.ToString(), FString::Printf(TEXT("%.3f"), Entry.Value));
+		}
+
+		// If we had any data, attach us
+		if (SetByCallersStatus.Data.Num() > 0)
+		{
+			SpecStatus.AddChild(SetByCallersStatus);
+		}
+	}
+
+	// Handle the EffectContext
+	{
+		FVisualLogStatusCategory EffectContextStatus;
+		EffectContextStatus.Category = TEXT("Effect Context");
+
+		const FGameplayEffectContextHandle& LocalContext = GetEffectContext();
+
+		if (const UGameplayAbility* GrantingAbility = LocalContext.GetAbility())
+		{
+			EffectContextStatus.Add(TEXT("Ability"), *GrantingAbility->GetName());
+		}
+
+		if (const UObject* SourceObject = LocalContext.GetSourceObject())
+		{
+			EffectContextStatus.Add(TEXT("SourceObject"), *SourceObject->GetName());
+		}
+
+		if (const AActor* Instigator = LocalContext.GetInstigator())
+		{
+			EffectContextStatus.Add(TEXT("Instigator"), *Instigator->GetName());
+		}
+
+		if (const AActor* EffectCauser = LocalContext.GetEffectCauser())
+		{
+			EffectContextStatus.Add(TEXT("EffectCauser"), *EffectCauser->GetName());
+		}
+
+		if (LocalContext.HasOrigin())
+		{
+			EffectContextStatus.Add(TEXT("Origin"), *LocalContext.GetOrigin().ToString());
+		}
+
+		if (const FHitResult* HitResult = LocalContext.GetHitResult())
+		{
+			EffectContextStatus.Add(TEXT("HitResult"), *HitResult->ToString());
+		}
+
+		// If we had any data, attach us
+		if (EffectContextStatus.Data.Num() > 0)
+		{
+			SpecStatus.AddChild(EffectContextStatus);
+		}
+	}
+
+	return SpecStatus;
+}
+#endif
+
 const FGameplayEffectModifiedAttribute* FGameplayEffectSpec::GetModifiedAttribute(const FGameplayAttribute& Attribute) const
 {
 	for (const FGameplayEffectModifiedAttribute& ModifiedAttribute : ModifiedAttributes)
@@ -3931,8 +4060,6 @@ PRAGMA_ENABLE_DEPRECATION_WARNINGS
 	}
 	else
 	{
-		UE_VLOG_UELOG(OwnerActor, LogGameplayEffects, Verbose, TEXT("Adding GE: %s. Auth: %d. ReplicationID: %d. ReplicationKey: %d."), *AppliedActiveGE->GetDebugString(), IsNetAuthority(), AppliedActiveGE->ReplicationID, AppliedActiveGE->ReplicationKey);
-
 		// Since we are applying it locally (and possibly predictively) invoke the cues.  Unless it's an Instant Cue, in which case we're not invoking the OnActive/WhileActive cues.
 		const bool bInvokeGameplayCueEvents = (Spec.Def->DurationPolicy != EGameplayEffectDurationType::Instant);
 		InternalOnActiveGameplayEffectAdded(*AppliedActiveGE, bInvokeGameplayCueEvents);
@@ -5112,39 +5239,95 @@ void FActiveGameplayEffectsContainer::GetActiveGameplayEffectDataByAttribute(TMu
 }
 
 #if ENABLE_VISUAL_LOG
-void FActiveGameplayEffectsContainer::GrabDebugSnapshot(FVisualLogEntry* Snapshot) const
+void FActiveGameplayEffectsContainer::DescribeSelfToVisLog(FVisualLogEntry* Snapshot) const
 {
-	FVisualLogStatusCategory ActiveEffectsCategory;
-	ActiveEffectsCategory.Category = TEXT("Effects");
-
 	TMultiMap<FGameplayAttribute, FActiveGameplayEffectsContainer::DebugExecutedGameplayEffectData> EffectMap;
 
 	GetActiveGameplayEffectDataByAttribute(EffectMap);
 
-	// For each attribute that was modified go through all of its modifiers and list them
-	TArray<FGameplayAttribute> AttributeKeys;
-	EffectMap.GetKeys(AttributeKeys);
-
-	for (const FGameplayAttribute& Attribute : AttributeKeys)
+	if (EffectMap.Num() > 0)
 	{
-		float CombinedModifierValue = 0.f;
-		ActiveEffectsCategory.Add(TEXT(" --- Attribute --- "), Attribute.GetName());
+		FVisualLogStatusCategory AttributeModStatus;
+		AttributeModStatus.Category = TEXT("Attribute Mods");
 
-		TArray<FActiveGameplayEffectsContainer::DebugExecutedGameplayEffectData> AttributeEffects;
-		EffectMap.MultiFind(Attribute, AttributeEffects);
+		// For each attribute that was modified go through all of its modifiers and list them
+		TArray<FGameplayAttribute> AttributeKeys;
+		EffectMap.GetKeys(AttributeKeys);
 
-		for (const FActiveGameplayEffectsContainer::DebugExecutedGameplayEffectData& DebugData : AttributeEffects)
+		for (const FGameplayAttribute& Attribute : AttributeKeys)
 		{
-			ActiveEffectsCategory.Add(DebugData.GameplayEffectName, DebugData.ActivationState);
-			ActiveEffectsCategory.Add(TEXT("Magnitude"), FString::Printf(TEXT("%f"), DebugData.Magnitude));
+			FVisualLogStatusCategory AttributeModCategory;
+			AttributeModCategory.Category = Attribute.GetName();
+
+			TArray<FActiveGameplayEffectsContainer::DebugExecutedGameplayEffectData> AttributeEffects;
+			EffectMap.MultiFind(Attribute, AttributeEffects);
 
-			if (DebugData.ActivationState != "INHIBITED")
+			float CombinedModifierValue = 0.f;
+			for (const FActiveGameplayEffectsContainer::DebugExecutedGameplayEffectData& DebugData : AttributeEffects)
 			{
-				CombinedModifierValue += DebugData.Magnitude;
+				AttributeModCategory.Add(DebugData.GameplayEffectName, DebugData.ActivationState);
+				AttributeModCategory.Add(TEXT("Magnitude"), FString::Printf(TEXT("%f"), DebugData.Magnitude));
+
+				if (DebugData.ActivationState != "INHIBITED")
+				{
+					CombinedModifierValue += DebugData.Magnitude;
+				}
 			}
+
+			AttributeModCategory.Add(TEXT("Total Modification"), FString::Printf(TEXT("%f"), CombinedModifierValue));
+			AttributeModStatus.AddChild(AttributeModCategory);
+		}
+
+		Snapshot->Status.Add(AttributeModStatus);
+	}
+
+	FVisualLogStatusCategory ActiveEffectsCategory;
+	ActiveEffectsCategory.Category = TEXT("Active Effects");
+
+	// Iterating through manually since this is a removal operation and we need to pass the index into InternalRemoveActiveGameplayEffect
+	int32 NumGameplayEffects = GetNumGameplayEffects();
+	for (int32 ActiveGEIdx = 0; ActiveGEIdx < NumGameplayEffects; ++ActiveGEIdx)
+	{
+		FVisualLogStatusCategory ActiveGELog;
+
+		const FActiveGameplayEffect& Effect = *GetActiveGameplayEffect(ActiveGEIdx);
+		ActiveGELog.Category = FString::Printf(TEXT("[%s] %s"), *Effect.Handle.ToString(), *Effect.Spec.ToSimpleString());
+
+		FVisualLogStatusCategory SpecLogStatus = Effect.Spec.GrabVisLogStatus();
+		ActiveGELog.AddChild(SpecLogStatus);
+
+		if (Effect.StartWorldTime > 0.0f)
+		{
+			ActiveGELog.Add(TEXT("StartWorldTime"), FString::Printf(TEXT("%.3f"), Effect.StartWorldTime));
+		}
+
+		if (Effect.StartServerWorldTime > 0.0f)
+		{
+			ActiveGELog.Add(TEXT("ServerStartWorldTime"), FString::Printf(TEXT("%.3f"), Effect.StartServerWorldTime));
+		}
+
+		if (Effect.DurationHandle.IsValid())
+		{
+			ActiveGELog.Add(TEXT("Duration"), FString::Printf(TEXT("%.3f"), Effect.GetDuration()));
+			ActiveGELog.Add(TEXT("TimeRemaining"), FString::Printf(TEXT("%.3f"), Effect.GetTimeRemaining(GetWorldTime())));
+		}
+
+		if (Effect.PredictionKey.IsValidKey())
+		{
+			ActiveGELog.Add(TEXT("PredictionKey"), Effect.PredictionKey.ToString());
+		}
+
+		if (Effect.IsPendingRemove)
+		{
+			ActiveGELog.Add(TEXT("IsPendingRemove"), TEXT("true"));
+		}
+
+		if (Effect.bIsInhibited)
+		{
+			ActiveGELog.Add(TEXT("Inhibited"), TEXT("true"));
 		}
 
-		ActiveEffectsCategory.Add(TEXT("Total Modification"), FString::Printf(TEXT("%f"), CombinedModifierValue));
+		ActiveEffectsCategory.AddChild(ActiveGELog);
 	}
 
 	Snapshot->Status.Add(ActiveEffectsCategory);
@@ -5737,7 +5920,7 @@ void FActiveGameplayEffectsContainer::DecrementLock()
 
 			if (Effect.IsPendingRemove)
 			{
-				UE_VLOG_UELOG(Owner->GetOwnerActor(), LogGameplayEffects, Verbose, TEXT("Finish PendingRemove: %s. Auth: %d"), *Effect.GetDebugString(), IsNetAuthority());
+				UE_LOG(LogGameplayEffects, Verbose, TEXT("%s: Finish PendingRemove: %s. Auth: %d"), *GetNameSafe(Owner->GetOwnerActor()), *Effect.GetDebugString(), IsNetAuthority());
 				GameplayEffects_Internal.RemoveAtSwap(idx, 1, false);
 				ModifiedArray = true;
 				PendingRemoves--;
diff --git a/Engine/Plugins/Runtime/GameplayAbilities/Source/GameplayAbilities/Private/GameplayEffectTypes.cpp b/Engine/Plugins/Runtime/GameplayAbilities/Source/GameplayAbilities/Private/GameplayEffectTypes.cpp
--- a/Engine/Plugins/Runtime/GameplayAbilities/Source/GameplayAbilities/Private/GameplayEffectTypes.cpp
+++ b/Engine/Plugins/Runtime/GameplayAbilities/Source/GameplayAbilities/Private/GameplayEffectTypes.cpp
@@ -950,7 +950,7 @@ void FGameplayEffectSpec::PrintAll() const
 
 FString FGameplayEffectSpec::ToSimpleString() const
 {
-	return FString::Printf(TEXT("%s"), *GetNameSafe(Def));
+	return GetNameSafe(Def);
 }
 
 const FGameplayTagContainer* FTagContainerAggregator::GetAggregatedTags() const
diff --git a/Engine/Plugins/Runtime/GameplayAbilities/Source/GameplayAbilities/Public/AbilitySystemComponent.h b/Engine/Plugins/Runtime/GameplayAbilities/Source/GameplayAbilities/Public/AbilitySystemComponent.h
--- a/Engine/Plugins/Runtime/GameplayAbilities/Source/GameplayAbilities/Public/AbilitySystemComponent.h
+++ b/Engine/Plugins/Runtime/GameplayAbilities/Source/GameplayAbilities/Public/AbilitySystemComponent.h
@@ -1272,6 +1272,8 @@ class GAMEPLAYABILITIES_API UAbilitySystemComponent : public UGameplayTasksCompo
 
 #if ENABLE_VISUAL_LOG
 	void ClearDebugInstantEffects();
+	
+	virtual void GrabDebugSnapshot(FVisualLogEntry* Snapshot) const override;
 #endif // ENABLE_VISUAL_LOG
 
 	UE_DEPRECATED(4.26, "This will be made private in future engine versions. Use SetClientDebugStrings, GetClientDebugStrings, or GetClientDebugStrings_Mutable instead.")
diff --git a/Engine/Plugins/Runtime/GameplayAbilities/Source/GameplayAbilities/Public/AbilitySystemLog.h b/Engine/Plugins/Runtime/GameplayAbilities/Source/GameplayAbilities/Public/AbilitySystemLog.h
--- a/Engine/Plugins/Runtime/GameplayAbilities/Source/GameplayAbilities/Public/AbilitySystemLog.h
+++ b/Engine/Plugins/Runtime/GameplayAbilities/Source/GameplayAbilities/Public/AbilitySystemLog.h
@@ -29,6 +29,7 @@ GAMEPLAYABILITIES_API DECLARE_LOG_CATEGORY_EXTERN(LogGameplayEffects, Display, A
 { \
 	UE_LOG(LogAbilitySystem, Verbosity, Format, ##__VA_ARGS__); \
 	UE_VLOG(Actor, VLogAbilitySystem, Verbosity, Format, ##__VA_ARGS__); \
+	static_assert(false, "ABILITY_VLOG is deprecated because the same messages do not make sense in both contexts. Use both UE_LOG and UE_VLOG"); \
 }
 
 #else
@@ -42,6 +43,7 @@ GAMEPLAYABILITIES_API DECLARE_LOG_CATEGORY_EXTERN(LogGameplayEffects, Display, A
 { \
 	UE_LOG(LogAbilitySystem, Verbosity, Format, ##__VA_ARGS__); \
 	UE_VLOG(Actor, VLogAbilitySystem, Verbosity, Format, ##__VA_ARGS__); \
+	static_assert(false, "ABILITY_VLOG is deprecated because the same messages do not make sense in both contexts. Use both UE_LOG and UE_VLOG"); \
 }
 
 #endif //NO_LOGGING
diff --git a/Engine/Plugins/Runtime/GameplayAbilities/Source/GameplayAbilities/Public/GameplayEffect.h b/Engine/Plugins/Runtime/GameplayAbilities/Source/GameplayAbilities/Public/GameplayEffect.h
--- a/Engine/Plugins/Runtime/GameplayAbilities/Source/GameplayAbilities/Public/GameplayEffect.h
+++ b/Engine/Plugins/Runtime/GameplayAbilities/Source/GameplayAbilities/Public/GameplayEffect.h
@@ -20,6 +20,8 @@
 #include "GameplayAbilitySpec.h"
 #include "ActiveGameplayEffectIterator.h"
 #include "UObject/ObjectKey.h"
+#include "VisualLogger/VisualLoggerDebugSnapshotInterface.h"
+#include "VisualLogger/VisualLoggerTypes.h"
 #include "GameplayEffect.generated.h"
 
 /**
@@ -1120,6 +1122,10 @@ struct GAMEPLAYABILITIES_API FGameplayEffectSpec
 	/** Simple const accessor to the dynamic asset tags */
 	const FGameplayTagContainer& GetDynamicAssetTags() const;
 
+#if ENABLE_VISUAL_LOG
+	FVisualLogStatusCategory GrabVisLogStatus() const;
+#endif
+
 private:
 
 	void CaptureDataFromSource(bool bSkipRecaptureSourceActorTags = false);
@@ -1616,7 +1622,8 @@ struct GAMEPLAYABILITIES_API FActiveGameplayEffectsContainer : public FFastArray
 	/** Stores a record of gameplay effects that have executed and their results. Useful for debugging */
 	TArray<DebugExecutedGameplayEffectData> DebugExecutedGameplayEffects;
 
-	void GrabDebugSnapshot(FVisualLogEntry* Snapshot) const;
+	/** Report our current state to the VisLog */
+	void DescribeSelfToVisLog(FVisualLogEntry* Snapshot) const;
 #endif // ENABLE_VISUAL_LOG
 
 	void GetActiveGameplayEffectDataByAttribute(TMultiMap<FGameplayAttribute, FActiveGameplayEffectsContainer::DebugExecutedGameplayEffectData>& EffectMap) const;
