Fix tag stomping on clients, when a MinimalReplicationTagCountMap that requires non-owning net connections is hooked up to an owner...

* In our case, the tags replicated through the AbilitySystemComponent before the client setup it's own MinimalReplicationTagCountMap (and called SetOwner).
* The new CountMap stomped the tags in the AbilitySystemComponent, and tags would be incorrect until the AbilitySystemComponent replicated again later.
* The stomping CountMaps had SetRequireNonOwningNetConnection(true) which was being ignored.
* We can fix the issue by having CountMaps respect RequireNonOwningNetConnection, even when being configured.

Standalone Testing:
* With two players in a single server:
* Ensured on rejoin, local player pawns tags match between the client & server.
* Ensured that the new code only runs on the local player, not remote players.

PIE Testing:
* Ensured code stepped through as expected on begin play.
* Ensured tags are replicating correctly, on players and vehicles.

#rnx

diff --git a/Engine/Plugins/Runtime/GameplayAbilities/Source/GameplayAbilities/Private/GameplayEffectTypes.cpp b/Engine/Plugins/Runtime/GameplayAbilities/Source/GameplayAbilities/Private/GameplayEffectTypes.cpp
--- a/Engine/Plugins/Runtime/GameplayAbilities/Source/GameplayAbilities/Private/GameplayEffectTypes.cpp
+++ b/Engine/Plugins/Runtime/GameplayAbilities/Source/GameplayAbilities/Private/GameplayEffectTypes.cpp
@@ -9,6 +9,7 @@
 #include "Misc/ConfigCacheIni.h"
 #include "AbilitySystemGlobals.h"
 #include "AbilitySystemComponent.h"
+#include "Engine/NetConnection.h"
 #include "Engine/PackageMapClient.h"
 
 #if WITH_EDITOR
@@ -1363,28 +1364,12 @@ bool FMinimalReplicationTagCountMap::NetSerialize(FArchive& Ar, class UPackageMa
 			TagMap.FindOrAdd(Tag) = 1;
 		}
 
+		UPackageMapClient* PackageMap = CastChecked<UPackageMapClient>(Map);
+		LastConnection = PackageMap ? PackageMap->GetConnection() : nullptr;
+
 		if (Owner)
 		{
-			bool bUpdateOwnerTagMap = true;
-			if (bRequireNonOwningNetConnection)
-			{
-				if (AActor* OwningActor = Owner->GetOwner())
-				{
-					// Note we deliberately only want to do this if the NetConnection is not null
-					if (UNetConnection* OwnerNetConnection = OwningActor->GetNetConnection())
-					{
-						if (OwnerNetConnection == CastChecked<UPackageMapClient>(Map)->GetConnection())
-						{
-							bUpdateOwnerTagMap = false;
-						}
-					}
-				}
-			}
-
-			if (bUpdateOwnerTagMap)
-			{
-				UpdateOwnerTagMap();
-			}
+			UpdateOwnerTagMap();
 		}
 	}
 
@@ -1418,6 +1403,30 @@ void FMinimalReplicationTagCountMap::RemoveAllTags()
 
 void FMinimalReplicationTagCountMap::UpdateOwnerTagMap()
 {
+	bool bUpdateOwnerTagMap = true;
+	if (bRequireNonOwningNetConnection)
+	{
+		if (Owner)
+		{
+			if (AActor* OwningActor = Owner->GetOwner())
+			{
+				// Note we deliberately only want to do this if the NetConnection is not null
+				if (UNetConnection* OwnerNetConnection = OwningActor->GetNetConnection())
+				{
+					if (OwnerNetConnection == LastConnection.Get())
+					{
+						bUpdateOwnerTagMap = false;
+					}
+				}
+			}
+		}
+	}
+
+	if (!bUpdateOwnerTagMap)
+	{
+		return;
+	}
+
 	if (Owner)
 	{
 		for (auto It = TagMap.CreateIterator(); It; ++It)
diff --git a/Engine/Plugins/Runtime/GameplayAbilities/Source/GameplayAbilities/Private/Serialization/MinimalReplicationTagCountNetSerializer.cpp b/Engine/Plugins/Runtime/GameplayAbilities/Source/GameplayAbilities/Private/Serialization/MinimalReplicationTagCountNetSerializer.cpp
--- a/Engine/Plugins/Runtime/GameplayAbilities/Source/GameplayAbilities/Private/Serialization/MinimalReplicationTagCountNetSerializer.cpp
+++ b/Engine/Plugins/Runtime/GameplayAbilities/Source/GameplayAbilities/Private/Serialization/MinimalReplicationTagCountNetSerializer.cpp
@@ -249,7 +249,7 @@ void FMinimalReplicationTagCountMapNetSerializer::FNetSerializerRegistryDelegate
 
 void FMinimalReplicationTagCountMapNetSerializer::FNetSerializerRegistryDelegates::OnPostFreezeNetSerializerRegistry()
 {
-	constexpr SIZE_T ExpectedSizeOfFMinimalReplicationTagCountMap = 96;
+	constexpr SIZE_T ExpectedSizeOfFMinimalReplicationTagCountMap = 104;
 	constexpr SIZE_T ExpectedAlignOfFMinimalReplicationTagCountMap = 8;
 
 	constexpr SIZE_T ActualSizeOfFMinimalReplicationTagCountMap = sizeof(FMinimalReplicationTagCountMap);
diff --git a/Engine/Plugins/Runtime/GameplayAbilities/Source/GameplayAbilities/Public/GameplayEffectTypes.h b/Engine/Plugins/Runtime/GameplayAbilities/Source/GameplayAbilities/Public/GameplayEffectTypes.h
--- a/Engine/Plugins/Runtime/GameplayAbilities/Source/GameplayAbilities/Public/GameplayEffectTypes.h
+++ b/Engine/Plugins/Runtime/GameplayAbilities/Source/GameplayAbilities/Public/GameplayEffectTypes.h
@@ -27,6 +27,7 @@
 class Error;
 class UAbilitySystemComponent;
 class UGameplayAbility;
+class UNetConnection;
 struct FActiveGameplayEffect;
 struct FGameplayEffectModCallbackData;
 struct FGameplayEffectSpec;
@@ -1603,6 +1604,8 @@ private:
 	friend UE::Net::FMinimalReplicationTagCountMapReplicationFragment;
 
 	bool bRequireNonOwningNetConnection = false;
+	TWeakObjectPtr<UNetConnection> LastConnection;
+
 	void UpdateOwnerTagMap();
 };
 
