Fixing the string wrapping regression in the Abilities Gameplay Debugger Category

#rnx

diff --git a/Engine/Plugins/Runtime/GameplayAbilities/Source/GameplayAbilities/Private/GameplayDebuggerCategory_Abilities.cpp b/Engine/Plugins/Runtime/GameplayAbilities/Source/GameplayAbilities/Private/GameplayDebuggerCategory_Abilities.cpp
--- a/Engine/Plugins/Runtime/GameplayAbilities/Source/GameplayAbilities/Private/GameplayDebuggerCategory_Abilities.cpp
+++ b/Engine/Plugins/Runtime/GameplayAbilities/Source/GameplayAbilities/Private/GameplayDebuggerCategory_Abilities.cpp
@@ -229,6 +229,8 @@ void FGameplayDebuggerCategory_Abilities::DrawGameplayTags(FGameplayDebuggerCanv
 		// and ones where they disagree.
 		if (!AbilityComp->IsOwnerActorAuthoritative())
 		{
+			CanvasContext.Printf(TEXT("Owned Tags Legend:  {cyan}Both  {yellow}Server  {green}Local"));
+
 			static FGameplayTagContainer LocalOnlyTags;
 			AbilityComp->GetOwnedGameplayTags(LocalOnlyTags);
 
@@ -236,17 +238,14 @@ void FGameplayDebuggerCategory_Abilities::DrawGameplayTags(FGameplayDebuggerCanv
 			ServerOnlyTags.RemoveTags(MatchingTags);
 			LocalOnlyTags.RemoveTags(MatchingTags);
 
-			// Wrap the strings to the viewport
-			FString MatchingTagsStr, ServerOnlyTagsStr, LocalOnlyTagsStr;
-			WrapStringAccordingToViewport(MatchingTags.ToStringSimple(), MatchingTagsStr, CanvasContext, CanvasWidth);
-			WrapStringAccordingToViewport(ServerOnlyTags.ToStringSimple(), ServerOnlyTagsStr, CanvasContext, CanvasWidth);
-			WrapStringAccordingToViewport(LocalOnlyTags.ToStringSimple(), LocalOnlyTagsStr, CanvasContext, CanvasWidth);
-
-			MatchingTagsStr = MatchingTagsStr.Len() > 0 ? FString::Printf(TEXT("{cyan}%s  "), *MatchingTagsStr) : FString{};
-			ServerOnlyTagsStr = ServerOnlyTagsStr.Len() > 0 ? FString::Printf(TEXT("{yellow}%s  "), *ServerOnlyTagsStr) : FString{};
-			LocalOnlyTagsStr = LocalOnlyTagsStr.Len() > 0 ? FString::Printf(TEXT("{green}%s  "), *LocalOnlyTagsStr) : FString{};
-
-			CanvasContext.Printf(TEXT("Owned Tags Legend:  {cyan}Both  {yellow}Server  {green}Local \n%s%s%s"), *MatchingTagsStr, *ServerOnlyTagsStr, *LocalOnlyTagsStr);
+			// Build up the strings
+			FString MatchingTagsStr = MatchingTags.Num() > 0 ? FString::Printf(TEXT("{cyan}%s  "), *MatchingTags.ToStringSimple()) : FString{};
+			FString ServerOnlyTagsStr = ServerOnlyTags.Num() > 0 ? FString::Printf(TEXT("{yellow}%s  "), *ServerOnlyTags.ToStringSimple()) : FString{};
+			FString LocalOnlyTagsStr = LocalOnlyTags.Num() > 0 ? FString::Printf(TEXT("{green}%s  "), *LocalOnlyTags.ToStringSimple()) : FString{};
+			
+			FString WrappedDebugText;
+			WrapStringAccordingToViewport(FString::Printf(TEXT("%s%s%s"), *MatchingTagsStr, *ServerOnlyTagsStr, *LocalOnlyTagsStr), WrappedDebugText, CanvasContext, CanvasWidth);
+			CanvasContext.Print(WrappedDebugText);
 		}
 		else
 		{
