Added some guard rails to make it easier to implement native GameplayEffects properly, or at least catch incorrectly implemented ones.

- Fixed up the native FInheritedTagContainer Add/Remove functions.  These would cause unexpected results if one were to rely upon the Added/Removed tags to be filled properly (which the GameplayEffectComponents are relying upon).
- Added some IsDataValid for TargetTagsGameplayEffectComponent.
- Catch IsDataValid and missing GEComponents in native-only GameplayEffect classes.

diff --git a/Engine/Plugins/Runtime/GameplayAbilities/Source/GameplayAbilities/Private/GameplayEffect.cpp b/Engine/Plugins/Runtime/GameplayAbilities/Source/GameplayAbilities/Private/GameplayEffect.cpp
--- a/Engine/Plugins/Runtime/GameplayAbilities/Source/GameplayAbilities/Private/GameplayEffect.cpp
+++ b/Engine/Plugins/Runtime/GameplayAbilities/Source/GameplayAbilities/Private/GameplayEffect.cpp
@@ -171,6 +171,41 @@ void UGameplayEffect::PostInitProperties()
 			SetVersion(EGameplayEffectVersion::Monolithic);
 		}
 	}
+
+	// Let's cover some easy-to-overlook issues when implementing a Gameplay Effect as a native class
+	// First, ensure we've actually caught all of the native GEComponents
+	if (!IsInBlueprint())
+	{
+		TArray<UObject*> PossibleComponents;
+		GetDefaultSubobjects(PossibleComponents);
+		for (UObject* Obj : PossibleComponents)
+		{
+			if (UGameplayEffectComponent* GEComponent = Cast<UGameplayEffectComponent>(Obj))
+			{
+				if (!ensureMsgf(GEComponents.Contains(GEComponent), TEXT("%s: %s should be added to GEComponents during the constructor or in PostInitProperties"), *GetName(), *GEComponent->GetName()))
+				{
+					GEComponents.Add(GEComponent);
+				}
+			}
+		}
+
+		// Catch any GEComponent configuration issues.  We usually want to do this late (after deserialization) for assets in PostLoad, but
+		// we've already determined this is a native class, so let's do that now (we won't get PostLoad).
+		FDataValidationContext DataValidationContext;
+		IsDataValid(DataValidationContext);
+
+		// Now spit out any of those issues in the log
+		TArray<FText> Warnings, Errors;
+		DataValidationContext.SplitIssues(Warnings, Errors);
+		for (const FText& WarningText : Warnings)
+		{
+			UE_LOG(LogGameplayEffects, Warning, TEXT("%s: %s"), *GetName(), *WarningText.ToString());
+		}
+		for (const FText& ErrorText : Errors)
+		{
+			UE_LOG(LogGameplayEffects, Error, TEXT("%s: %s"), *GetName(), *ErrorText.ToString());
+		}
+	}
 #endif
 }
 
@@ -210,7 +245,7 @@ EDataValidationResult UGameplayEffect::IsDataValid(FDataValidationContext& Conte
 
 	if (ValidationResult != EDataValidationResult::Invalid)
 	{
-		for (UGameplayEffectComponent* GEComponent : GEComponents)
+		for (const UGameplayEffectComponent* GEComponent : GEComponents)
 		{
 			if (GEComponent)
 			{
@@ -5858,11 +5893,15 @@ void FInheritedTagContainer::ApplyTo(FGameplayTagContainer& ApplyToContainer) co
 
 void FInheritedTagContainer::AddTag(const FGameplayTag& TagToAdd)
 {
+	Removed.RemoveTag(TagToAdd);
+	Added.AddTag(TagToAdd);
 	CombinedTags.AddTag(TagToAdd);
 }
 
 void FInheritedTagContainer::RemoveTag(const FGameplayTag& TagToRemove)
 {
+	Added.RemoveTag(TagToRemove);
+	Removed.AddTag(TagToRemove);
 	CombinedTags.RemoveTag(TagToRemove);
 }
 
diff --git a/Engine/Plugins/Runtime/GameplayAbilities/Source/GameplayAbilities/Private/GameplayEffectComponents/TargetTagsGameplayEffectComponent.cpp b/Engine/Plugins/Runtime/GameplayAbilities/Source/GameplayAbilities/Private/GameplayEffectComponents/TargetTagsGameplayEffectComponent.cpp
--- a/Engine/Plugins/Runtime/GameplayAbilities/Source/GameplayAbilities/Private/GameplayEffectComponents/TargetTagsGameplayEffectComponent.cpp
+++ b/Engine/Plugins/Runtime/GameplayAbilities/Source/GameplayAbilities/Private/GameplayEffectComponents/TargetTagsGameplayEffectComponent.cpp
@@ -1,6 +1,9 @@
 // Copyright Epic Games, Inc. All Rights Reserved.
 
 #include "GameplayEffectComponents/TargetTagsGameplayEffectComponent.h"
+#include "Misc/DataValidation.h"
+
+#define LOCTEXT_NAMESPACE "TargetTagsGameplayEffectComponent"
 
 void UTargetTagsGameplayEffectComponent::PostInitProperties()
 {
@@ -37,6 +40,21 @@ void UTargetTagsGameplayEffectComponent::PostEditChangeProperty(FPropertyChanged
 		Owner->OnGameplayEffectChanged();
 	}
 }
+
+EDataValidationResult UTargetTagsGameplayEffectComponent::IsDataValid(FDataValidationContext& Context) const
+{
+	EDataValidationResult Result = Super::IsDataValid(Context);
+
+	const bool bInstantEffect = (GetOwner()->DurationPolicy == EGameplayEffectDurationType::Instant);
+	if (bInstantEffect && !InheritableGrantedTagsContainer.CombinedTags.IsEmpty())
+	{
+		Context.AddError(FText::FormatOrdered(LOCTEXT("GEInstantAndTargetTags", "GE {0} is set to Instant so TargetTagsGameplayEffectComponent will not be able to apply its tags."), FText::FromString(GetNameSafe(GetOwner()))));
+		Result = EDataValidationResult::Invalid;
+	}
+
+	return Result;
+}
+
 #endif // WITH_EDITOR
 
 void UTargetTagsGameplayEffectComponent::SetAndApplyTargetTagChanges(const FInheritedTagContainer& TagContainerMods)
@@ -56,3 +74,5 @@ void UTargetTagsGameplayEffectComponent::ApplyTargetTagChanges() const
 	UGameplayEffect* Owner = GetOwner();
 	InheritableGrantedTagsContainer.ApplyTo(Owner->CachedGrantedTags);
 }
+
+#undef LOCTEXT_NAMESPACE
\ No newline at end of file
diff --git a/Engine/Plugins/Runtime/GameplayAbilities/Source/GameplayAbilities/Public/GameplayEffect.h b/Engine/Plugins/Runtime/GameplayAbilities/Source/GameplayAbilities/Public/GameplayEffect.h
--- a/Engine/Plugins/Runtime/GameplayAbilities/Source/GameplayAbilities/Public/GameplayEffect.h
+++ b/Engine/Plugins/Runtime/GameplayAbilities/Source/GameplayAbilities/Public/GameplayEffect.h
@@ -614,7 +614,7 @@ struct GAMEPLAYABILITIES_API FInheritedTagContainer
 {
 	GENERATED_USTRUCT_BODY()
 
-	/** Tags that I inherited and tags that I added minus tags that I removed*/
+	/** Tags that I inherited and tags that I added minus tags that I removed */
 	UPROPERTY(VisibleAnywhere, BlueprintReadOnly, Category = Application)
 	FGameplayTagContainer CombinedTags;
 
@@ -631,7 +631,10 @@ struct GAMEPLAYABILITIES_API FInheritedTagContainer
 	/** Apply the Added and Removed tags to the passed-in container (does not have to be the previously configured Parent!) */
 	void ApplyTo(FGameplayTagContainer& ApplyToContainer) const;
 
+	/** Add a tag that will appear in addition to any inherited tags */
 	void AddTag(const FGameplayTag& TagToAdd);
+
+	/** Remove a tag that will be omitted from any inherited tags */
 	void RemoveTag(const FGameplayTag& TagToRemove);
 
 	bool operator==(const FInheritedTagContainer& Other) const;
diff --git a/Engine/Plugins/Runtime/GameplayAbilities/Source/GameplayAbilities/Public/GameplayEffectComponents/TargetTagsGameplayEffectComponent.h b/Engine/Plugins/Runtime/GameplayAbilities/Source/GameplayAbilities/Public/GameplayEffectComponents/TargetTagsGameplayEffectComponent.h
--- a/Engine/Plugins/Runtime/GameplayAbilities/Source/GameplayAbilities/Public/GameplayEffectComponents/TargetTagsGameplayEffectComponent.h
+++ b/Engine/Plugins/Runtime/GameplayAbilities/Source/GameplayAbilities/Public/GameplayEffectComponents/TargetTagsGameplayEffectComponent.h
@@ -29,6 +29,9 @@ public:
 	/** Needed to properly update FInheritedTagContainer properties */
 	virtual void PostEditChangeProperty(FPropertyChangedEvent& PropertyChangedEvent) override;
 
+	/** Validate incompatible configurations */
+	virtual EDataValidationResult IsDataValid(class FDataValidationContext& Context) const override;
+
 private:
 	/** Get a cached version of the FProperty Name for PostEditChangeProperty */
 	static const FName& GetInheritableGrantedTagsContainerName()
