Instant and Periodic GameplayEffects should respect Modifier tag requirements

diff --git a/Engine/Plugins/Runtime/GameplayAbilities/Source/GameplayAbilities/Private/GameplayEffect.cpp b/Engine/Plugins/Runtime/GameplayAbilities/Source/GameplayAbilities/Private/GameplayEffect.cpp
--- a/Engine/Plugins/Runtime/GameplayAbilities/Source/GameplayAbilities/Private/GameplayEffect.cpp
+++ b/Engine/Plugins/Runtime/GameplayAbilities/Source/GameplayAbilities/Private/GameplayEffect.cpp
@@ -84,6 +84,10 @@ namespace UE::GameplayEffect
 
 	TAutoConsoleVariable<int32> CVarGameplayEffectMaxVersion(TEXT("AbilitySystem.GameplayEffects.MaxVersion"), (int32)EGameplayEffectVersion::Current, TEXT("Override the Gameplay Effect Current Version (disabling upgrade code paths)"), ECVF_Default);
 
+	// Fix introduced in UE5.4
+	bool bUseModifierTagRequirementsOnAllGameplayEffects = true;
+	FAutoConsoleVariableRef CVarUseModTagReqsOnAllGE{ TEXT("AbilitySystem.Fix.UseModTagReqsOnAllGE"), bUseModifierTagRequirementsOnAllGameplayEffects, TEXT("Fix an issue where MustHave/MustNotHave tags did not apply to Instant and Periodic Gameplay Effects"), ECVF_Default };
+
 #if WITH_EDITOR
 	namespace EditorOnly
 	{
@@ -2821,10 +2825,26 @@ void FActiveGameplayEffectsContainer::PredictivelyExecuteEffectSpec(FGameplayEff
 
 	bool ModifierSuccessfullyExecuted = false;
 
+	ensureMsgf(SpecToUse.Modifiers.Num() == SpecToUse.Def->Modifiers.Num(), TEXT("GE Spec %s modifiers did not match the Definition.  This indicates an error much earlier (when setting up the GESpec)"), *SpecToUse.ToSimpleString());
 	for (int32 ModIdx = 0; ModIdx < SpecToUse.Modifiers.Num(); ++ModIdx)
 	{
 		const FGameplayModifierInfo& ModDef = SpecToUse.Def->Modifiers[ModIdx];
 
+		if (UE::GameplayEffect::bUseModifierTagRequirementsOnAllGameplayEffects)
+		{
+			// Check tag requirements. This code path is for Instant & Periodic effects.
+			// Duration effects are a separate code path; they use aggregators and their requirements checks are in FAggregatorMod::UpdateQualifies.
+			if (!ModDef.SourceTags.IsEmpty() && !ModDef.SourceTags.RequirementsMet(Spec.CapturedSourceTags.GetActorTags()))
+			{
+				continue;
+			}
+
+			if (!ModDef.TargetTags.IsEmpty() && !ModDef.TargetTags.RequirementsMet(Spec.CapturedTargetTags.GetActorTags()))
+			{
+				continue;
+			}
+		}
+
 		FGameplayModifierEvaluatedData EvalData(ModDef.Attribute, ModDef.ModifierOp, SpecToUse.GetModifierMagnitude(ModIdx, true));
 		ModifierSuccessfullyExecuted |= InternalExecuteMod(SpecToUse, EvalData);
 	}
@@ -2935,10 +2955,26 @@ void FActiveGameplayEffectsContainer::ExecuteActiveEffectsFrom(FGameplayEffectSp
 	// ------------------------------------------------------
 	
 	bool ModifierSuccessfullyExecuted = false;
-
+	
+	ensureMsgf(SpecToUse.Modifiers.Num() == SpecToUse.Def->Modifiers.Num(), TEXT("GE Spec %s modifiers did not match the Definition.  This indicates an error much earlier (when setting up the GESpec)"), *SpecToUse.ToSimpleString());
 	for (int32 ModIdx = 0; ModIdx < SpecToUse.Modifiers.Num(); ++ModIdx)
 	{
 		const FGameplayModifierInfo& ModDef = SpecToUse.Def->Modifiers[ModIdx];
+
+		if (UE::GameplayEffect::bUseModifierTagRequirementsOnAllGameplayEffects)
+		{
+			// Check tag requirements. This code path is for Instant & Periodic effects.
+			// Duration effects are a separate code path; they use aggregators and their requirements checks are in FAggregatorMod::UpdateQualifies.
+			if (!ModDef.SourceTags.IsEmpty() && !ModDef.SourceTags.RequirementsMet(Spec.CapturedSourceTags.GetActorTags()))
+			{
+				continue;
+			}
+
+			if (!ModDef.TargetTags.IsEmpty() && !ModDef.TargetTags.RequirementsMet(Spec.CapturedTargetTags.GetActorTags()))
+			{
+				continue;
+			}
+		}
 		
 		FGameplayModifierEvaluatedData EvalData(ModDef.Attribute, ModDef.ModifierOp, SpecToUse.GetModifierMagnitude(ModIdx, true));
 		ModifierSuccessfullyExecuted |= InternalExecuteMod(SpecToUse, EvalData);
