Avoid invalid pointer accesses by adding null checks

#localization none
#tests none/trivial fix

diff --git a/Engine/Plugins/Runtime/GameplayAbilities/Source/GameplayAbilities/Private/AbilitySystemComponent_Abilities.cpp b/Engine/Plugins/Runtime/GameplayAbilities/Source/GameplayAbilities/Private/AbilitySystemComponent_Abilities.cpp
--- a/Engine/Plugins/Runtime/GameplayAbilities/Source/GameplayAbilities/Private/AbilitySystemComponent_Abilities.cpp
+++ b/Engine/Plugins/Runtime/GameplayAbilities/Source/GameplayAbilities/Private/AbilitySystemComponent_Abilities.cpp
@@ -3419,7 +3419,7 @@ bool UAbilitySystemComponent::ServerCurrentMontageJumpToSectionName_Validate(UAn
 void UAbilitySystemComponent::ServerCurrentMontageJumpToSectionName_Implementation(UAnimSequenceBase* ClientAnimation, FName SectionName)
 {
 	UAnimInstance* AnimInstance = AbilityActorInfo.IsValid() ? AbilityActorInfo->GetAnimInstance() : nullptr;
-	if (AnimInstance)
+	if (AnimInstance && LocalAnimMontageInfo.AnimMontage)
 	{
 		UAnimSequenceBase* CurrentAnimation = LocalAnimMontageInfo.AnimMontage->IsDynamicMontage() ? LocalAnimMontageInfo.AnimMontage->GetFirstAnimReference() : LocalAnimMontageInfo.AnimMontage;
 		if (ClientAnimation == CurrentAnimation)
@@ -3459,7 +3459,7 @@ bool UAbilitySystemComponent::ServerCurrentMontageSetPlayRate_Validate(UAnimSequ
 void UAbilitySystemComponent::ServerCurrentMontageSetPlayRate_Implementation(UAnimSequenceBase* ClientAnimation, float InPlayRate)
 {
 	UAnimInstance* AnimInstance = AbilityActorInfo.IsValid() ? AbilityActorInfo->GetAnimInstance() : nullptr;
-	if (AnimInstance)
+	if (AnimInstance && LocalAnimMontageInfo.AnimMontage)
 	{
 		UAnimSequenceBase* CurrentAnimation = LocalAnimMontageInfo.AnimMontage->IsDynamicMontage() ? LocalAnimMontageInfo.AnimMontage->GetFirstAnimReference() : LocalAnimMontageInfo.AnimMontage;
 		if (ClientAnimation == CurrentAnimation)
