Add gameplay cue suppression to GameplayCueManager

diff --git a/Engine/Plugins/Runtime/GameplayAbilities/Source/GameplayAbilities/Private/GameplayCueManager.cpp b/Engine/Plugins/Runtime/GameplayAbilities/Source/GameplayAbilities/Private/GameplayCueManager.cpp
--- a/Engine/Plugins/Runtime/GameplayAbilities/Source/GameplayAbilities/Private/GameplayCueManager.cpp
+++ b/Engine/Plugins/Runtime/GameplayAbilities/Source/GameplayAbilities/Private/GameplayCueManager.cpp
@@ -50,6 +50,9 @@ static FAutoConsoleVariableRef CVarDurationeGameplayCues(TEXT("AbilitySystem.Gam
 int32 GameplayCueRunOnDedicatedServer = 0;
 static FAutoConsoleVariableRef CVarDedicatedServerGameplayCues(TEXT("AbilitySystem.GameplayCue.RunOnDedicatedServer"), GameplayCueRunOnDedicatedServer, TEXT("Run gameplay cue events on dedicated server"), ECVF_Default );
 
+bool EnableSuppressCuesOnGameplayCueManager = true;
+static FAutoConsoleVariableRef CVarEnableSuppressCuesOnGameplayCueManager(TEXT("AbilitySystem.GameplayCue.EnableSuppressCuesOnGameplayCueManager"), EnableSuppressCuesOnGameplayCueManager, TEXT("Allows the GameplayCueManager to suppress cues when the bSuppressGameplayCues is set on the target AbilitySystemComponent"), ECVF_Default );
+
 #if WITH_EDITOR
 USceneComponent* UGameplayCueManager::PreviewComponent = nullptr;
 UWorld* UGameplayCueManager::PreviewWorld = nullptr;
@@ -1245,6 +1248,11 @@ void UGameplayCueManager::InvokeGameplayCueAddedAndWhileActive_FromSpec(UAbility
 		return;
 	}
 
+	if (EnableSuppressCuesOnGameplayCueManager && OwningComponent && OwningComponent->bSuppressGameplayCues)
+	{
+		return;
+	}
+
 	IAbilitySystemReplicationProxyInterface* ReplicationInterface = OwningComponent->GetReplicationInterface();
 	if (ReplicationInterface == nullptr)
 	{
@@ -1295,6 +1303,11 @@ void UGameplayCueManager::InvokeGameplayCueExecuted_FromSpec(UAbilitySystemCompo
 		return;
 	}
 
+	if (EnableSuppressCuesOnGameplayCueManager && OwningComponent && OwningComponent->bSuppressGameplayCues)
+	{
+		return;
+	}
+
 	FGameplayCuePendingExecute PendingCue;
 
 	if (AbilitySystemAlwaysConvertGESpecToGCParams)
@@ -1327,6 +1340,11 @@ void UGameplayCueManager::InvokeGameplayCueExecuted_FromSpec(UAbilitySystemCompo
 
 void UGameplayCueManager::InvokeGameplayCueExecuted(UAbilitySystemComponent* OwningComponent, const FGameplayTag GameplayCueTag, FPredictionKey PredictionKey, FGameplayEffectContextHandle EffectContext)
 {
+	if (EnableSuppressCuesOnGameplayCueManager && OwningComponent && OwningComponent->bSuppressGameplayCues)
+	{
+		return;
+	}
+
 	if (OwningComponent)
 	{
 		FGameplayCuePendingExecute PendingCue;
@@ -1342,6 +1360,11 @@ void UGameplayCueManager::InvokeGameplayCueExecuted(UAbilitySystemComponent* Own
 
 void UGameplayCueManager::InvokeGameplayCueExecuted_WithParams(UAbilitySystemComponent* OwningComponent, const FGameplayTag GameplayCueTag, FPredictionKey PredictionKey, FGameplayCueParameters GameplayCueParameters)
 {
+	if (EnableSuppressCuesOnGameplayCueManager && OwningComponent && OwningComponent->bSuppressGameplayCues)
+	{
+		return;
+	}
+
 	if (OwningComponent)
 	{
 		FGameplayCuePendingExecute PendingCue;
