PR #11155: AbilitySystem.GameplayCueNotifyTagCheckOnRemove typo fix

#11155

diff --git a/Engine/Plugins/Runtime/GameplayAbilities/Source/GameplayAbilities/Private/GameplayCueNotify_Actor.cpp b/Engine/Plugins/Runtime/GameplayAbilities/Source/GameplayAbilities/Private/GameplayCueNotify_Actor.cpp
--- a/Engine/Plugins/Runtime/GameplayAbilities/Source/GameplayAbilities/Private/GameplayCueNotify_Actor.cpp
+++ b/Engine/Plugins/Runtime/GameplayAbilities/Source/GameplayAbilities/Private/GameplayCueNotify_Actor.cpp
@@ -181,7 +181,7 @@ void AGameplayCueNotify_Actor::K2_EndGameplayCue()
 }
 
 int32 GameplayCueNotifyTagCheckOnRemove = 1;
-static FAutoConsoleVariableRef CVarGameplayCueNotifyActorStacking(TEXT("AbilitySystem.GameplayCueNotifyTagCheckOnRemove"), GameplayCueNotifyTagCheckOnRemove, TEXT("Check that target no longer has tag when removing GamepalyCues"), ECVF_Default );
+static FAutoConsoleVariableRef CVarGameplayCueNotifyActorStacking(TEXT("AbilitySystem.GameplayCueNotifyTagCheckOnRemove"), GameplayCueNotifyTagCheckOnRemove, TEXT("Check that target no longer has tag when removing GameplayCues"), ECVF_Default );
 
 void AGameplayCueNotify_Actor::HandleGameplayCue(AActor* MyTarget, EGameplayCueEvent::Type EventType, const FGameplayCueParameters& Parameters)
 {
