Fix GameFeatureAction_AddAttributeDefaults not cleaning up references to objects when unregistered

#tests unloaded a plugin with debugger attached in editor

diff --git a/Engine/Plugins/Experimental/AbilitySystemGameFeatureActions/Source/AbilitySystemGameFeatureActions/Private/GameFeatureAction_AddAttributeDefaults.cpp b/Engine/Plugins/Experimental/AbilitySystemGameFeatureActions/Source/AbilitySystemGameFeatureActions/Private/GameFeatureAction_AddAttributeDefaults.cpp
--- a/Engine/Plugins/Experimental/AbilitySystemGameFeatureActions/Source/AbilitySystemGameFeatureActions/Private/GameFeatureAction_AddAttributeDefaults.cpp
+++ b/Engine/Plugins/Experimental/AbilitySystemGameFeatureActions/Source/AbilitySystemGameFeatureActions/Private/GameFeatureAction_AddAttributeDefaults.cpp
@@ -2,12 +2,18 @@
 
 #include "GameFeatureAction_AddAttributeDefaults.h"
 #include "AbilitySystemGlobals.h"
+#include "HAL/IConsoleManager.h"
 #include "Misc/PackageName.h"
 
 #include UE_INLINE_GENERATED_CPP_BY_NAME(GameFeatureAction_AddAttributeDefaults)
 
 #define LOCTEXT_NAMESPACE "GameFeatures"
 
+namespace GameFeatureAction_AddAttributeDefaults
+{
+	static TAutoConsoleVariable<bool> CVarAllowRemoveAttributeDefaultTables(TEXT("GameFeatureAction_AddAttributeDefaults.AllowRemoveAttributeDefaultTables"), true, TEXT("Removes hard references when unregistering"));
+}
+
 //////////////////////////////////////////////////////////////////////
 // UGameFeatureAction_AddAttributeDefaults
 
@@ -28,8 +34,24 @@ void UGameFeatureAction_AddAttributeDefaults::OnGameFeatureRegistering()
 	}
 #endif // WITH_EDITOR
 
-	UAbilitySystemGlobals& AbilitySystemGlobals = UAbilitySystemGlobals::Get();
-	AbilitySystemGlobals.AddAttributeDefaultTables(*AttribDefaultTableNamesToAdd);
+	if (!AttribDefaultTableNamesToAdd->IsEmpty())
+	{
+		FNameBuilder OwnerNameBuilder;
+		GetPathName(nullptr, OwnerNameBuilder);
+		AttributeDefaultTablesOwnerName = FName(OwnerNameBuilder.ToView());
+
+		UAbilitySystemGlobals& AbilitySystemGlobals = UAbilitySystemGlobals::Get();
+		AbilitySystemGlobals.AddAttributeDefaultTables(AttributeDefaultTablesOwnerName, *AttribDefaultTableNamesToAdd);
+	}
+}
+
+void UGameFeatureAction_AddAttributeDefaults::OnGameFeatureUnregistering()
+{
+	if (!AttribDefaultTableNames.IsEmpty() && GameFeatureAction_AddAttributeDefaults::CVarAllowRemoveAttributeDefaultTables.GetValueOnAnyThread())
+	{
+		UAbilitySystemGlobals& AbilitySystemGlobals = UAbilitySystemGlobals::Get();
+		AbilitySystemGlobals.RemoveAttributeDefaultTables(AttributeDefaultTablesOwnerName, AttribDefaultTableNames);
+	}
 }
 
 //////////////////////////////////////////////////////////////////////
diff --git a/Engine/Plugins/Experimental/AbilitySystemGameFeatureActions/Source/AbilitySystemGameFeatureActions/Public/GameFeatureAction_AddAttributeDefaults.h b/Engine/Plugins/Experimental/AbilitySystemGameFeatureActions/Source/AbilitySystemGameFeatureActions/Public/GameFeatureAction_AddAttributeDefaults.h
--- a/Engine/Plugins/Experimental/AbilitySystemGameFeatureActions/Source/AbilitySystemGameFeatureActions/Public/GameFeatureAction_AddAttributeDefaults.h
+++ b/Engine/Plugins/Experimental/AbilitySystemGameFeatureActions/Source/AbilitySystemGameFeatureActions/Public/GameFeatureAction_AddAttributeDefaults.h
@@ -20,11 +20,15 @@ class UGameFeatureAction_AddAttributeDefaults final : public UGameFeatureAction
 public:
 	//~UGameFeatureAction interface
 	virtual void OnGameFeatureRegistering() override;
+	virtual void OnGameFeatureUnregistering() override;
 	//~End of UGameFeatureAction interface
 
 	/** List of attribute default tables to add */
 	UPROPERTY(EditAnywhere, Category = Attributes)
 	TArray<FSoftObjectPath> AttribDefaultTableNames;
+
+private:
+	FName AttributeDefaultTablesOwnerName;
 };
 
 #if UE_ENABLE_INCLUDE_ORDER_DEPRECATED_IN_5_2
diff --git a/Engine/Plugins/Runtime/GameplayAbilities/Source/GameplayAbilities/Private/AbilitySystemGlobals.cpp b/Engine/Plugins/Runtime/GameplayAbilities/Source/GameplayAbilities/Private/AbilitySystemGlobals.cpp
--- a/Engine/Plugins/Runtime/GameplayAbilities/Source/GameplayAbilities/Private/AbilitySystemGlobals.cpp
+++ b/Engine/Plugins/Runtime/GameplayAbilities/Source/GameplayAbilities/Private/AbilitySystemGlobals.cpp
@@ -362,16 +362,64 @@ FAttributeSetInitter* UAbilitySystemGlobals::GetAttributeSetInitter() const
 	return GlobalAttributeSetInitter.Get();
 }
 
-void UAbilitySystemGlobals::AddAttributeDefaultTables(const TArray<FSoftObjectPath>& AttribDefaultTableNames)
+void UAbilitySystemGlobals::AddAttributeDefaultTables(const FName OwnerName, const TArray<FSoftObjectPath>& AttribDefaultTableNames)
 {
 	for (const FSoftObjectPath& TableName : AttribDefaultTableNames)
 	{
-		GlobalAttributeSetDefaultsTableNames.AddUnique(TableName);
+		if (TArray<FName>* Found = GlobalAttributeSetDefaultsTableNamesWithOwners.Find(TableName))
+		{
+			Found->Add(OwnerName);
+		}
+		else
+		{
+			TArray<FName> Owners = { OwnerName };
+			GlobalAttributeSetDefaultsTableNamesWithOwners.Add(TableName, MoveTemp(Owners));
+		}
 	}
 
 	InitAttributeDefaults();
 }
 
+void UAbilitySystemGlobals::RemoveAttributeDefaultTables(const FName OwnerName, const TArray<FSoftObjectPath>& AttribDefaultTableNames)
+{
+	bool bModified = false;
+
+	for (const FSoftObjectPath& TableName : AttribDefaultTableNames)
+	{
+		if (TableName.IsValid())
+		{
+			if (TArray<FName>* Found = GlobalAttributeSetDefaultsTableNamesWithOwners.Find(TableName))
+			{
+				Found->RemoveSingle(OwnerName);
+
+				// If no references remain, clear the pointer in GlobalAttributeDefaultsTables to allow GC
+				if (Found->IsEmpty())
+				{
+					GlobalAttributeSetDefaultsTableNamesWithOwners.Remove(TableName);
+
+					// Only if not listed in config file
+					if (!GlobalAttributeSetDefaultsTableNames.Contains(TableName))
+					{
+						// Remove reference to allow GC so package can be unloaded
+						if (UCurveTable* AttribTable = Cast<UCurveTable>(TableName.ResolveObject()))
+						{
+							if (GlobalAttributeDefaultsTables.Remove(AttribTable) > 0)
+							{
+								bModified = true;
+							}
+						}
+					}
+				}
+			}
+		}
+	}
+
+	if (bModified)
+	{
+		ReloadAttributeDefaults();
+	}
+}
+
 void UAbilitySystemGlobals::InitAttributeDefaults()
 {
  	bool bLoadedAnyDefaults = false;
@@ -401,6 +449,21 @@ void UAbilitySystemGlobals::InitAttributeDefaults()
 		}
  	}
 	
+	// Handle global curve tables for attribute defaults not defined by config file (registered by plugins or other systems calling AddAttributeDefaultTables)
+	for (const TPair<FSoftObjectPath, TArray<FName>>& It : GlobalAttributeSetDefaultsTableNamesWithOwners)
+	{
+		const FSoftObjectPath& AttribDefaultTableName = It.Key;
+		if (AttribDefaultTableName.IsValid() && !GlobalAttributeSetDefaultsTableNames.Contains(AttribDefaultTableName))
+		{
+			UCurveTable* AttribTable = Cast<UCurveTable>(AttribDefaultTableName.TryLoad());
+			if (AttribTable)
+			{
+				GlobalAttributeDefaultsTables.AddUnique(AttribTable);
+				bLoadedAnyDefaults = true;
+			}
+		}
+	}
+
 	if (bLoadedAnyDefaults)
 	{
 		// Subscribe for reimports if in the editor
diff --git a/Engine/Plugins/Runtime/GameplayAbilities/Source/GameplayAbilities/Public/AbilitySystemGlobals.h b/Engine/Plugins/Runtime/GameplayAbilities/Source/GameplayAbilities/Public/AbilitySystemGlobals.h
--- a/Engine/Plugins/Runtime/GameplayAbilities/Source/GameplayAbilities/Public/AbilitySystemGlobals.h
+++ b/Engine/Plugins/Runtime/GameplayAbilities/Source/GameplayAbilities/Public/AbilitySystemGlobals.h
@@ -308,7 +308,8 @@ class GAMEPLAYABILITIES_API UAbilitySystemGlobals : public UObject
 	UPROPERTY()
 	FNetSerializeScriptStructCache	EffectContextStructCache;
 
-	void AddAttributeDefaultTables(const TArray<FSoftObjectPath>& AttribDefaultTableNames);
+	void AddAttributeDefaultTables(const FName OwnerName, const TArray<FSoftObjectPath>& AttribDefaultTableNames);
+	void RemoveAttributeDefaultTables(const FName OwnerName, const TArray<FSoftObjectPath>& AttribDefaultTableNames);
 
 protected:
 
@@ -406,6 +407,12 @@ protected:
 	/** Used to initialize attribute sets */
 	TSharedPtr<FAttributeSetInitter> GlobalAttributeSetInitter;
 
+	/** 
+	 * Curve table names to use for default values for attribute sets, keyed off of Name/Levels (with owners to allow removal of hard reference by GlobalAttributeDefaultsTables)
+	 * Required to allow unloading of plugins
+	*/
+	TMap<FSoftObjectPath, TArray<FName>> GlobalAttributeSetDefaultsTableNamesWithOwners;
+
 	template <class T>
 	T* InternalGetLoadTable(T*& Table, FString TableName);
 
