Gameplay Cue Notify Actors are now removed properly on Listen Server instances.

Previously when Gameplay Effects were removed, the Gameplay Cue Tags would still be present on the listen server.  AbilitySystem.GameplayCueNotifyTagCheckOnRemove would then refuse to issue the Gameplay Cue Removal event since that Tag was still present.  The removal code has been rewritten to account for this.

diff --git a/Engine/Plugins/Runtime/GameplayAbilities/Source/GameplayAbilities/Private/AbilitySystemComponent.cpp b/Engine/Plugins/Runtime/GameplayAbilities/Source/GameplayAbilities/Private/AbilitySystemComponent.cpp
--- a/Engine/Plugins/Runtime/GameplayAbilities/Source/GameplayAbilities/Private/AbilitySystemComponent.cpp
+++ b/Engine/Plugins/Runtime/GameplayAbilities/Source/GameplayAbilities/Private/AbilitySystemComponent.cpp
@@ -1378,26 +1378,34 @@ void UAbilitySystemComponent::RemoveGameplayCue_Internal(const FGameplayTag Game
 {
 	if (IsOwnerActorAuthoritative())
 	{
-		const bool bWasInList = GameplayCueContainer.HasCue(GameplayCueTag);
-
-		// Force replication so GameplayCue removals are properly replicated to all clients during Mixed and Minimal replication modes
-		ForceReplication();
+		int32 NumMatchingCues = 0;
+		for (const FActiveGameplayCue& GameplayCue : GameplayCueContainer.GameplayCues)
+		{
+			NumMatchingCues += (GameplayCue.GameplayCueTag == GameplayCueTag);
+		}
 
-		if (bWasInList)
+		if (NumMatchingCues > 0)
 		{
-			for (const FActiveGameplayCue& GameplayCue : GameplayCueContainer.GameplayCues)
+			// AbilitySystem.GameplayCueNotifyTagCheckOnRemove assumes the tag is removed before any invocation of EGameplayCueEvent::Removed.
+			// We cannot use GameplayCueContainer.RemoveCue because that removes the cues while updating the TagMap.
+			// Instead, we need to manually count the removals, update the tag map, then Invoke the Cue events while removing the Cues.
+			UpdateTagMap(GameplayCueTag, -NumMatchingCues);
+
+			for (int32 Index = GameplayCueContainer.GameplayCues.Num() - 1; Index >= 0; --Index)
 			{
+				const FActiveGameplayCue& GameplayCue = GameplayCueContainer.GameplayCues[Index];
 				if (GameplayCue.GameplayCueTag == GameplayCueTag)
 				{
-					// Call on server here, clients get it from repnotify
+					// Call on server here, clients get it from repnotify on the GameplayCueContainer rather than a multicast rpc
 					InvokeGameplayCueEvent(GameplayCueTag, EGameplayCueEvent::Removed, GameplayCue.Parameters);
+					GameplayCueContainer.GameplayCues.RemoveAt(Index);
 				}
 			}
+
+			// Ensure that the clients are aware of these changes ASAP
+			GameplayCueContainer.MarkArrayDirty();
+			ForceReplication();
 		}
-		
-		GameplayCueContainer.RemoveCue(GameplayCueTag);
-		
-		// Don't need to multicast broadcast this, ActiveGameplayCues replication handles it
 	}
 	else if (ScopedPredictionKey.IsLocalClientKey())
 	{
