Add warnings to GameplayEffect netserialize functions

diff --git a/Engine/Plugins/Runtime/GameplayAbilities/Source/GameplayAbilities/Private/GameplayEffectTypes.cpp b/Engine/Plugins/Runtime/GameplayAbilities/Source/GameplayAbilities/Private/GameplayEffectTypes.cpp
--- a/Engine/Plugins/Runtime/GameplayAbilities/Source/GameplayAbilities/Private/GameplayEffectTypes.cpp
+++ b/Engine/Plugins/Runtime/GameplayAbilities/Source/GameplayAbilities/Private/GameplayEffectTypes.cpp
@@ -235,6 +235,10 @@ void FGameplayEffectContext::AddHitResult(const FHitResult& InHitResult, bool bR
 
 bool FGameplayEffectContext::NetSerialize(FArchive& Ar, class UPackageMap* Map, bool& bOutSuccess)
 {
+	//!!!!!!!!!!!!!!!! Warning !!!!!!!!!!!!!!!
+	// Any changes to this function also need to be done to FGameplayEffectContextNetSerializer to support Iris replication
+	//!!!!!!!!!!!!!!!! Warning !!!!!!!!!!!!!!!
+
 	uint8 RepBits = 0;
 	if (Ar.IsSaving())
 	{
diff --git a/Engine/Plugins/Runtime/GameplayAbilities/Source/GameplayAbilities/Private/Serialization/GameplayEffectContextNetSerializer.cpp b/Engine/Plugins/Runtime/GameplayAbilities/Source/GameplayAbilities/Private/Serialization/GameplayEffectContextNetSerializer.cpp
--- a/Engine/Plugins/Runtime/GameplayAbilities/Source/GameplayAbilities/Private/Serialization/GameplayEffectContextNetSerializer.cpp
+++ b/Engine/Plugins/Runtime/GameplayAbilities/Source/GameplayAbilities/Private/Serialization/GameplayEffectContextNetSerializer.cpp
@@ -110,6 +110,10 @@ EReplicationStateTraits FGameplayEffectContextNetSerializer::HitResultStateTrait
 
 void FGameplayEffectContextNetSerializer::Serialize(FNetSerializationContext& Context, const FNetSerializeArgs& Args)
 {
+	//!!!!!!!!!!!!!!!! Warning !!!!!!!!!!!!!!!
+	// Any changes to this function also need to be done to FGameplayEffectContext::NetSerialize to support generic replication
+	//!!!!!!!!!!!!!!!! Warning !!!!!!!!!!!!!!!
+
 	const QuantizedType& Value = *reinterpret_cast<QuantizedType*>(Args.Source);
 
 	FNetBitStreamWriter* Writer = Context.GetBitStreamWriter();
