AbilitySystem.Effect.Apply return value has been fixed to account for Instant Effects.

#changelist validated
#virtualized

diff --git a/Engine/Plugins/Runtime/GameplayAbilities/Source/GameplayAbilities/Private/AbilitySystemCheatManagerExtension.cpp b/Engine/Plugins/Runtime/GameplayAbilities/Source/GameplayAbilities/Private/AbilitySystemCheatManagerExtension.cpp
--- a/Engine/Plugins/Runtime/GameplayAbilities/Source/GameplayAbilities/Private/AbilitySystemCheatManagerExtension.cpp
+++ b/Engine/Plugins/Runtime/GameplayAbilities/Source/GameplayAbilities/Private/AbilitySystemCheatManagerExtension.cpp
@@ -453,7 +453,11 @@ void UAbilitySystemCheatManagerExtension::EffectApply(const FString& PartialName
 	FActiveGameplayEffectHandle ActiveGEHandle = ASC->ApplyGameplayEffectSpecToSelf(GESpec, FakePredictionKey);
 	if (ActiveGEHandle.IsValid())
 	{
-		UE_LOG(LogConsoleResponse, Log, TEXT("Successfully Applied Gameplay Effect '%s' on Player '%s'."), *GameplayEffectClass->GetName(), *PC->GetName());
+		UE_LOG(LogConsoleResponse, Log, TEXT("Successfully Applied (Non-Instant) Gameplay Effect '%s' on Player '%s'."), *GameplayEffectClass->GetName(), *PC->GetName());
+	}
+	else if (ActiveGEHandle.WasSuccessfullyApplied())
+	{
+		UE_LOG(LogConsoleResponse, Log, TEXT("Executed (Instant) Gameplay Effect '%s' on Player '%s'."), *GameplayEffectClass->GetName(), *PC->GetName());
 	}
 	else
 	{
diff --git a/Engine/Plugins/Runtime/GameplayAbilities/Source/GameplayAbilities/Private/GameplayEffect.cpp b/Engine/Plugins/Runtime/GameplayAbilities/Source/GameplayAbilities/Private/GameplayEffect.cpp
--- a/Engine/Plugins/Runtime/GameplayAbilities/Source/GameplayAbilities/Private/GameplayEffect.cpp
+++ b/Engine/Plugins/Runtime/GameplayAbilities/Source/GameplayAbilities/Private/GameplayEffect.cpp
@@ -3901,6 +3901,7 @@ FActiveGameplayEffect* FActiveGameplayEffectsContainer::ApplyGameplayEffectSpec(
 		{
 			if (!HandleActiveGameplayEffectStackOverflow(*ExistingStackableGE, ExistingSpec, Spec))
 			{
+				UE_VLOG_UELOG(Owner, LogGameplayEffects, Log, TEXT("Application of %s denied (StackLimit)"), *Spec.ToSimpleString());
 				return nullptr;
 			}
 		}
