Iris - Update MinimalReplicationTagCountMap logic to catch up with recent changes.

#rnx

#changelist validated
#virtualized

diff --git a/Engine/Plugins/Runtime/GameplayAbilities/Source/GameplayAbilities/Private/GameplayEffectTypes.cpp b/Engine/Plugins/Runtime/GameplayAbilities/Source/GameplayAbilities/Private/GameplayEffectTypes.cpp
--- a/Engine/Plugins/Runtime/GameplayAbilities/Source/GameplayAbilities/Private/GameplayEffectTypes.cpp
+++ b/Engine/Plugins/Runtime/GameplayAbilities/Source/GameplayAbilities/Private/GameplayEffectTypes.cpp
@@ -1296,6 +1296,7 @@ void FMinimalReplicationTagCountMap::RemoveTag(const FGameplayTag& Tag)
 	}
 }
 
+// WARNING: Changes to this implementation REQUIRES making sure FMinimalReplicationTagCountMapNetSerializer and FMinimalReplicationTagCountMapReplicationFragment remains compatible.
 bool FMinimalReplicationTagCountMap::NetSerialize(FArchive& Ar, class UPackageMap* Map, bool& bOutSuccess)
 {
 	const int32 CountBits = UAbilitySystemGlobals::Get().MinimalReplicationTagCountBits;
diff --git a/Engine/Plugins/Runtime/GameplayAbilities/Source/GameplayAbilities/Private/Serialization/MinimalReplicationTagCountMapReplicationFragment.cpp b/Engine/Plugins/Runtime/GameplayAbilities/Source/GameplayAbilities/Private/Serialization/MinimalReplicationTagCountMapReplicationFragment.cpp
--- a/Engine/Plugins/Runtime/GameplayAbilities/Source/GameplayAbilities/Private/Serialization/MinimalReplicationTagCountMapReplicationFragment.cpp
+++ b/Engine/Plugins/Runtime/GameplayAbilities/Source/GameplayAbilities/Private/Serialization/MinimalReplicationTagCountMapReplicationFragment.cpp
@@ -110,29 +110,20 @@ void FMinimalReplicationTagCountMapReplicationFragment::MimicMinimalReplicationT
 	uint8* ExternalStatePointer = reinterpret_cast<uint8*>(Owner) + ReplicationStateDescriptor->MemberProperties[0]->GetOffset_ForGC();
 	FMinimalReplicationTagCountMap* ExternalSourceState = reinterpret_cast<FMinimalReplicationTagCountMap*>(ExternalStatePointer);
 
-	// Check whether we need to perform the complex logic or not.
-	UAbilitySystemComponent* StateOwner = ExternalSourceState->Owner;
-	bool UpdateOwnerTagMap = StateOwner != nullptr;
-	if (ExternalSourceState->bRequireNonOwningNetConnection && StateOwner)
-	{
-		if (AActor* OwningActor = StateOwner->GetOwner())
-		{			
-			if (const UNetConnection* OwnerNetConnection = OwningActor->GetNetConnection())
-			{
-				if (OwnerNetConnection->GetConnectionId() == Context.NetSerializationContext->GetLocalConnectionId())
-				{
-					UpdateOwnerTagMap = false;
-				}
-			}
-		}
-	}
+	// Mark map dirty for replay
+	ExternalSourceState->MapID++;
+
+	// UpdateOwnerTagMap performs most of the logic regarding whether it should update or not and needs LastConnection to do it.
+	const uint32 ConnectionId = Context.NetSerializationContext->GetLocalConnectionId();
+	UObject* UserData = Context.NetSerializationContext->GetLocalConnectionUserData(ConnectionId);
+	ExternalSourceState->LastConnection = Cast<UNetConnection>(UserData);
 
-	if (!UpdateOwnerTagMap)
+	const UAbilitySystemComponent* StateOwner = ExternalSourceState->Owner;
+	const bool bUpdateOwnerTagMap = StateOwner != nullptr;
+	if (bUpdateOwnerTagMap)
 	{
-		return;
+		ExternalSourceState->UpdateOwnerTagMap();
 	}
-
-	ExternalSourceState->UpdateOwnerTagMap();
 }
 
 void FMinimalReplicationTagCountMapReplicationFragment::CallRepNotify(FReplicationStateApplyContext& ApplyContext)
