Improving the GAS Debug Facilities.  Attributes can now be debugged in the Gameplay Debugger.  Layout of the GameplayDebugger is better formatted.

diff --git a/Engine/Plugins/Runtime/GameplayAbilities/Source/GameplayAbilities/Private/AbilitySystemComponent.cpp b/Engine/Plugins/Runtime/GameplayAbilities/Source/GameplayAbilities/Private/AbilitySystemComponent.cpp
--- a/Engine/Plugins/Runtime/GameplayAbilities/Source/GameplayAbilities/Private/AbilitySystemComponent.cpp
+++ b/Engine/Plugins/Runtime/GameplayAbilities/Source/GameplayAbilities/Private/AbilitySystemComponent.cpp
@@ -2507,6 +2507,8 @@ void UAbilitySystemComponent::Debug_Internal(FAbilitySystemComponentDebugInfo& I
 		}
 	}
 
+	DebugLine(Info, TEXT("Tip: Use the GameplayDebugger for enhanced functionality"), 4.0f, 0.0f);
+
 	FGameplayTagContainer OwnerTags;
 	GetOwnedGameplayTags(OwnerTags);
 
diff --git a/Engine/Plugins/Runtime/GameplayAbilities/Source/GameplayAbilities/Private/GameplayDebuggerCategory_Abilities.cpp b/Engine/Plugins/Runtime/GameplayAbilities/Source/GameplayAbilities/Private/GameplayDebuggerCategory_Abilities.cpp
--- a/Engine/Plugins/Runtime/GameplayAbilities/Source/GameplayAbilities/Private/GameplayDebuggerCategory_Abilities.cpp
+++ b/Engine/Plugins/Runtime/GameplayAbilities/Source/GameplayAbilities/Private/GameplayDebuggerCategory_Abilities.cpp
@@ -1,20 +1,34 @@
 // Copyright Epic Games, Inc. All Rights Reserved.
 
 #include "GameplayDebuggerCategory_Abilities.h"
-#include "Engine/Engine.h"
-#include "Engine/GameViewportClient.h"
 
 #if WITH_GAMEPLAY_DEBUGGER_MENU
 
-#include "GameplayTagContainer.h"
-#include "GameplayAbilitySpec.h"
-#include "GameplayEffect.h"
 #include "AbilitySystemGlobals.h"
 #include "AbilitySystemComponent.h"
+#include "Engine/ActorChannel.h"
 #include "Engine/Canvas.h"
+#include "Engine/Engine.h"
+#include "Engine/GameViewportClient.h"
 #include "Engine/NetConnection.h"
 #include "GameFramework/PlayerController.h"
+#include "GameplayTagContainer.h"
+#include "GameplayAbilitySpec.h"
+#include "GameplayEffect.h"
+#include "Net/Subsystems/NetworkSubsystem.h"
+#include "Net/RepLayout.h"
 
+namespace UE::AbilitySystem::Debug
+{
+	// This is the longest name we can use for the UI (string format truncate with %.35s).  We use a variety of letters because MeasureString depends on kerning.
+	const FString LongestDebugObjectName{ TEXT("ABCDEFGHIJKLMNOPQRSTUVWXYZ_ ABCDEFGH") };
+
+	// Let's define some names for consistent look across all of the categories
+	const TCHAR* BothColor = TEXT("{cyan}");
+	const TCHAR* ServerColor = TEXT("{yellow}");
+	const TCHAR* LocalColor = TEXT("{green}");
+	const TCHAR* NonReplicatedColor = TEXT("{violetred}");
+}
 
 FGameplayDebuggerCategory_Abilities::FGameplayDebuggerCategory_Abilities()
 {
@@ -24,10 +38,12 @@ FGameplayDebuggerCategory_Abilities::FGameplayDebuggerCategory_Abilities()
 	const FName KeyNameOne{ "One" };
 	const FName KeyNameTwo{ "Two" };
 	const FName KeyNameThree{ "Three" };
+	const FName KeyNameFour{ "Four" };
 
 	BindKeyPress(KeyNameOne, FGameplayDebuggerInputModifier::Shift, this, &FGameplayDebuggerCategory_Abilities::OnShowGameplayTagsToggle, EGameplayDebuggerInputMode::Local);
 	BindKeyPress(KeyNameTwo, FGameplayDebuggerInputModifier::Shift, this, &FGameplayDebuggerCategory_Abilities::OnShowGameplayAbilitiesToggle, EGameplayDebuggerInputMode::Local);
 	BindKeyPress(KeyNameThree, FGameplayDebuggerInputModifier::Shift, this, &FGameplayDebuggerCategory_Abilities::OnShowGameplayEffectsToggle, EGameplayDebuggerInputMode::Local);
+	BindKeyPress(KeyNameFour, FGameplayDebuggerInputModifier::Shift, this, &FGameplayDebuggerCategory_Abilities::OnShowGameplayAttributesToggle, EGameplayDebuggerInputMode::Local);
 }
 
 TSharedRef<FGameplayDebuggerCategory> FGameplayDebuggerCategory_Abilities::MakeInstance()
@@ -50,11 +66,18 @@ void FGameplayDebuggerCategory_Abilities::OnShowGameplayEffectsToggle()
 	bShowGameplayEffects = !bShowGameplayEffects;
 }
 
+void FGameplayDebuggerCategory_Abilities::OnShowGameplayAttributesToggle()
+{
+	bShowGameplayAttributes = !bShowGameplayAttributes;
+}
+
 void FGameplayDebuggerCategory_Abilities::FRepData::Serialize(FArchive& Ar)
 {
 	bool bSuccess;
 	OwnedTags.NetSerialize(Ar, ClientPackageMap.Get(), bSuccess);
 
+	Ar << TagCounts;
+
 	int32 NumAbilities = Abilities.Num();
 	Ar << NumAbilities;
 	if (Ar.IsLoading())
@@ -86,12 +109,26 @@ void FGameplayDebuggerCategory_Abilities::FRepData::Serialize(FArchive& Ar)
 		Ar << GameplayEffects[Idx].Stacks;
 		Ar << GameplayEffects[Idx].Level;
 	}
+
+	int32 NumAttrib = Attributes.Num();
+	Ar << NumAttrib;
+	if (Ar.IsLoading())
+	{
+		Attributes.SetNum(NumAttrib);
+	}
+
+	for (int32 Idx = 0; Idx < NumAttrib; ++Idx)
+	{
+		Ar << Attributes[Idx].AttributeName;
+		Ar << Attributes[Idx].BaseValue;
+		Ar << Attributes[Idx].CurrentValue;
+		Ar << Attributes[Idx].NetworkStatus;
+	}
 }
 
 void FGameplayDebuggerCategory_Abilities::CollectData(APlayerController* OwnerPC, AActor* DebugActor)
 {
-	UAbilitySystemComponent* AbilityComp = UAbilitySystemGlobals::GetAbilitySystemComponentFromActor(DebugActor);
-	if (AbilityComp)
+	if (const UAbilitySystemComponent* AbilityComp = UAbilitySystemGlobals::GetAbilitySystemComponentFromActor(DebugActor))
 	{
 		// Save off the package map for serialization over the network
 		UNetConnection* NetConnection = OwnerPC->GetNetConnection();
@@ -99,6 +136,14 @@ void FGameplayDebuggerCategory_Abilities::CollectData(APlayerController* OwnerPC
 
 		AbilityComp->GetOwnedGameplayTags(DataPack.OwnedTags);
 
+		// Copy over the tag counts
+		DataPack.TagCounts.Empty();
+		for (const FGameplayTag& Tag : DataPack.OwnedTags)
+		{
+			DataPack.TagCounts.Add(AbilityComp->GetTagCount(Tag));
+		}
+
+		// Gameplay Effects
 		TArray<FGameplayEffectSpec> ActiveEffectSpecs;
 		AbilityComp->GetAllActiveGameplayEffectSpecs(ActiveEffectSpecs);
 		for (int32 Idx = 0; Idx < ActiveEffectSpecs.Num(); Idx++)
@@ -119,6 +164,7 @@ void FGameplayDebuggerCategory_Abilities::CollectData(APlayerController* OwnerPC
 			DataPack.GameplayEffects.Add(ItemData);
 		}
 
+		// Abilities
 		const TArray<FGameplayAbilitySpec>& AbilitySpecs = AbilityComp->GetActivatableAbilities();
 		for (int32 Idx = 0; Idx < AbilitySpecs.Num(); Idx++)
 		{
@@ -137,9 +183,107 @@ void FGameplayDebuggerCategory_Abilities::CollectData(APlayerController* OwnerPC
 
 			DataPack.Abilities.Add(ItemData);
 		}
+
+		// Attributes
+		DataPack.Attributes = CollectAttributeData(OwnerPC, AbilityComp);
 	}
 }
 
+TArray<FGameplayDebuggerCategory_Abilities::FRepData::FGameplayAttributeDebug> FGameplayDebuggerCategory_Abilities::CollectAttributeData(const APlayerController* OwnerPC, const UAbilitySystemComponent* AbilityComp) const
+{
+	TArray<FRepData::FGameplayAttributeDebug> DebugAttributes;
+
+	const AActor* Actor = AbilityComp ? AbilityComp->GetOwner() : nullptr;
+	if (!Actor || !Actor->GetWorld())
+	{
+		FRepData::FGameplayAttributeDebug& DebugAttribute = DebugAttributes.AddDefaulted_GetRef();
+		DebugAttribute.AttributeName = TEXT("Invalid AbilitySystemComponent");
+		return DebugAttributes;
+	}
+
+	// We need to do a lot of work to detect if the AttributeSet is replicated which only occurs of AbilitySystemComponent is replicated
+	const ENetRole LocalRole = Actor->GetLocalRole();
+	const ELifetimeCondition NetCondition = Actor->AllowActorComponentToReplicate(AbilityComp);
+
+	bool bASCReplicates = (NetCondition != ELifetimeCondition::COND_Never) && AbilityComp->IsSupportedForNetworking() && AbilityComp->IsNameStableForNetworking();
+	if (bASCReplicates)
+	{
+		// Lots of extra work due to COND_NetGroup
+		if (const UNetworkSubsystem* NetConditionGroupSubsystem = Actor->GetWorld()->GetSubsystem<UNetworkSubsystem>())
+		{
+			FReplicationFlags RepFlags;
+			RepFlags.bNetOwner = (Actor->GetNetConnection() == OwnerPC->GetNetConnection());
+			RepFlags.bNetSimulated = !RepFlags.bNetOwner && Actor->GetNetConnection();
+			RepFlags.bRolesOnly = true;
+			const TStaticBitArray<COND_Max> ConditionMap = UE::Net::BuildConditionMapFromRepFlags(RepFlags);
+
+			bASCReplicates = UActorChannel::CanSubObjectReplicateToClient(OwnerPC, NetCondition, AbilityComp, ConditionMap, NetConditionGroupSubsystem->GetNetConditionGroupManager());
+		}
+	}
+
+	// Grab the AttributeSet rather than the Attributes themselves so we can check the network functionality
+	static TMap<UClass*, TArray<FGameplayAttribute>> AttributeSetClassToGameplayAttributes;
+	for (const UAttributeSet* AttributeSet : AbilityComp->GetSpawnedAttributes())
+	{
+		const TSubclassOf<UAttributeSet> AttributeSetClass = AttributeSet ? AttributeSet->GetClass() : nullptr;
+		if (!AttributeSet || !AttributeSetClass)
+		{
+			continue;
+		}
+
+		// Keep the FGameplayAttributes in a static map as we'll be looking these up (at least) twice per frame
+		TArray<FGameplayAttribute>* LocalAttributes = AttributeSetClassToGameplayAttributes.Find(AttributeSetClass.Get());
+		if (!LocalAttributes)
+		{
+			LocalAttributes = &AttributeSetClassToGameplayAttributes.Add(AttributeSetClass.Get());
+			UAttributeSet::GetAttributesFromSetClass(AttributeSetClass, *LocalAttributes);
+		}
+
+		// These are all of the replication conditions per variable
+		TArray<FLifetimeProperty> LifetimeProps;
+		AttributeSet->GetLifetimeReplicatedProps(LifetimeProps);
+
+		// Network status can change per AttributeSet, so figure it out
+		FRepData::ENetworkStatus NetworkStatus = (LocalRole == ENetRole::ROLE_Authority) ? FRepData::ENetworkStatus::ServerOnly : FRepData::ENetworkStatus::LocalOnly;
+		const bool bAttributeSetReplicates = bASCReplicates && AttributeSet->IsSupportedForNetworking();
+
+		// Now just gather the debug data
+		for (const FGameplayAttribute& Attrib : *LocalAttributes)
+		{
+			if (!Attrib.IsValid())
+			{
+				continue;
+			}
+
+			FRepData::FGameplayAttributeDebug& DebugAttribute = DebugAttributes.AddDefaulted_GetRef();
+
+			DebugAttribute.AttributeName = Attrib.AttributeName;
+			DebugAttribute.BaseValue = AbilityComp->GetNumericAttributeBase(Attrib);
+			DebugAttribute.CurrentValue = AbilityComp->GetNumericAttribute(Attrib);
+			DebugAttribute.NetworkStatus = NetworkStatus;
+
+			// Override the status to network if all replication tests pass
+			if (bAttributeSetReplicates)
+			{
+				if (FProperty* Property = Attrib.GetUProperty())
+				{
+					if (Property->HasAnyPropertyFlags(EPropertyFlags::CPF_Net))
+					{
+						FLifetimeProperty* RepProperty = LifetimeProps.FindByPredicate([RepIndex = Property->RepIndex](const FLifetimeProperty& Item) { return Item.RepIndex == RepIndex; });
+						if (RepProperty && RepProperty->Condition != ELifetimeCondition::COND_Never)
+						{
+							// InvalidRepIndex maps to INDEX_NONE even though unsigned (wrap-around)
+							DebugAttribute.NetworkStatus = FRepData::ENetworkStatus::Networked;
+						}
+					}
+				}
+			}
+		}
+	}
+
+	return DebugAttributes;
+}
+
 bool FGameplayDebuggerCategory_Abilities::WrapStringAccordingToViewport(const FString& StrIn, FString& StrOut, FGameplayDebuggerCanvasContext& CanvasContext, float ViewportWitdh) const
 {
 	if (!StrIn.IsEmpty())
@@ -182,7 +326,11 @@ void FGameplayDebuggerCategory_Abilities::DrawData(APlayerController* OwnerPC, F
 		CanvasContext.CursorY -= CanvasContext.GetLineHeight();
 		const TCHAR* Active = TEXT("{green}");
 		const TCHAR* Inactive = TEXT("{yellow}");
-		CanvasContext.Printf(TEXT("Tags [%s%s{white}]\tAbilities [%s%s{white}]\tEffects [%s%s{white}]"), bShowGameplayTags ? Active : Inactive, *GetInputHandlerDescription(0), bShowGameplayAbilities ? Active : Inactive, *GetInputHandlerDescription(1), bShowGameplayEffects ? Active : Inactive, *GetInputHandlerDescription(2));
+		CanvasContext.Printf(TEXT("Tags [%s%s{white}]\tAbilities [%s%s{white}]\tEffects [%s%s{white}]\tAttributes [%s%s{white}]"),
+			bShowGameplayTags ? Active : Inactive, *GetInputHandlerDescription(0),
+			bShowGameplayAbilities ? Active : Inactive, *GetInputHandlerDescription(1),
+			bShowGameplayEffects ? Active : Inactive, *GetInputHandlerDescription(2),
+			bShowGameplayAttributes ? Active : Inactive, *GetInputHandlerDescription(3));
 	}
 
 	static float LastDrawDataEndSize = CanvasContext.Canvas->SizeY - CanvasContext.CursorY - CanvasContext.CursorX;
@@ -213,35 +361,97 @@ void FGameplayDebuggerCategory_Abilities::DrawData(APlayerController* OwnerPC, F
 		DrawGameplayEffects(CanvasContext, OwnerPC);
 	}
 
+	if (bShowGameplayAttributes)
+	{
+		DrawGameplayAttributes(CanvasContext, OwnerPC);
+	}
+
 	LastDrawDataEndSize = CanvasContext.CursorY - ThisDrawDataStartPos;
 }
 
 void FGameplayDebuggerCategory_Abilities::DrawGameplayTags(FGameplayDebuggerCanvasContext& CanvasContext, const APlayerController* OwnerPC) const
 {
+	using namespace UE::AbilitySystem::Debug;
 	const float CanvasWidth = CanvasContext.Canvas->SizeX;
 
 	AActor* LocalDebugActor = FindLocalDebugActor();
-	if (UAbilitySystemComponent* AbilityComp = UAbilitySystemGlobals::GetAbilitySystemComponentFromActor(LocalDebugActor))
+	if (const UAbilitySystemComponent* AbilityComp = UAbilitySystemGlobals::GetAbilitySystemComponentFromActor(LocalDebugActor))
 	{
 		FGameplayTagContainer ServerOnlyTags = DataPack.OwnedTags;
+		const TArray<int32>& ServerTagCounts = DataPack.TagCounts;
+		ensureMsgf(ServerTagCounts.Num() == ServerOnlyTags.Num(), TEXT("Server's OwnedTags Array and ServerTagCounts Array did not match in size. Check FGameplayDebuggerCategory_Abilities's CollectData or NetSerialize functions."));
+
+		// Fill in the map for easy access
+		TMap<FGameplayTag, int32> ServerTagCountMap;
+		for (int32 Idx = 0 ; Idx < ServerOnlyTags.Num() ; ++Idx)
+		{
+			const FGameplayTag& Tag = ServerOnlyTags.GetByIndex(Idx);
+			const int32 ServerTagCount = Idx < ServerTagCounts.Num() ? ServerTagCounts[Idx] : 0;
+
+			ServerTagCountMap.Add(Tag, ServerTagCount);
+		}
+
+		auto BuildTagsString = [&](const FGameplayTagContainer& TagContainer, const TCHAR* DefaultColorName) -> FString
+			{
+				if (TagContainer.IsEmpty())
+				{
+					return FString{};
+				}
+
+				TStringBuilder<1024> StringBuilder;
+				StringBuilder.Append(DefaultColorName);
+
+				for (const FGameplayTag& Tag : TagContainer)
+				{
+					const int32 LocalTagCount = AbilityComp->GetTagCount(Tag);
+					const int32 ServerTagCount = ServerTagCountMap.FindRef(Tag, 0);
+
+					// If we're out of sync, then we should print both
+					if (DefaultColorName == BothColor && LocalTagCount != ServerTagCount)
+					{
+						StringBuilder.Appendf(TEXT(" %s %s(x%d) %s(x%d)%s,"), *Tag.ToString(), ServerColor, ServerTagCount, LocalColor, LocalTagCount, DefaultColorName);
+						continue;
+					}
+
+					// We're in sync (or only exist in one location), just print the numbers
+					const int32 TagCount = DefaultColorName == ServerColor ? ServerTagCount : LocalTagCount;
+					if (TagCount == 1)
+					{
+						StringBuilder.Appendf(TEXT(" %s,"), *Tag.ToString());
+					}
+					else if (TagCount > 1)
+					{
+						StringBuilder.Appendf(TEXT(" %s (x%d),"), *Tag.ToString(), TagCount);
+					}
+					else if (TagCount <= 0)
+					{
+						// Special case where the count is zero.  We should call this out.
+						StringBuilder.Appendf(TEXT(" {red}%s (x%d)%s,"), *Tag.ToString(), TagCount, DefaultColorName);
+					}
+				}
+
+				// Remove the last comma
+				StringBuilder.RemoveSuffix(1);
+				return StringBuilder.ToString();
+			};
 
 		// If we're not the authority, we should represent the tags in such a way that the user can see the difference between agreed upon server/client tags
 		// and ones where they disagree.
 		if (!AbilityComp->IsOwnerActorAuthoritative())
 		{
-			CanvasContext.Printf(TEXT("Owned Tags Legend:  {cyan}Both  {yellow}Server  {green}Local"));
+			CanvasContext.Printf(TEXT("Owned Tags\tLegend:  %sBoth  %sServer  %sLocal"), BothColor, ServerColor, LocalColor);
 
-			static FGameplayTagContainer LocalOnlyTags;
+			FGameplayTagContainer LocalOnlyTags;
 			AbilityComp->GetOwnedGameplayTags(LocalOnlyTags);
 
-			FGameplayTagContainer MatchingTags = ServerOnlyTags.FilterExact(LocalOnlyTags);
+			const FGameplayTagContainer MatchingTags = ServerOnlyTags.FilterExact(LocalOnlyTags);
 			ServerOnlyTags.RemoveTags(MatchingTags);
 			LocalOnlyTags.RemoveTags(MatchingTags);
 
 			// Build up the strings
-			FString MatchingTagsStr = MatchingTags.Num() > 0 ? FString::Printf(TEXT("{cyan}%s  "), *MatchingTags.ToStringSimple()) : FString{};
-			FString ServerOnlyTagsStr = ServerOnlyTags.Num() > 0 ? FString::Printf(TEXT("{yellow}%s  "), *ServerOnlyTags.ToStringSimple()) : FString{};
-			FString LocalOnlyTagsStr = LocalOnlyTags.Num() > 0 ? FString::Printf(TEXT("{green}%s  "), *LocalOnlyTags.ToStringSimple()) : FString{};
+			const FString MatchingTagsStr = BuildTagsString(MatchingTags, BothColor);
+			const FString ServerOnlyTagsStr = BuildTagsString(ServerOnlyTags, ServerColor);
+			const FString LocalOnlyTagsStr = BuildTagsString(LocalOnlyTags, LocalColor);
 			
 			FString WrappedDebugText;
 			WrapStringAccordingToViewport(FString::Printf(TEXT("%s%s%s"), *MatchingTagsStr, *ServerOnlyTagsStr, *LocalOnlyTagsStr), WrappedDebugText, CanvasContext, CanvasWidth);
@@ -249,9 +459,10 @@ void FGameplayDebuggerCategory_Abilities::DrawGameplayTags(FGameplayDebuggerCanv
 		}
 		else
 		{
+			// As the authority, the source of the truth should be the ServerTags we already gathered
 			FString ServerOnlyTagsStr;
-			WrapStringAccordingToViewport(ServerOnlyTags.ToStringSimple(), ServerOnlyTagsStr, CanvasContext, CanvasWidth);
-			CanvasContext.Printf(TEXT("Owned Tags: \n{cyan}%s"), *ServerOnlyTagsStr);
+			WrapStringAccordingToViewport(BuildTagsString(ServerOnlyTags, ServerColor), ServerOnlyTagsStr, CanvasContext, CanvasWidth);
+			CanvasContext.Printf(TEXT("Owned Tags: \n%s"), *ServerOnlyTagsStr);
 		}
 	}
 
@@ -261,13 +472,34 @@ void FGameplayDebuggerCategory_Abilities::DrawGameplayTags(FGameplayDebuggerCanv
 
 void FGameplayDebuggerCategory_Abilities::DrawGameplayEffects(FGameplayDebuggerCanvasContext& CanvasContext, const APlayerController* OwnerPC) const
 {
+	using namespace UE::AbilitySystem::Debug;
+
+	// Find some stable naming sizes
+	constexpr float Padding = 10.0f;
+	static float ObjNameSize = 0.0f, SrcNameSize = 0.0f;
+	if (ObjNameSize <= 0.0f)
+	{
+		float TempSizeY = 0.0f;
+		CanvasContext.MeasureString(*LongestDebugObjectName, ObjNameSize, TempSizeY);
+		CanvasContext.MeasureString(FString::Printf(TEXT("source: %.35s"), *LongestDebugObjectName), SrcNameSize, TempSizeY);
+		ObjNameSize += Padding;
+		SrcNameSize += Padding;
+	}
+
 	CanvasContext.Printf(TEXT("Gameplay Effects: {yellow}%d"), DataPack.GameplayEffects.Num());
+	CanvasContext.CursorX += Padding;
 	for (int32 Idx = 0; Idx < DataPack.GameplayEffects.Num(); Idx++)
 	{
 		const FRepData::FGameplayEffectDebug& ItemData = DataPack.GameplayEffects[Idx];
 
-		FStringBuilderBase Desc;
-		Desc.Appendf(TEXT("\t{yellow}%s {grey}source:{white}%s {grey}duration:{white}"), *ItemData.Effect, *ItemData.Context);
+		float CursorX = CanvasContext.CursorX;
+		float CursorY = CanvasContext.CursorY;
+
+		CanvasContext.PrintAt(CursorX, CursorY, FColor::Yellow, *ItemData.Effect.Left(35));
+		CanvasContext.PrintfAt(CursorX + ObjNameSize, CursorY, FColor::Silver, TEXT("source: {white}%.35s"), *ItemData.Context);
+
+		TStringBuilder<1024> Desc;
+		Desc.Appendf(TEXT(" {grey}duration: {white}"), *ItemData.Context);
 		if (ItemData.Duration > 0.0f)
 		{
 			Desc.Appendf(TEXT("%.3f"), ItemData.Duration);
@@ -279,20 +511,22 @@ void FGameplayDebuggerCategory_Abilities::DrawGameplayEffects(FGameplayDebuggerC
 
 		if (ItemData.Period > 0.0f)
 		{
-			Desc.Appendf(TEXT(" {grey}period:{white}%.3f"), ItemData.Period);
+			Desc.Appendf(TEXT(" {grey}period: {white}%.3f"), ItemData.Period);
 		}
 
 		if (ItemData.Stacks > 1)
 		{
-			Desc.Appendf(TEXT(" {grey}stacks:{white}%d"), ItemData.Stacks);
+			Desc.Appendf(TEXT(" {grey}stacks: {white}%d"), ItemData.Stacks);
 		}
 
 		if (ItemData.Level > 1.0f)
 		{
-			Desc.Appendf(TEXT(" {grey}level:{white}%.2f"), ItemData.Level);
+			Desc.Appendf(TEXT(" {grey}level: {white}%.2f"), ItemData.Level);
 		}
 
-		CanvasContext.Print(Desc.ToString());
+		CanvasContext.PrintAt(CursorX + ObjNameSize + SrcNameSize, CursorY, *Desc);
+		CanvasContext.MoveToNewLine();
+		CanvasContext.CursorX += Padding;
 	}
 	
 	// End with a newline to separate from the other categories
@@ -301,33 +535,27 @@ void FGameplayDebuggerCategory_Abilities::DrawGameplayEffects(FGameplayDebuggerC
 
 void FGameplayDebuggerCategory_Abilities::DrawGameplayAbilities(FGameplayDebuggerCanvasContext& CanvasContext, const APlayerController* OwnerPC) const
 {
+	using namespace UE::AbilitySystem::Debug;
+
 	const float CanvasWidth = CanvasContext.Canvas->SizeX;
+	Algo::Sort(DataPack.Abilities, [](const FRepData::FGameplayAbilityDebug& ItemOne, const FRepData::FGameplayAbilityDebug& ItemTwo) { return ItemOne.Ability < ItemTwo.Ability; });
 
-	// Let's do something semi-smart and stable to resize the columns to readable yet wide enough
 	int32 NumActive = 0;
-	FString LongestLengthObjectName;
 	for (const FRepData::FGameplayAbilityDebug& ItemData : DataPack.Abilities)
 	{
-		if (ItemData.Ability.Len() > LongestLengthObjectName.Len())
-		{
-			LongestLengthObjectName = ItemData.Ability;
-		}
-
-		if (ItemData.Source.Len() > LongestLengthObjectName.Len())
-		{
-			LongestLengthObjectName = ItemData.Source;
-		}
-
 		NumActive += ItemData.bIsActive;
 	}
 
 	// Measure the individual string sizes, so that we can size the columns properly
+	// We're picking a long-enough name for the object name sizes
 	constexpr float Padding = 10.0f;
-	static float ObjNameSize = 150.0f, SourceNameSize = 100.0f, LevelNameSize = 100.0f, TempSizeY = 0.0f;
-	static int32 CachedLen = 0;
-	if (LongestLengthObjectName.Len() > CachedLen)
+	static float ObjNameSize = 0.0f, SourceNameSize = 0.0f, LevelNameSize = 0.0f;
+	if (ObjNameSize <= 0.0f)
 	{
-		CanvasContext.MeasureString(LongestLengthObjectName, ObjNameSize, TempSizeY);
+		float TempSizeY = 0.0f;
+
+		// We have to actually use representative strings because of the kerning
+		CanvasContext.MeasureString(*LongestDebugObjectName, ObjNameSize, TempSizeY);
 		CanvasContext.MeasureString(TEXT("source: "), SourceNameSize, TempSizeY);
 		CanvasContext.MeasureString(TEXT("level: 00"), LevelNameSize, TempSizeY);
 		ObjNameSize += Padding;
@@ -336,15 +564,209 @@ void FGameplayDebuggerCategory_Abilities::DrawGameplayAbilities(FGameplayDebugge
 	const int NumColumns = FMath::Max(1, FMath::FloorToInt(CanvasWidth / ColumnWidth));
 
 	CanvasContext.Printf(TEXT("Gameplay Abilities: \t{yellow}Granted[%d] \t{cyan}Active[%d]"), DataPack.Abilities.Num(), NumActive);
+	CanvasContext.CursorX += Padding;
 	for (const FRepData::FGameplayAbilityDebug& ItemData : DataPack.Abilities)
 	{
 		float CursorX = CanvasContext.CursorX;
 		float CursorY = CanvasContext.CursorY;
 
 		// Print positions manually to align them properly
-		CanvasContext.PrintAt(CursorX + ObjNameSize * 0, CursorY, ItemData.bIsActive ? FColor::Cyan : FColor::Yellow, ItemData.Ability);
-		CanvasContext.PrintAt(CursorX + ObjNameSize * 1, CursorY, FString::Printf(TEXT("{grey}source:{white}%s"), *ItemData.Source));
-		CanvasContext.PrintAt(CursorX + ObjNameSize * 2 + SourceNameSize, CursorY, FString::Printf(TEXT("{grey}level:{white}%02d"), ItemData.Level));
+		CanvasContext.PrintAt(CursorX + ObjNameSize * 0, CursorY, ItemData.bIsActive ? FColor::Cyan : FColor::Yellow, ItemData.Ability.Left(35));
+		CanvasContext.PrintAt(CursorX + ObjNameSize * 1, CursorY, FString::Printf(TEXT("{grey}source: {white}%.35s"), *ItemData.Source));
+		CanvasContext.PrintAt(CursorX + ObjNameSize * 2 + SourceNameSize, CursorY, FString::Printf(TEXT("{grey}level: {white}%02d"), ItemData.Level));
+
+		// PrintAt would have reset these values, restore them.
+		CanvasContext.CursorX = CursorX + (CanvasWidth / NumColumns);
+		CanvasContext.CursorY = CursorY;
+
+		// If we're going to overflow, go to the next line...
+		if (CanvasContext.CursorX + ColumnWidth >= CanvasWidth)
+		{
+			CanvasContext.MoveToNewLine();
+			CanvasContext.CursorX += Padding;
+		}
+	}
+
+	// End the row with a newline
+	if (CanvasContext.CursorX != CanvasContext.DefaultX)
+	{
+		CanvasContext.MoveToNewLine();
+	}
+
+	// End the category with a newline to separate from the other categories
+	CanvasContext.MoveToNewLine();
+}
+
+void FGameplayDebuggerCategory_Abilities::DrawGameplayAttributes(FGameplayDebuggerCanvasContext& CanvasContext, const APlayerController* OwnerPC) const
+{
+	using namespace UE::AbilitySystem::Debug;
+
+	struct FDebugAttributeData
+	{
+		FString AttributeName;
+
+		float ServerBaseValue = 0.0f;
+		float ServerCurrentValue = 0.0f;
+		float LocalBaseValue = 0.0f;
+		float LocalCurrentValue = 0.0f;
+
+		FRepData::ENetworkStatus NetworkStatus = FRepData::ENetworkStatus::LocalOnly;
+	};
+
+	const bool bConsiderNetworkStatus = !OwnerPC->IsNetMode(ENetMode::NM_Standalone);
+	TArray<FRepData::FGameplayAttributeDebug> ServerAttributes = DataPack.Attributes;
+	TArray<FRepData::FGameplayAttributeDebug> LocalAttributes;
+	if (bConsiderNetworkStatus)
+	{
+		if (UAbilitySystemComponent* LocalASC = UAbilitySystemGlobals::GetAbilitySystemComponentFromActor(FindLocalDebugActor()))
+		{
+			LocalAttributes = CollectAttributeData(OwnerPC, LocalASC);
+		}
+	}
+
+	bool bAllAttributesNetworked = true;
+
+	// Reverse the order of iteration so RemoveAt has a better chance at removing near the end
+	TArray<FDebugAttributeData> AttributeDebugData;
+	for (int32 Index = LocalAttributes.Num() - 1 ; Index >= 0 ; --Index)
+	{
+		const FRepData::FGameplayAttributeDebug& LocalAttribute = LocalAttributes[Index];
+
+		FDebugAttributeData& DebugDatum = AttributeDebugData.AddDefaulted_GetRef();
+		DebugDatum.AttributeName = LocalAttribute.AttributeName;
+		DebugDatum.LocalBaseValue = LocalAttribute.BaseValue;
+		DebugDatum.LocalCurrentValue = LocalAttribute.CurrentValue;
+		DebugDatum.NetworkStatus = FRepData::ENetworkStatus::LocalOnly;
+
+		int32 ServerIndex = ServerAttributes.FindLastByPredicate([FindName = LocalAttribute.AttributeName](const FRepData::FGameplayAttributeDebug& Item) { return FindName == Item.AttributeName; });
+		if (ServerIndex != INDEX_NONE)
+		{
+			const FRepData::FGameplayAttributeDebug& ServerAttribute = ServerAttributes[ServerIndex];
+
+			DebugDatum.ServerBaseValue = ServerAttribute.BaseValue;
+			DebugDatum.ServerCurrentValue = ServerAttribute.CurrentValue;
+			DebugDatum.NetworkStatus = ServerAttribute.NetworkStatus == (FRepData::ENetworkStatus::Networked) ? FRepData::ENetworkStatus::Networked : FRepData::ENetworkStatus::Detached;
+
+			// Remove them from server-only array
+			ServerAttributes.RemoveAt(ServerIndex);
+		}
+
+		if (DebugDatum.NetworkStatus != FRepData::ENetworkStatus::Networked)
+			bAllAttributesNetworked = false;
+	}
+
+	// If any entries are still in this array, they are server only and mark them as such
+	for (const FRepData::FGameplayAttributeDebug& ServerAttribute : ServerAttributes)
+	{
+		bAllAttributesNetworked = false;
+
+		FDebugAttributeData& DebugDatum = AttributeDebugData.AddDefaulted_GetRef();
+		DebugDatum.AttributeName = ServerAttribute.AttributeName;
+		DebugDatum.ServerBaseValue = ServerAttribute.BaseValue;
+		DebugDatum.ServerCurrentValue = ServerAttribute.CurrentValue;
+		DebugDatum.NetworkStatus = FRepData::ENetworkStatus::ServerOnly;
+	}
+
+	// Finally sort to keep everything in alphabetical order
+	auto SortByAttributeName = [](const FDebugAttributeData& ItemOne, const FDebugAttributeData& ItemTwo) { return ItemOne.AttributeName < ItemTwo.AttributeName; };
+	AttributeDebugData.Sort(SortByAttributeName);
+
+	// Measure a large string once and save the value, so that we can size the columns properly and not continually resize it
+	static float MaxLineSize = 0.0f;
+	if (MaxLineSize <= 0.0f)
+	{
+		float TempSizeY;
+		const FString MaxLineString = FString::Printf(TEXT("%.35s: 0.000 [0.000]"), *UE::AbilitySystem::Debug::LongestDebugObjectName);
+		CanvasContext.MeasureString(MaxLineString, MaxLineSize, TempSizeY);
+	}
+
+	constexpr float Padding = 10.0f;
+	const float ColumnWidth = MaxLineSize + Padding;
+	const float CanvasWidth = CanvasContext.Canvas->SizeX;
+	const int NumColumns = FMath::Max(1, FMath::FloorToInt(CanvasWidth / ColumnWidth));
+
+	// Figure out the colors we're going to be using and print the legend
+	const TCHAR* ExpectedColor = BothColor;
+	const TCHAR* Modified = TEXT("*");
+	FString HeaderText = FString::Printf(TEXT("Attributes [%d]"), AttributeDebugData.Num());
+	if (!bConsiderNetworkStatus || bAllAttributesNetworked)
+	{
+		// Normal case: no legend required, single color which is the local case
+		ExpectedColor = ServerColor;
+	}
+	else if (!bAllAttributesNetworked)
+	{
+		// Not all attributes are networked; we should display all cases for easier debugging
+		HeaderText.Appendf(TEXT("\tLegend:  %sBoth  %sServer  %sLocal  %sNonReplicated"), BothColor, ServerColor, LocalColor, NonReplicatedColor);
+		ExpectedColor = BothColor;
+	}
+	else
+	{
+		// All attributes are networked; the values may differ but assume the server is correct
+		HeaderText.Appendf(TEXT("\tLegend:  %sServer  %sLocal"), ServerColor, LocalColor);
+		ExpectedColor = ServerColor;
+	}
+
+	const TCHAR* StatusTexts[+FRepData::ENetworkStatus::MAX];
+	StatusTexts[+FRepData::ENetworkStatus::Networked] = ExpectedColor; // This one is going to change based on 'expected' value (e.g. if all local, keep the UI consistently colored)
+	StatusTexts[+FRepData::ENetworkStatus::ServerOnly] = ServerColor;
+	StatusTexts[+FRepData::ENetworkStatus::LocalOnly] = LocalColor;
+	StatusTexts[+FRepData::ENetworkStatus::Detached] = NonReplicatedColor;
+
+	CanvasContext.Print(HeaderText);
+	CanvasContext.CursorX += Padding;
+
+	for (const FDebugAttributeData& AttributeData : AttributeDebugData)
+	{
+		const bool bModified = (AttributeData.ServerBaseValue != AttributeData.ServerCurrentValue) || (AttributeData.LocalBaseValue != AttributeData.LocalCurrentValue);
+		const TCHAR* StatusText = StatusTexts[+AttributeData.NetworkStatus];
+
+		const bool bServerValueMatch = (AttributeData.ServerBaseValue == AttributeData.ServerCurrentValue);
+		const bool bLocalValueMatch = (AttributeData.LocalBaseValue == AttributeData.LocalCurrentValue);
+		const bool bNetworkValueMatch = (AttributeData.LocalBaseValue == AttributeData.ServerBaseValue) && (AttributeData.LocalCurrentValue == AttributeData.ServerCurrentValue);
+
+		// Let's build up the attribute value string which is just trying to represent the four states: srv cur [srv base] local cur [local base] in as little text as possible
+		TStringBuilder<1024> AttributeValueStr;
+		if (bNetworkValueMatch && bServerValueMatch && bLocalValueMatch)
+		{
+			// Everything matches, let's choose white and any one of the values (since they match)
+			AttributeValueStr.Appendf(TEXT("{white}%.4g "), AttributeData.ServerBaseValue);
+		}
+		else
+		{
+			// Append status that is happening on the Server
+			if (AttributeData.NetworkStatus != FRepData::ENetworkStatus::LocalOnly)
+			{
+				if (bServerValueMatch)
+				{
+					AttributeValueStr.Appendf(TEXT("{white}%.4g "), AttributeData.ServerCurrentValue);
+				}
+				else
+				{
+					AttributeValueStr.Appendf(TEXT("{white}%.4g [%.4g] "), AttributeData.ServerCurrentValue, AttributeData.ServerBaseValue);
+				}
+			}
+
+			// Append status that is happening on the Client (locally)
+			if (AttributeData.NetworkStatus != FRepData::ENetworkStatus::ServerOnly)
+			{
+				if (bLocalValueMatch)
+				{
+					AttributeValueStr.Appendf(TEXT("%s%.4g "), LocalColor, AttributeData.LocalCurrentValue);
+				}
+				else
+				{
+					AttributeValueStr.Appendf(TEXT("%s%.4g [%.4g] "), LocalColor, AttributeData.LocalCurrentValue, AttributeData.LocalBaseValue);
+				}
+			}
+		}
+
+		// Print positions manually to align things properly
+		const float CursorX = CanvasContext.CursorX;
+		const float CursorY = CanvasContext.CursorY;
+
+		const FString AttributeDebugText = FString::Printf(TEXT("%s%s%.35s: %s"), bModified ? Modified : TEXT(""), StatusText, *AttributeData.AttributeName, AttributeValueStr.ToString());
+		CanvasContext.PrintAt(CursorX, CursorY, AttributeDebugText);
 
 		// PrintAt would have reset these values, restore them.
 		CanvasContext.CursorX = CursorX + (CanvasWidth / NumColumns);
@@ -354,6 +776,7 @@ void FGameplayDebuggerCategory_Abilities::DrawGameplayAbilities(FGameplayDebugge
 		if (CanvasContext.CursorX + ColumnWidth >= CanvasWidth)
 		{
 			CanvasContext.MoveToNewLine();
+			CanvasContext.CursorX += Padding;
 		}
 	}
 
diff --git a/Engine/Plugins/Runtime/GameplayAbilities/Source/GameplayAbilities/Private/GameplayDebuggerCategory_Abilities.h b/Engine/Plugins/Runtime/GameplayAbilities/Source/GameplayAbilities/Private/GameplayDebuggerCategory_Abilities.h
--- a/Engine/Plugins/Runtime/GameplayAbilities/Source/GameplayAbilities/Private/GameplayDebuggerCategory_Abilities.h
+++ b/Engine/Plugins/Runtime/GameplayAbilities/Source/GameplayAbilities/Private/GameplayDebuggerCategory_Abilities.h
@@ -10,6 +10,7 @@
 
 class AActor;
 class APlayerController;
+class UAbilitySystemComponent;
 class UPackageMap;
 
 class FGameplayDebuggerCategory_Abilities : public FGameplayDebuggerCategory
@@ -25,13 +26,14 @@ public:
 	void OnShowGameplayTagsToggle();
 	void OnShowGameplayAbilitiesToggle();
 	void OnShowGameplayEffectsToggle();
+	void OnShowGameplayAttributesToggle();
 		
 protected:
 
 	void DrawGameplayTags(FGameplayDebuggerCanvasContext& CanvasContext, const APlayerController* OwnerPC) const;
 	void DrawGameplayAbilities(FGameplayDebuggerCanvasContext& CanvasContext, const APlayerController* OwnerPC) const;
 	void DrawGameplayEffects(FGameplayDebuggerCanvasContext& CanvasContext, const APlayerController* OwnerPC) const;
-
+	void DrawGameplayAttributes(FGameplayDebuggerCanvasContext& CanvasContext, const APlayerController* OwnerPC) const;
 
 	struct FRepData
 	{
@@ -39,13 +41,14 @@ protected:
 		TWeakObjectPtr<UPackageMap>	ClientPackageMap;
 
 		FGameplayTagContainer OwnedTags;
+		TArray<int32> TagCounts;
 
 		struct FGameplayAbilityDebug
 		{
 			FString Ability;
 			FString Source;
-			int32 Level;
-			bool bIsActive;
+			int32 Level = 0;
+			bool bIsActive = false;
 		};
 		TArray<FGameplayAbilityDebug> Abilities;
 
@@ -53,23 +56,44 @@ protected:
 		{
 			FString Effect;
 			FString Context;
-			float Duration;
-			float Period;
-			int32 Stacks;
-			float Level;
+			float Duration = 0.0f;
+			float Period = 0.0f;
+			int32 Stacks = 0;
+			float Level = 0.0f;
 		};
 		TArray<FGameplayEffectDebug> GameplayEffects;
 
+		// Attributes can exist purposefully on ServerOnly or LocalOnly.  They are by default Networked on both.  We can get 'detached' if both sides have attributes that aren't networked.
+		enum class ENetworkStatus : uint8
+		{
+			ServerOnly, LocalOnly, Networked, Detached, MAX
+		};
+
+		struct FGameplayAttributeDebug
+		{
+			FString AttributeName;
+			float BaseValue = 0.0f;
+			float CurrentValue = 0.0f;
+			ENetworkStatus NetworkStatus = ENetworkStatus::ServerOnly;
+		};
+		TArray<FGameplayAttributeDebug> Attributes;
+
 		void Serialize(FArchive& Ar);
+
+		// Unary operator + for quick conversion from enum class to int32
+		friend constexpr int32 operator+(const FGameplayDebuggerCategory_Abilities::FRepData::ENetworkStatus& value) { return static_cast<int32>(value); }
 	};
 	FRepData DataPack;
 
 	bool WrapStringAccordingToViewport(const FString& iStr, FString& oStr, FGameplayDebuggerCanvasContext& CanvasContext, float ViewportWitdh) const;
 
 private:
+	TArray<FRepData::FGameplayAttributeDebug> CollectAttributeData(const APlayerController* OwnerPC, const UAbilitySystemComponent* ASC) const;
+
 	bool bShowGameplayTags = true;
 	bool bShowGameplayAbilities = true;
 	bool bShowGameplayEffects = true;
+	bool bShowGameplayAttributes = true;
 };
 
 #endif // WITH_GAMEPLAY_DEBUGGER_MENU
