Gameplay Effects could previously hold on to old GameplayTags because we are copying the (new, upgraded) data into the old, deprecated variables (for backwards compat).  Let's ensure everything is up to date when saving.

diff --git a/Engine/Plugins/Runtime/GameplayAbilities/Source/GameplayAbilities/Private/GameplayEffect.cpp b/Engine/Plugins/Runtime/GameplayAbilities/Source/GameplayAbilities/Private/GameplayEffect.cpp
--- a/Engine/Plugins/Runtime/GameplayAbilities/Source/GameplayAbilities/Private/GameplayEffect.cpp
+++ b/Engine/Plugins/Runtime/GameplayAbilities/Source/GameplayAbilities/Private/GameplayEffect.cpp
@@ -787,6 +787,41 @@ void UGameplayEffect::ConvertUIComponent()
 		UIData = nullptr;
 	}
 }
+
+void UGameplayEffect::PreSave(FObjectPreSaveContext SaveContext)
+{
+	Super::PreSave(SaveContext);
+
+	// Don't need any more data updates because during Cook the data shouldn't have changed
+	if (SaveContext.IsCooking())
+	{
+		return;
+	}
+
+	// Since we've deprecated these fields in favor of GEComponents, we should also manually
+	// clear them out of their gameplay tag references, then rebuild them using the new data, 
+	// otherwise we would end up with stale references when a GEComponent is removed.
+	if (GetVersion() >= EGameplayEffectVersion::Modular53)
+	{
+		InheritableGameplayEffectTags = FInheritedTagContainer{};
+		InheritableOwnedTagsContainer = FInheritedTagContainer{};
+		InheritableBlockedAbilityTagsContainer = FInheritedTagContainer{};
+		OngoingTagRequirements = FGameplayTagRequirements{};
+		ApplicationTagRequirements = FGameplayTagRequirements{};
+		RemovalTagRequirements = FGameplayTagRequirements{};
+		RemoveGameplayEffectsWithTags = FInheritedTagContainer{};
+		GrantedApplicationImmunityTags = FGameplayTagRequirements{};
+		GrantedApplicationImmunityQuery = FGameplayEffectQuery{};
+		RemoveGameplayEffectQuery = FGameplayEffectQuery{};
+		GrantedAbilities.Empty();
+
+		// Now that we've removed all of the stale data, run through the upgrade path again which will copy the new components
+		// back into the deprecated variables.  This allows us to keep for extended backwards compatibility.
+		FPostCDOCompiledContext PostCDOCompiledContext;
+		PostCDOCompiled(PostCDOCompiledContext);
+	}
+}
+
 PRAGMA_ENABLE_DEPRECATION_WARNINGS
 
 /** Let's keep track of the version so that we can upgrade the Components properly.  This will be set properly in PostLoad after upgrades. */
diff --git a/Engine/Plugins/Runtime/GameplayAbilities/Source/GameplayAbilities/Public/GameplayEffect.h b/Engine/Plugins/Runtime/GameplayAbilities/Source/GameplayAbilities/Public/GameplayEffect.h
--- a/Engine/Plugins/Runtime/GameplayAbilities/Source/GameplayAbilities/Public/GameplayEffect.h
+++ b/Engine/Plugins/Runtime/GameplayAbilities/Source/GameplayAbilities/Public/GameplayEffect.h
@@ -2168,6 +2168,9 @@ protected:
 	/** Sets the Version of the class to denote it's been upgraded */
 	void SetVersion(EGameplayEffectVersion Version);
 
+	/** We should intercept the Save call and revalidate all of our deprecated values to avoid hanging onto stale data */
+	virtual void PreSave(FObjectPreSaveContext SaveContext) override;
+
 private:
 	// Helper functions for converting data to use components
 	void ConvertAbilitiesComponent();
