Added warning inside AddGameplayCueOnActor for parameters not replicating in some cases.

#rnx

diff --git a/Engine/Plugins/Runtime/GameplayAbilities/Source/GameplayAbilities/Private/GameplayCueFunctionLibrary.cpp b/Engine/Plugins/Runtime/GameplayAbilities/Source/GameplayAbilities/Private/GameplayCueFunctionLibrary.cpp
--- a/Engine/Plugins/Runtime/GameplayAbilities/Source/GameplayAbilities/Private/GameplayCueFunctionLibrary.cpp
+++ b/Engine/Plugins/Runtime/GameplayAbilities/Source/GameplayAbilities/Private/GameplayCueFunctionLibrary.cpp
@@ -7,6 +7,7 @@
 
 #include UE_INLINE_GENERATED_CPP_BY_NAME(GameplayCueFunctionLibrary)
 
+DEFINE_LOG_CATEGORY_STATIC(LogGameplayCueFunctionLibrary, Log, All);
 
 //////////////////////////////////////////////////////////////////////////
 // UGameplayCueFunctionLibrary
@@ -57,6 +58,11 @@ void UGameplayCueFunctionLibrary::AddGameplayCueOnActor(AActor* Target, const FG
 
 	if (ASC)
 	{
+		if (ASC->ReplicationProxyEnabled && (Parameters != FGameplayCueParameters()))
+		{
+			UE_LOG(LogGameplayCueFunctionLibrary, Warning, TEXT("%hs: Parameters may not replicate when using AbilitySystemComponent with ReplicationProxyEnabled!"), __FUNCTION__);
+		}
+
 		// The actor has an ability system so the event will fire on authority only and will be replicated.
 		if (Target->GetLocalRole() == ROLE_Authority)
 		{
diff --git a/Engine/Plugins/Runtime/GameplayAbilities/Source/GameplayAbilities/Private/GameplayEffectTypes.cpp b/Engine/Plugins/Runtime/GameplayAbilities/Source/GameplayAbilities/Private/GameplayEffectTypes.cpp
--- a/Engine/Plugins/Runtime/GameplayAbilities/Source/GameplayAbilities/Private/GameplayEffectTypes.cpp
+++ b/Engine/Plugins/Runtime/GameplayAbilities/Source/GameplayAbilities/Private/GameplayEffectTypes.cpp
@@ -1034,6 +1034,27 @@ FGameplayCueParameters::FGameplayCueParameters(const struct FGameplayEffectConte
 	UAbilitySystemGlobals::Get().InitGameplayCueParameters(*this, InEffectContext);
 }
 
+bool FGameplayCueParameters::operator==(const FGameplayCueParameters& Other) const
+{
+	return ((NormalizedMagnitude == Other.NormalizedMagnitude) &&
+		    (RawMagnitude == Other.RawMagnitude) &&
+		    (Location == Other.Location) &&
+		    (Normal == Other.Normal) &&
+		    (GameplayEffectLevel == Other.GameplayEffectLevel) &&
+		    (AbilityLevel == Other.AbilityLevel) &&
+		    (EffectContext == Other.EffectContext) &&
+			(MatchedTagName == Other.MatchedTagName) &&
+			(OriginalTag == Other.OriginalTag) &&
+			(AggregatedSourceTags == Other.AggregatedSourceTags) &&
+			(AggregatedTargetTags == Other.AggregatedTargetTags) &&
+			(Instigator == Other.Instigator) &&
+			(EffectCauser == Other.EffectCauser) &&
+			(SourceObject == Other.SourceObject) &&
+			(PhysicalMaterial == Other.PhysicalMaterial) &&
+			(TargetAttachComponent == Other.TargetAttachComponent) &&
+			(bReplicateLocationWhenUsingMinimalRepProxy == Other.bReplicateLocationWhenUsingMinimalRepProxy));
+}
+
 bool FGameplayCueParameters::NetSerialize(FArchive& Ar, class UPackageMap* Map, bool& bOutSuccess)
 {
 	static const uint8 NUM_LEVEL_BITS = 5; // need to bump this up to support 20 levels for AbilityLevel
diff --git a/Engine/Plugins/Runtime/GameplayAbilities/Source/GameplayAbilities/Public/GameplayEffectTypes.h b/Engine/Plugins/Runtime/GameplayAbilities/Source/GameplayAbilities/Public/GameplayEffectTypes.h
--- a/Engine/Plugins/Runtime/GameplayAbilities/Source/GameplayAbilities/Public/GameplayEffectTypes.h
+++ b/Engine/Plugins/Runtime/GameplayAbilities/Source/GameplayAbilities/Public/GameplayEffectTypes.h
@@ -826,6 +826,12 @@ struct GAMEPLAYABILITIES_API FGameplayCueParameters
 	FGameplayCueParameters(const struct FGameplayEffectSpecForRPC &Spec);
 	FGameplayCueParameters(const struct FGameplayEffectContextHandle& EffectContext);
 
+	bool operator==(const FGameplayCueParameters& Other) const;
+	bool operator!=(const FGameplayCueParameters& Other) const
+	{
+		return !(*this == Other);
+	}
+
 	/** Magnitude of source gameplay effect, normalzed from 0-1. Use this for "how strong is the gameplay effect" (0=min, 1=,max) */
 	UPROPERTY(BlueprintReadWrite, Category=GameplayCue)
 	float NormalizedMagnitude;
