Exposing Gameplay Effect's GetAssetTags / GetGrantedTags to Blueprints

#changelist validated
#virtualized

diff --git a/Engine/Plugins/Runtime/GameplayAbilities/Source/GameplayAbilities/Private/AbilitySystemBlueprintLibrary.cpp b/Engine/Plugins/Runtime/GameplayAbilities/Source/GameplayAbilities/Private/AbilitySystemBlueprintLibrary.cpp
--- a/Engine/Plugins/Runtime/GameplayAbilities/Source/GameplayAbilities/Private/AbilitySystemBlueprintLibrary.cpp
+++ b/Engine/Plugins/Runtime/GameplayAbilities/Source/GameplayAbilities/Private/AbilitySystemBlueprintLibrary.cpp
@@ -1151,6 +1151,28 @@ const UGameplayEffect* UAbilitySystemBlueprintLibrary::GetGameplayEffectFromActi
 	return nullptr;
 }
 
+const FGameplayTagContainer& UAbilitySystemBlueprintLibrary::GetGameplayEffectAssetTags(TSubclassOf<UGameplayEffect> EffectClass)
+{
+	if (const UGameplayEffect* DefaultGE = EffectClass.GetDefaultObject())
+	{
+		return DefaultGE->GetAssetTags();
+	}
+
+	static const FGameplayTagContainer Empty{};
+	return Empty;
+}
+
+const FGameplayTagContainer& UAbilitySystemBlueprintLibrary::GetGameplayEffectGrantedTags(TSubclassOf<UGameplayEffect> EffectClass)
+{
+	if (const UGameplayEffect* DefaultGE = EffectClass.GetDefaultObject())
+	{
+		return DefaultGE->GetGrantedTags();
+	}
+
+	static const FGameplayTagContainer Empty{};
+	return Empty;
+}
+
 const UGameplayAbility* UAbilitySystemBlueprintLibrary::GetGameplayAbilityFromSpecHandle(UAbilitySystemComponent* AbilitySystem, const FGameplayAbilitySpecHandle& AbilitySpecHandle, bool& bIsInstance)
 {
 	// validate the ASC
diff --git a/Engine/Plugins/Runtime/GameplayAbilities/Source/GameplayAbilities/Public/AbilitySystemBlueprintLibrary.h b/Engine/Plugins/Runtime/GameplayAbilities/Source/GameplayAbilities/Public/AbilitySystemBlueprintLibrary.h
--- a/Engine/Plugins/Runtime/GameplayAbilities/Source/GameplayAbilities/Public/AbilitySystemBlueprintLibrary.h
+++ b/Engine/Plugins/Runtime/GameplayAbilities/Source/GameplayAbilities/Public/AbilitySystemBlueprintLibrary.h
@@ -412,6 +412,18 @@ class GAMEPLAYABILITIES_API UAbilitySystemBlueprintLibrary : public UBlueprintFu
 	UFUNCTION(BlueprintPure, Category = "Ability|GameplayEffect")
 	static const UGameplayEffect* GetGameplayEffectFromActiveEffectHandle(const FActiveGameplayEffectHandle& ActiveHandle);
 
+	// -------------------------------------------------------------------------------
+	//		Gameplay Effect
+	// -------------------------------------------------------------------------------
+
+	/** Returns all tags that the Gameplay Effect *has* (that denote the GE Asset itself) and *does not* grant to any Actor. */
+	UFUNCTION(BlueprintPure, Category = "Ability|GameplayEffect", meta = (DisplayName = "Get Asset Tags"))
+	static const FGameplayTagContainer& GetGameplayEffectAssetTags(TSubclassOf<UGameplayEffect> EffectClass);
+
+	/** Returns all tags that the Gameplay Effect grants to the target Actor */
+	UFUNCTION(BlueprintPure, Category = "Ability|GameplayEffect", meta = (DisplayName = "Get Granted Tags"))
+	static const FGameplayTagContainer& GetGameplayEffectGrantedTags(TSubclassOf<UGameplayEffect> EffectClass);
+
 	// -------------------------------------------------------------------------------
 	//		GameplayAbility
 	// -------------------------------------------------------------------------------
