PR #10717: Optional parameter to GameplayTagCountContainer::Reset allows a user to empty the container while keeping all of the registered tag count change callbacks.

diff --git a/Engine/Plugins/Runtime/GameplayAbilities/Source/GameplayAbilities/Private/GameplayEffectTypes.cpp b/Engine/Plugins/Runtime/GameplayAbilities/Source/GameplayAbilities/Private/GameplayEffectTypes.cpp
--- a/Engine/Plugins/Runtime/GameplayAbilities/Source/GameplayAbilities/Private/GameplayEffectTypes.cpp
+++ b/Engine/Plugins/Runtime/GameplayAbilities/Source/GameplayAbilities/Private/GameplayEffectTypes.cpp
@@ -507,13 +507,17 @@ FOnGameplayEffectTagCountChanged& FGameplayTagCountContainer::RegisterGameplayTa
 	return Info.OnAnyChange;
 }
 
-void FGameplayTagCountContainer::Reset()
+void FGameplayTagCountContainer::Reset(bool bResetCallbacks)
 {
-	GameplayTagEventMap.Reset();
 	GameplayTagCountMap.Reset();
 	ExplicitTagCountMap.Reset();
 	ExplicitTags.Reset();
-	OnAnyTagChangeDelegate.Clear();
+
+	if (bResetCallbacks)
+	{
+		GameplayTagEventMap.Reset();
+		OnAnyTagChangeDelegate.Clear();
+	}
 }
 
 bool FGameplayTagCountContainer::UpdateExplicitTags(const FGameplayTag& Tag, const int32 CountDelta, const bool bDeferParentTagsOnRemove)
diff --git a/Engine/Plugins/Runtime/GameplayAbilities/Source/GameplayAbilities/Public/GameplayEffectTypes.h b/Engine/Plugins/Runtime/GameplayAbilities/Source/GameplayAbilities/Public/GameplayEffectTypes.h
--- a/Engine/Plugins/Runtime/GameplayAbilities/Source/GameplayAbilities/Public/GameplayEffectTypes.h
+++ b/Engine/Plugins/Runtime/GameplayAbilities/Source/GameplayAbilities/Public/GameplayEffectTypes.h
@@ -1228,7 +1228,12 @@ struct GAMEPLAYABILITIES_API FGameplayTagCountContainer
 		return ExplicitTags;
 	}
 
-	void Reset();
+	/**
+	 * Removes all of the tags. Does not notify any delegates.
+	 * 
+	 * @param bResetCallbacks	If true, also remove all of the registered tag count change delegates
+	 */
+	void Reset(bool bResetCallbacks = true);
 
 	/** Fills in ParentTags from GameplayTags */
 	void FillParentTags()
