Slightly modified version of PR #10711: Create GetExplicitTagCoung util for FGameplayTagCountContainer

Added some Test cases so I can verify the comments are correct.

diff --git a/Engine/Plugins/Runtime/GameplayAbilities/Source/GameplayAbilities/Private/Tests/GameplayTagCountContainerTests.cpp b/Engine/Plugins/Runtime/GameplayAbilities/Source/GameplayAbilities/Private/Tests/GameplayTagCountContainerTests.cpp
new file mode 100644
--- /dev/null
+++ b/Engine/Plugins/Runtime/GameplayAbilities/Source/GameplayAbilities/Private/Tests/GameplayTagCountContainerTests.cpp
@@ -0,0 +1,65 @@
+// Copyright Epic Games, Inc. All Rights Reserved.
+
+#include "Misc/AutomationTest.h"
+#include "GameplayEffectTypes.h"
+
+IMPLEMENT_SIMPLE_AUTOMATION_TEST(FGameplayTagCountContainerTests, "System.AbilitySystem.GameplayTagCountContainer", EAutomationTestFlags::ApplicationContextMask | EAutomationTestFlags::EngineFilter)
+
+#include "NativeGameplayTags.h"
+
+namespace UE::Private::GameplayTagCountContainerTests
+{
+	UE_DEFINE_GAMEPLAY_TAG_STATIC(TestsDotGenericTag, "Tests.GenericTag");
+	UE_DEFINE_GAMEPLAY_TAG_STATIC(TestsDotGenericTagDotOne, "Tests.GenericTag.One");
+	UE_DEFINE_GAMEPLAY_TAG_STATIC(TestsDotGenericTagDotTwo, "Tests.GenericTag.Two");
+}
+
+bool FGameplayTagCountContainerTests::RunTest(const FString& Parameters)
+{
+	using namespace UE::Private::GameplayTagCountContainerTests;
+
+	// Create the container and add the first tag
+	FGameplayTagCountContainer TagCountContainer;
+	TagCountContainer.SetTagCount(TestsDotGenericTagDotOne, 1);
+
+	// Test for the first tag after adding
+	FGameplayTagContainer ContainerOne{ TestsDotGenericTagDotOne };
+	TestTrue(TEXT("HasAllMatchingGameplayTags after Adding"), TagCountContainer.HasAllMatchingGameplayTags(ContainerOne));
+	TestTrue(TEXT("HasAnyMatchingGameplayTags after Adding"), TagCountContainer.HasAnyMatchingGameplayTags(ContainerOne));
+	TestTrue(TEXT("HasMatchingGameplayTag after Adding"), TagCountContainer.HasMatchingGameplayTag(TestsDotGenericTagDotOne));
+	TestTrue(TEXT("GetTagCount of One == 1"), TagCountContainer.GetTagCount(TestsDotGenericTagDotOne) == 1);
+	TestTrue(TEXT("GetTagCount of Parent == 1"), TagCountContainer.GetTagCount(TestsDotGenericTag) == 1);
+
+	// False for the second tag before adding
+	FGameplayTagContainer ContainerTwo{ TestsDotGenericTagDotTwo };
+	TestFalse(TEXT("HasAllMatchingGameplayTags before Adding"), TagCountContainer.HasAllMatchingGameplayTags(ContainerTwo));
+	TestFalse(TEXT("HasAnyMatchingGameplayTags before Adding"), TagCountContainer.HasAnyMatchingGameplayTags(ContainerTwo));
+	TestFalse(TEXT("HasMatchingGameplayTag before Adding"), TagCountContainer.HasMatchingGameplayTag(TestsDotGenericTagDotTwo));
+
+	// True for the second tag after adding
+	TagCountContainer.SetTagCount(TestsDotGenericTagDotTwo, 2);
+	TestTrue(TEXT("GetTagCount of One == 1"), TagCountContainer.GetTagCount(TestsDotGenericTagDotOne) == 1);
+	TestTrue(TEXT("GetTagCount of Two == 2"), TagCountContainer.GetTagCount(TestsDotGenericTagDotTwo) == 2);
+	TestTrue(TEXT("GetTagCount of Parent == 3"), TagCountContainer.GetTagCount(TestsDotGenericTag) == 3);
+
+	// Remove one
+	TagCountContainer.UpdateTagCount(TestsDotGenericTagDotOne, -1);
+	TestFalse(TEXT("HasAllMatchingGameplayTags after Removing"), TagCountContainer.HasAllMatchingGameplayTags(ContainerOne));
+	TestFalse(TEXT("HasAnyMatchingGameplayTags after Removing"), TagCountContainer.HasAnyMatchingGameplayTags(ContainerOne));
+	TestTrue(TEXT("GetTagCount of One == 0"), TagCountContainer.GetTagCount(TestsDotGenericTagDotOne) == 0);
+	TestTrue(TEXT("GetTagCount of Two == 2"), TagCountContainer.GetTagCount(TestsDotGenericTagDotTwo) == 2);
+	TestTrue(TEXT("GetTagCount of Parent == 2"), TagCountContainer.GetTagCount(TestsDotGenericTag) == 2);
+
+	// Explicit Tags
+	TestTrue(TEXT("GetExplicitTagCount of Parent == 0"), TagCountContainer.GetExplicitTagCount(TestsDotGenericTag) == 0);
+	TestTrue(TEXT("GetExplicitTagCount of One == 0"), TagCountContainer.GetExplicitTagCount(TestsDotGenericTagDotOne) == 0);
+	TestTrue(TEXT("GetExplicitTagCount of Two == 2"), TagCountContainer.GetExplicitTagCount(TestsDotGenericTagDotTwo) == 2);
+
+	// Test the Explicit Tag Container
+	FGameplayTagContainer ExplicitTags = TagCountContainer.GetExplicitGameplayTags();
+	TestFalse(TEXT("GetExplicitGameplayTags Contains Parent == False"), ExplicitTags.HasTagExact(TestsDotGenericTag));
+	TestFalse(TEXT("GetExplicitGameplayTags Contains One == False"), ExplicitTags.HasTagExact(TestsDotGenericTagDotOne));
+	TestTrue(TEXT("GetExplicitGameplayTags Contains Two == True"), ExplicitTags.HasTagExact(TestsDotGenericTagDotTwo));
+
+	return true;
+}
diff --git a/Engine/Plugins/Runtime/GameplayAbilities/Source/GameplayAbilities/Public/GameplayEffectTypes.h b/Engine/Plugins/Runtime/GameplayAbilities/Source/GameplayAbilities/Public/GameplayEffectTypes.h
--- a/Engine/Plugins/Runtime/GameplayAbilities/Source/GameplayAbilities/Public/GameplayEffectTypes.h
+++ b/Engine/Plugins/Runtime/GameplayAbilities/Source/GameplayAbilities/Public/GameplayEffectTypes.h
@@ -1181,7 +1181,8 @@ struct GAMEPLAYABILITIES_API FGameplayTagCountContainer
 	}
 
 	/**
-	* return the count for a specified tag 
+	* return the hierarchical count for a specified tag
+	* e.g. if A.B & A.C were added, GetTagCount("A") would return 2.
 	*
 	* @param Tag			Tag to update
 	*
@@ -1197,6 +1198,24 @@ struct GAMEPLAYABILITIES_API FGameplayTagCountContainer
 		return 0;
 	}
 
+	/**
+	* return how many times the exact specified tag has been added to the container (ignores the tag hierarchy)
+	* e.g. if A.B & A.C were added, GetExplicitTagCount("A") would return 0, and GetExplicitTagCount("A.B") would return 1.
+	*
+	* @param Tag			Tag to update
+	*
+	* @return the count of the passed in tag
+	*/
+	FORCEINLINE int32 GetExplicitTagCount(const FGameplayTag& Tag) const
+	{
+		if (const int32* Ptr = ExplicitTagCountMap.Find(Tag))
+		{
+			return *Ptr;
+		}
+
+		return 0;
+	}
+
 	/**
 	 *	Broadcasts the AnyChange event for this tag. This is called when the stack count of the backing gameplay effect change.
 	 *	It is up to the receiver of the broadcasted delegate to decide what to do with this.
