Fixing up some user-facing descriptions of FInheritableTagContainer.  Inspired by PR #11406.

diff --git a/Engine/Plugins/Runtime/GameplayAbilities/Source/GameplayAbilities/Public/GameplayEffect.h b/Engine/Plugins/Runtime/GameplayAbilities/Source/GameplayAbilities/Public/GameplayEffect.h
--- a/Engine/Plugins/Runtime/GameplayAbilities/Source/GameplayAbilities/Public/GameplayEffect.h
+++ b/Engine/Plugins/Runtime/GameplayAbilities/Source/GameplayAbilities/Public/GameplayEffect.h
@@ -615,15 +615,15 @@ struct GAMEPLAYABILITIES_API FInheritedTagContainer
 	GENERATED_USTRUCT_BODY()
 
 	/** Tags that I inherited and tags that I added minus tags that I removed */
-	UPROPERTY(VisibleAnywhere, BlueprintReadOnly, Category = Application)
+	UPROPERTY(VisibleAnywhere, BlueprintReadOnly, Category = Application, meta=(Tooltip="CombinedTags = Inherited - Removed + Added"))
 	FGameplayTagContainer CombinedTags;
 
-	/** Tags that I have in addition to my parent's tags */
-	UPROPERTY(EditDefaultsOnly, BlueprintReadOnly, Category = Application)
+	/** Tags that I have (in addition to my parent's tags) */
+	UPROPERTY(EditDefaultsOnly, BlueprintReadOnly, Category = Application, meta=(DisplayName="Add to Inherited"))
 	FGameplayTagContainer Added;
 
-	/** Tags that should be removed if my parent had them */
-	UPROPERTY(EditDefaultsOnly, BlueprintReadOnly, Category = Application)
+	/** Tags that should be removed (only if my parent had them).  Note: we cannot use this to remove a tag that exists on a target. It only modifies the result of CombinedTags. */
+	UPROPERTY(EditDefaultsOnly, BlueprintReadOnly, Category = Application, meta=(DisplayName="Remove from Inherited"))
 	FGameplayTagContainer Removed;
 
 	void UpdateInheritedTagProperties(const FInheritedTagContainer* Parent);
