Changed the friendly DisplayNames of the GameplayEffectComponents to make them easier to comprehend.  Inspired by PR #11357.

diff --git a/Engine/Plugins/Runtime/GameplayAbilities/Source/GameplayAbilities/Public/GameplayEffectComponents/AbilitiesGameplayEffectComponent.h b/Engine/Plugins/Runtime/GameplayAbilities/Source/GameplayAbilities/Public/GameplayEffectComponents/AbilitiesGameplayEffectComponent.h
--- a/Engine/Plugins/Runtime/GameplayAbilities/Source/GameplayAbilities/Public/GameplayEffectComponents/AbilitiesGameplayEffectComponent.h
+++ b/Engine/Plugins/Runtime/GameplayAbilities/Source/GameplayAbilities/Public/GameplayEffectComponents/AbilitiesGameplayEffectComponent.h
@@ -32,7 +32,7 @@ struct FGameplayAbilitySpecConfig
 /**
  * Grants additional Gameplay Abilities to the Target of a Gameplay Effect while active
  */
-UCLASS()
+UCLASS(DisplayName="Grant Gameplay Abilities")
 class GAMEPLAYABILITIES_API UAbilitiesGameplayEffectComponent : public UGameplayEffectComponent
 {
 	friend class UGameplayEffect; // for upgrade path
diff --git a/Engine/Plugins/Runtime/GameplayAbilities/Source/GameplayAbilities/Public/GameplayEffectComponents/AdditionalEffectsGameplayEffectComponent.h b/Engine/Plugins/Runtime/GameplayAbilities/Source/GameplayAbilities/Public/GameplayEffectComponents/AdditionalEffectsGameplayEffectComponent.h
--- a/Engine/Plugins/Runtime/GameplayAbilities/Source/GameplayAbilities/Public/GameplayEffectComponents/AdditionalEffectsGameplayEffectComponent.h
+++ b/Engine/Plugins/Runtime/GameplayAbilities/Source/GameplayAbilities/Public/GameplayEffectComponents/AdditionalEffectsGameplayEffectComponent.h
@@ -7,7 +7,7 @@
 #include "AdditionalEffectsGameplayEffectComponent.generated.h"
 
 /** Add additional Gameplay Effects that attempt to activate under certain conditions (or no conditions) */
-UCLASS(CollapseCategories)
+UCLASS(CollapseCategories, DisplayName="Apply Additional Effects")
 class GAMEPLAYABILITIES_API UAdditionalEffectsGameplayEffectComponent : public UGameplayEffectComponent
 {
 	GENERATED_BODY()
diff --git a/Engine/Plugins/Runtime/GameplayAbilities/Source/GameplayAbilities/Public/GameplayEffectComponents/AssetTagsGameplayEffectComponent.h b/Engine/Plugins/Runtime/GameplayAbilities/Source/GameplayAbilities/Public/GameplayEffectComponents/AssetTagsGameplayEffectComponent.h
--- a/Engine/Plugins/Runtime/GameplayAbilities/Source/GameplayAbilities/Public/GameplayEffectComponents/AssetTagsGameplayEffectComponent.h
+++ b/Engine/Plugins/Runtime/GameplayAbilities/Source/GameplayAbilities/Public/GameplayEffectComponents/AssetTagsGameplayEffectComponent.h
@@ -7,7 +7,7 @@
 #include "AssetTagsGameplayEffectComponent.generated.h"
 
 /** These are tags that the Gameplay Effect Asset itself 'has' (owns). These do _not_ transfer to any Actors */
-UCLASS()
+UCLASS(DisplayName="Tags This Effect Has (Asset Tags)")
 class GAMEPLAYABILITIES_API UAssetTagsGameplayEffectComponent : public UGameplayEffectComponent
 {
 	GENERATED_BODY()
diff --git a/Engine/Plugins/Runtime/GameplayAbilities/Source/GameplayAbilities/Public/GameplayEffectComponents/BlockAbilityTagsGameplayEffectComponent.h b/Engine/Plugins/Runtime/GameplayAbilities/Source/GameplayAbilities/Public/GameplayEffectComponents/BlockAbilityTagsGameplayEffectComponent.h
--- a/Engine/Plugins/Runtime/GameplayAbilities/Source/GameplayAbilities/Public/GameplayEffectComponents/BlockAbilityTagsGameplayEffectComponent.h
+++ b/Engine/Plugins/Runtime/GameplayAbilities/Source/GameplayAbilities/Public/GameplayEffectComponents/BlockAbilityTagsGameplayEffectComponent.h
@@ -7,7 +7,7 @@
 #include "BlockAbilityTagsGameplayEffectComponent.generated.h"
 
 /** Handles blocking the activation of Gameplay Abilities based on Gameplay Tags for the Target actor of the owner Gameplay Effect */
-UCLASS()
+UCLASS(DisplayName="Block Abilities with Tags")
 class GAMEPLAYABILITIES_API UBlockAbilityTagsGameplayEffectComponent : public UGameplayEffectComponent
 {
 	GENERATED_BODY()
diff --git a/Engine/Plugins/Runtime/GameplayAbilities/Source/GameplayAbilities/Public/GameplayEffectComponents/ChanceToApplyGameplayEffectComponent.h b/Engine/Plugins/Runtime/GameplayAbilities/Source/GameplayAbilities/Public/GameplayEffectComponents/ChanceToApplyGameplayEffectComponent.h
--- a/Engine/Plugins/Runtime/GameplayAbilities/Source/GameplayAbilities/Public/GameplayEffectComponents/ChanceToApplyGameplayEffectComponent.h
+++ b/Engine/Plugins/Runtime/GameplayAbilities/Source/GameplayAbilities/Public/GameplayEffectComponents/ChanceToApplyGameplayEffectComponent.h
@@ -8,7 +8,7 @@
 #include "ChanceToApplyGameplayEffectComponent.generated.h"
 
 /** Applies a probablity to the application conditions of the Gameplay Effect. */
-UCLASS()
+UCLASS(DisplayName="Chance To Apply This Effect")
 class GAMEPLAYABILITIES_API UChanceToApplyGameplayEffectComponent : public UGameplayEffectComponent
 {
 	GENERATED_BODY()
diff --git a/Engine/Plugins/Runtime/GameplayAbilities/Source/GameplayAbilities/Public/GameplayEffectComponents/CustomCanApplyGameplayEffectComponent.h b/Engine/Plugins/Runtime/GameplayAbilities/Source/GameplayAbilities/Public/GameplayEffectComponents/CustomCanApplyGameplayEffectComponent.h
--- a/Engine/Plugins/Runtime/GameplayAbilities/Source/GameplayAbilities/Public/GameplayEffectComponents/CustomCanApplyGameplayEffectComponent.h
+++ b/Engine/Plugins/Runtime/GameplayAbilities/Source/GameplayAbilities/Public/GameplayEffectComponents/CustomCanApplyGameplayEffectComponent.h
@@ -8,7 +8,7 @@
 #include "CustomCanApplyGameplayEffectComponent.generated.h"
 
 /** Handles configuration of a CustomApplicationRequirement function to see if this GameplayEffect should apply */
-UCLASS()
+UCLASS(DisplayName="Custom Can Apply This Effect")
 class GAMEPLAYABILITIES_API UCustomCanApplyGameplayEffectComponent : public UGameplayEffectComponent
 {
 	GENERATED_BODY()
diff --git a/Engine/Plugins/Runtime/GameplayAbilities/Source/GameplayAbilities/Public/GameplayEffectComponents/ImmunityGameplayEffectComponent.h b/Engine/Plugins/Runtime/GameplayAbilities/Source/GameplayAbilities/Public/GameplayEffectComponents/ImmunityGameplayEffectComponent.h
--- a/Engine/Plugins/Runtime/GameplayAbilities/Source/GameplayAbilities/Public/GameplayEffectComponents/ImmunityGameplayEffectComponent.h
+++ b/Engine/Plugins/Runtime/GameplayAbilities/Source/GameplayAbilities/Public/GameplayEffectComponents/ImmunityGameplayEffectComponent.h
@@ -10,7 +10,7 @@
  * Immunity is blocking the application of _other_ GameplayEffectSpecs.
  * This registers a global handler on the ASC to block the application of other GameplayEffectSpecs.
  */
-UCLASS()
+UCLASS(DisplayName="Immunity to Other Effects")
 class GAMEPLAYABILITIES_API UImmunityGameplayEffectComponent : public UGameplayEffectComponent
 {
 	friend class UGameplayEffect; // Needed for upgrade path
diff --git a/Engine/Plugins/Runtime/GameplayAbilities/Source/GameplayAbilities/Public/GameplayEffectComponents/RemoveOtherGameplayEffectComponent.h b/Engine/Plugins/Runtime/GameplayAbilities/Source/GameplayAbilities/Public/GameplayEffectComponents/RemoveOtherGameplayEffectComponent.h
--- a/Engine/Plugins/Runtime/GameplayAbilities/Source/GameplayAbilities/Public/GameplayEffectComponents/RemoveOtherGameplayEffectComponent.h
+++ b/Engine/Plugins/Runtime/GameplayAbilities/Source/GameplayAbilities/Public/GameplayEffectComponents/RemoveOtherGameplayEffectComponent.h
@@ -9,7 +9,7 @@
 struct FGameplayEffectRemovalInfo;
 
 /** Remove other Gameplay Effects based on certain conditions */
-UCLASS()
+UCLASS(DisplayName="Remove Other Effects")
 class GAMEPLAYABILITIES_API URemoveOtherGameplayEffectComponent : public UGameplayEffectComponent
 {
 	GENERATED_BODY()
diff --git a/Engine/Plugins/Runtime/GameplayAbilities/Source/GameplayAbilities/Public/GameplayEffectComponents/TargetTagRequirementsGameplayEffectComponent.h b/Engine/Plugins/Runtime/GameplayAbilities/Source/GameplayAbilities/Public/GameplayEffectComponents/TargetTagRequirementsGameplayEffectComponent.h
--- a/Engine/Plugins/Runtime/GameplayAbilities/Source/GameplayAbilities/Public/GameplayEffectComponents/TargetTagRequirementsGameplayEffectComponent.h
+++ b/Engine/Plugins/Runtime/GameplayAbilities/Source/GameplayAbilities/Public/GameplayEffectComponents/TargetTagRequirementsGameplayEffectComponent.h
@@ -9,7 +9,7 @@
 struct FGameplayEffectRemovalInfo;
 
 /** Specifies tag requirements that the Target (owner of the Gameplay Effect) must have if this GE should apply or continue to execute */
-UCLASS()
+UCLASS(DisplayName="Require Tags to Apply/Continue This Effect")
 class GAMEPLAYABILITIES_API UTargetTagRequirementsGameplayEffectComponent : public UGameplayEffectComponent
 {
 	GENERATED_BODY()
@@ -39,7 +39,7 @@ private:
 	void OnTagChanged(const FGameplayTag GameplayTag, int32 NewCount, FActiveGameplayEffectHandle ActiveGEHandle) const;
 
 public:
-	/** Tag requirements for this GameplayEffect to be applied to a target. This is pass/fail at the time of application. If fail, this GE fails to apply. */
+	/** Tag requirements the target must have for this GameplayEffect to be applied. This is pass/fail at the time of application. If fail, this GE fails to apply. */
 	UPROPERTY(EditDefaultsOnly, Category = Tags)
 	FGameplayTagRequirements ApplicationTagRequirements;
 
diff --git a/Engine/Plugins/Runtime/GameplayAbilities/Source/GameplayAbilities/Public/GameplayEffectComponents/TargetTagsGameplayEffectComponent.h b/Engine/Plugins/Runtime/GameplayAbilities/Source/GameplayAbilities/Public/GameplayEffectComponents/TargetTagsGameplayEffectComponent.h
--- a/Engine/Plugins/Runtime/GameplayAbilities/Source/GameplayAbilities/Public/GameplayEffectComponents/TargetTagsGameplayEffectComponent.h
+++ b/Engine/Plugins/Runtime/GameplayAbilities/Source/GameplayAbilities/Public/GameplayEffectComponents/TargetTagsGameplayEffectComponent.h
@@ -7,7 +7,7 @@
 #include "TargetTagsGameplayEffectComponent.generated.h"
 
 /** Handles granting Tags to the Target (sometimes referred to as the Owner) of the Gameplay Effect */
-UCLASS()
+UCLASS(DisplayName="Grant Tags to Target Actor")
 class GAMEPLAYABILITIES_API UTargetTagsGameplayEffectComponent : public UGameplayEffectComponent
 {
 	GENERATED_BODY()
diff --git a/Engine/Plugins/Runtime/GameplayAbilities/Source/GameplayAbilities/Public/GameplayEffectUIData_TextOnly.h b/Engine/Plugins/Runtime/GameplayAbilities/Source/GameplayAbilities/Public/GameplayEffectUIData_TextOnly.h
--- a/Engine/Plugins/Runtime/GameplayAbilities/Source/GameplayAbilities/Public/GameplayEffectUIData_TextOnly.h
+++ b/Engine/Plugins/Runtime/GameplayAbilities/Source/GameplayAbilities/Public/GameplayEffectUIData_TextOnly.h
@@ -9,7 +9,7 @@
  * UI data that contains only text. This is mostly used as an example of a subclass of UGameplayEffectUIData.
  * If your game needs only text, this is a reasonable class to use. To include more data, make a custom subclass of UGameplayEffectUIData.
  */
-UCLASS()
+UCLASS(DisplayName="UI Data (Text Only)")
 class GAMEPLAYABILITIES_API UGameplayEffectUIData_TextOnly : public UGameplayEffectUIData
 {
 	GENERATED_BODY()
