Warn against FPostReplicatedReceiveParameters::bHasMoreUnmappedReferences in Gameplay Effects replication.

We had an issue earlier in development when Iris had set this flag and it caused GameplayCues to stop firing.  We just want to catch this earlier.

diff --git a/Engine/Plugins/Runtime/GameplayAbilities/Source/GameplayAbilities/Private/GameplayEffect.cpp b/Engine/Plugins/Runtime/GameplayAbilities/Source/GameplayAbilities/Private/GameplayEffect.cpp
--- a/Engine/Plugins/Runtime/GameplayAbilities/Source/GameplayAbilities/Private/GameplayEffect.cpp
+++ b/Engine/Plugins/Runtime/GameplayAbilities/Source/GameplayAbilities/Private/GameplayEffect.cpp
@@ -4863,11 +4863,20 @@ void FActiveGameplayEffectsContainer::PostReplicatedReceive(const FFastArraySeri
 
 		if (!Parameters.bHasMoreUnmappedReferences) // Do not invoke GCs when we have missing information (like AActor*s in EffectContext)
 		{
+			NumConsecutiveUnmappedReferencesDebug = 0;
 			if (Owner->IsReadyForGameplayCues())
 			{
 				Owner->HandleDeferredGameplayCues(this);
 			}
 		}
+		else
+		{
+			++NumConsecutiveUnmappedReferencesDebug;
+
+			constexpr uint32 HighNumberOfConsecutiveUnmappedRefs = 30;
+			ensureMsgf(NumConsecutiveUnmappedReferencesDebug < HighNumberOfConsecutiveUnmappedRefs, TEXT("%hs: bHasMoreUnmappedReferences is preventing GameplayCues from firing"), __func__);
+			UE_CLOG((NumConsecutiveUnmappedReferencesDebug % HighNumberOfConsecutiveUnmappedRefs) == 0, LogAbilitySystem, Error, TEXT("%hs: bHasMoreUnmappedReferences is preventing GameplayCues from firing (%u consecutive misses)"), __func__, NumConsecutiveUnmappedReferencesDebug);
+		}
 	}
 }
 
diff --git a/Engine/Plugins/Runtime/GameplayAbilities/Source/GameplayAbilities/Public/GameplayEffect.h b/Engine/Plugins/Runtime/GameplayAbilities/Source/GameplayAbilities/Public/GameplayEffect.h
--- a/Engine/Plugins/Runtime/GameplayAbilities/Source/GameplayAbilities/Public/GameplayEffect.h
+++ b/Engine/Plugins/Runtime/GameplayAbilities/Source/GameplayAbilities/Public/GameplayEffect.h
@@ -1942,6 +1942,7 @@ private:
 
 	mutable int32 ScopedLockCount;
 	int32 PendingRemoves;
+	uint32 NumConsecutiveUnmappedReferencesDebug = 0;
 
 	FActiveGameplayEffect*	PendingGameplayEffectHead;	// Head of pending GE linked list
 	FActiveGameplayEffect** PendingGameplayEffectNext;	// Points to the where to store the next pending GE (starts pointing at head, as more are added, points further down the list).
