UGameplayEffectComponent::OnGameplayEffectChanged has been changed to non-const to allow the proper flow of updating its values on Post Blueprint Compilation.

This allows us to update the inheritance values on Blueprint compilation which will happen after serialization (which could stomp PostInitProperties values).

diff --git a/Engine/Plugins/Runtime/GameplayAbilities/Source/GameplayAbilities/Private/GameplayEffect.cpp b/Engine/Plugins/Runtime/GameplayAbilities/Source/GameplayAbilities/Private/GameplayEffect.cpp
--- a/Engine/Plugins/Runtime/GameplayAbilities/Source/GameplayAbilities/Private/GameplayEffect.cpp
+++ b/Engine/Plugins/Runtime/GameplayAbilities/Source/GameplayAbilities/Private/GameplayEffect.cpp
@@ -336,6 +336,10 @@ void UGameplayEffect::OnGameplayEffectChanged()
 			// Ensure the SubObject is fully loaded
 			GEComponent->ConditionalPostLoad();
 			GEComponent->OnGameplayEffectChanged();
+
+			PRAGMA_DISABLE_DEPRECATION_WARNINGS
+			const_cast<const UGameplayEffectComponent*>(GEComponent)->OnGameplayEffectChanged();
+			PRAGMA_ENABLE_DEPRECATION_WARNINGS
 		}
 	}
 }
diff --git a/Engine/Plugins/Runtime/GameplayAbilities/Source/GameplayAbilities/Private/GameplayEffectComponents/AssetTagsGameplayEffectComponent.cpp b/Engine/Plugins/Runtime/GameplayAbilities/Source/GameplayAbilities/Private/GameplayEffectComponents/AssetTagsGameplayEffectComponent.cpp
--- a/Engine/Plugins/Runtime/GameplayAbilities/Source/GameplayAbilities/Private/GameplayEffectComponents/AssetTagsGameplayEffectComponent.cpp
+++ b/Engine/Plugins/Runtime/GameplayAbilities/Source/GameplayAbilities/Private/GameplayEffectComponents/AssetTagsGameplayEffectComponent.cpp
@@ -25,10 +25,10 @@ void UAssetTagsGameplayEffectComponent::PostInitProperties()
 #endif // WITH_EDITORONLY_DATA
 }
 
-void UAssetTagsGameplayEffectComponent::OnGameplayEffectChanged() const
+void UAssetTagsGameplayEffectComponent::OnGameplayEffectChanged()
 {
 	Super::OnGameplayEffectChanged();
-	ApplyAssetTagChanges();
+	SetAndApplyAssetTagChanges(InheritableAssetTags);
 }
 
 #if WITH_EDITOR
@@ -38,8 +38,6 @@ void UAssetTagsGameplayEffectComponent::PostEditChangeProperty(FPropertyChangedE
 
 	if (PropertyChangedEvent.GetMemberPropertyName()  == GetInheritableAssetTagsName())
 	{
-		SetAndApplyAssetTagChanges(InheritableAssetTags);
-		
 		// Tell the GE it needs to reconfigure itself based on these updated properties (this will reaggregate the tags)
 		UGameplayEffect* Owner = GetOwner();
 		Owner->OnGameplayEffectChanged();
diff --git a/Engine/Plugins/Runtime/GameplayAbilities/Source/GameplayAbilities/Private/GameplayEffectComponents/BlockAbilityTagsGameplayEffectComponent.cpp b/Engine/Plugins/Runtime/GameplayAbilities/Source/GameplayAbilities/Private/GameplayEffectComponents/BlockAbilityTagsGameplayEffectComponent.cpp
--- a/Engine/Plugins/Runtime/GameplayAbilities/Source/GameplayAbilities/Private/GameplayEffectComponents/BlockAbilityTagsGameplayEffectComponent.cpp
+++ b/Engine/Plugins/Runtime/GameplayAbilities/Source/GameplayAbilities/Private/GameplayEffectComponents/BlockAbilityTagsGameplayEffectComponent.cpp
@@ -16,10 +16,10 @@ void UBlockAbilityTagsGameplayEffectComponent::PostInitProperties()
 #endif
 }
 
-void UBlockAbilityTagsGameplayEffectComponent::OnGameplayEffectChanged() const
+void UBlockAbilityTagsGameplayEffectComponent::OnGameplayEffectChanged()
 {
 	Super::OnGameplayEffectChanged();
-	ApplyBlockedAbilityTagChanges();
+	SetAndApplyBlockedAbilityTagChanges(InheritableBlockedAbilityTagsContainer);
 }
 
 #if WITH_EDITOR
@@ -29,8 +29,6 @@ void UBlockAbilityTagsGameplayEffectComponent::PostEditChangeProperty(FPropertyC
 
 	if (PropertyChangedEvent.GetMemberPropertyName() == GetInheritableBlockedAbilityTagsContainerPropertyName())
 	{
-		SetAndApplyBlockedAbilityTagChanges(InheritableBlockedAbilityTagsContainer);
-		
 		// Tell the GE it needs to reconfigure itself based on these updated properties (this will reaggregate the tags)
 		UGameplayEffect* Owner = GetOwner();
 		Owner->OnGameplayEffectChanged();
diff --git a/Engine/Plugins/Runtime/GameplayAbilities/Source/GameplayAbilities/Private/GameplayEffectComponents/TargetTagsGameplayEffectComponent.cpp b/Engine/Plugins/Runtime/GameplayAbilities/Source/GameplayAbilities/Private/GameplayEffectComponents/TargetTagsGameplayEffectComponent.cpp
--- a/Engine/Plugins/Runtime/GameplayAbilities/Source/GameplayAbilities/Private/GameplayEffectComponents/TargetTagsGameplayEffectComponent.cpp
+++ b/Engine/Plugins/Runtime/GameplayAbilities/Source/GameplayAbilities/Private/GameplayEffectComponents/TargetTagsGameplayEffectComponent.cpp
@@ -19,10 +19,10 @@ void UTargetTagsGameplayEffectComponent::PostInitProperties()
 #endif // WITH_EDITORONLY_DATA
 }
 
-void UTargetTagsGameplayEffectComponent::OnGameplayEffectChanged() const
+void UTargetTagsGameplayEffectComponent::OnGameplayEffectChanged()
 {
 	Super::OnGameplayEffectChanged();
-	ApplyTargetTagChanges();
+	SetAndApplyTargetTagChanges(InheritableGrantedTagsContainer);
 }
 
 #if WITH_EDITOR
@@ -33,8 +33,6 @@ void UTargetTagsGameplayEffectComponent::PostEditChangeProperty(FPropertyChanged
 
 	if (PropertyChangedEvent.GetMemberPropertyName() == GetInheritableGrantedTagsContainerName())
 	{
-		SetAndApplyTargetTagChanges(InheritableGrantedTagsContainer);
-
 		// Tell the GE it needs to reconfigure itself based on these updated properties (this will reaggregate the tags)
 		UGameplayEffect* Owner = GetOwner();
 		Owner->OnGameplayEffectChanged();
diff --git a/Engine/Plugins/Runtime/GameplayAbilities/Source/GameplayAbilities/Public/GameplayEffectComponent.h b/Engine/Plugins/Runtime/GameplayAbilities/Source/GameplayAbilities/Public/GameplayEffectComponent.h
--- a/Engine/Plugins/Runtime/GameplayAbilities/Source/GameplayAbilities/Public/GameplayEffectComponent.h
+++ b/Engine/Plugins/Runtime/GameplayAbilities/Source/GameplayAbilities/Public/GameplayEffectComponent.h
@@ -66,6 +66,9 @@ public:
 	/**
 	 * Let us know that the owning GameplayEffect has been modified, thus apply an asset-related changes to the owning GameplayEffect (e.g. any of its fields)
 	 */
+	virtual void OnGameplayEffectChanged() {}
+
+	UE_DEPRECATED(5.4, "Use OnGameplayEffectChanged without const -- we often want to cache data on GEComponent when the GE changes")
 	virtual void OnGameplayEffectChanged() const {}
 
 #if WITH_EDITOR
diff --git a/Engine/Plugins/Runtime/GameplayAbilities/Source/GameplayAbilities/Public/GameplayEffectComponents/AssetTagsGameplayEffectComponent.h b/Engine/Plugins/Runtime/GameplayAbilities/Source/GameplayAbilities/Public/GameplayEffectComponents/AssetTagsGameplayEffectComponent.h
--- a/Engine/Plugins/Runtime/GameplayAbilities/Source/GameplayAbilities/Public/GameplayEffectComponents/AssetTagsGameplayEffectComponent.h
+++ b/Engine/Plugins/Runtime/GameplayAbilities/Source/GameplayAbilities/Public/GameplayEffectComponents/AssetTagsGameplayEffectComponent.h
@@ -17,7 +17,7 @@ public:
 	virtual void PostInitProperties() override;
 
 	/** Needed to properly apply FInheritedTagContainer properties */
-	virtual void OnGameplayEffectChanged() const override;
+	virtual void OnGameplayEffectChanged() override;
 
 	/** Gets the Asset Tags inherited tag structure (as configured) */
 	const FInheritedTagContainer& GetConfiguredAssetTagChanges() const { return InheritableAssetTags; }
diff --git a/Engine/Plugins/Runtime/GameplayAbilities/Source/GameplayAbilities/Public/GameplayEffectComponents/BlockAbilityTagsGameplayEffectComponent.h b/Engine/Plugins/Runtime/GameplayAbilities/Source/GameplayAbilities/Public/GameplayEffectComponents/BlockAbilityTagsGameplayEffectComponent.h
--- a/Engine/Plugins/Runtime/GameplayAbilities/Source/GameplayAbilities/Public/GameplayEffectComponents/BlockAbilityTagsGameplayEffectComponent.h
+++ b/Engine/Plugins/Runtime/GameplayAbilities/Source/GameplayAbilities/Public/GameplayEffectComponents/BlockAbilityTagsGameplayEffectComponent.h
@@ -17,7 +17,7 @@ public:
 	virtual void PostInitProperties() override;
 
 	/** Needed to properly apply FInheritedTagContainer properties */
-	virtual void OnGameplayEffectChanged() const override;
+	virtual void OnGameplayEffectChanged() override;
 
 	/** Gets the Blocked Ability Tags inherited tag structure (as configured) */
 	const FInheritedTagContainer& GetConfiguredBlockedAbilityTagChanges() const { return InheritableBlockedAbilityTagsContainer; }
diff --git a/Engine/Plugins/Runtime/GameplayAbilities/Source/GameplayAbilities/Public/GameplayEffectComponents/TargetTagsGameplayEffectComponent.h b/Engine/Plugins/Runtime/GameplayAbilities/Source/GameplayAbilities/Public/GameplayEffectComponents/TargetTagsGameplayEffectComponent.h
--- a/Engine/Plugins/Runtime/GameplayAbilities/Source/GameplayAbilities/Public/GameplayEffectComponents/TargetTagsGameplayEffectComponent.h
+++ b/Engine/Plugins/Runtime/GameplayAbilities/Source/GameplayAbilities/Public/GameplayEffectComponents/TargetTagsGameplayEffectComponent.h
@@ -17,7 +17,7 @@ public:
 	virtual void PostInitProperties() override;
 
 	/** Needed to properly load FInheritedTagContainer properties */
-	virtual void OnGameplayEffectChanged() const override;
+	virtual void OnGameplayEffectChanged() override;
 
 	/** Gets the (Granted) Target Tags inherited tag structure (as configured) */
 	const FInheritedTagContainer& GetConfiguredTargetTagChanges() const { return InheritableGrantedTagsContainer; }
