PR #11156 (modified): [GAS] Fixed bugs caused by gameplay cue getting removed after OnRemove is sent to the cue

Note:  This will change behavior slightly.  Previously, removing a GameplayCue would cause a single EGameplayTagEventType::AnyCountChange callback with its Tag and Count, even though all equivalent Cues would receive an OnRemove.  So previously:  15 GC Tags would be a single callback and now it will be 15 callbacks.  This will change behavior if you're using something like AbilityTask_WaitGameplayTagCountChanged on the GameplayCue Tag (that callback would fire 15x now, and just once previously).  I believe this simplification of the code (along with such an unlikely edge case) warrants the change in behavior.

The CVar to toggle this functionality is GameplayCue.Fix.UseEqualTagCountAndRemovalCallbacks (it is now defaulted to true)

diff --git a/Engine/Plugins/Runtime/GameplayAbilities/Source/GameplayAbilities/Private/AbilitySystemComponent.cpp b/Engine/Plugins/Runtime/GameplayAbilities/Source/GameplayAbilities/Private/AbilitySystemComponent.cpp
--- a/Engine/Plugins/Runtime/GameplayAbilities/Source/GameplayAbilities/Private/AbilitySystemComponent.cpp
+++ b/Engine/Plugins/Runtime/GameplayAbilities/Source/GameplayAbilities/Private/AbilitySystemComponent.cpp
@@ -1378,34 +1378,7 @@ void UAbilitySystemComponent::RemoveGameplayCue_Internal(const FGameplayTag Game
 {
 	if (IsOwnerActorAuthoritative())
 	{
-		int32 NumMatchingCues = 0;
-		for (const FActiveGameplayCue& GameplayCue : GameplayCueContainer.GameplayCues)
-		{
-			NumMatchingCues += (GameplayCue.GameplayCueTag == GameplayCueTag);
-		}
-
-		if (NumMatchingCues > 0)
-		{
-			// AbilitySystem.GameplayCueNotifyTagCheckOnRemove assumes the tag is removed before any invocation of EGameplayCueEvent::Removed.
-			// We cannot use GameplayCueContainer.RemoveCue because that removes the cues while updating the TagMap.
-			// Instead, we need to manually count the removals, update the tag map, then Invoke the Cue events while removing the Cues.
-			UpdateTagMap(GameplayCueTag, -NumMatchingCues);
-
-			for (int32 Index = GameplayCueContainer.GameplayCues.Num() - 1; Index >= 0; --Index)
-			{
-				const FActiveGameplayCue& GameplayCue = GameplayCueContainer.GameplayCues[Index];
-				if (GameplayCue.GameplayCueTag == GameplayCueTag)
-				{
-					// Call on server here, clients get it from repnotify on the GameplayCueContainer rather than a multicast rpc
-					InvokeGameplayCueEvent(GameplayCueTag, EGameplayCueEvent::Removed, GameplayCue.Parameters);
-					GameplayCueContainer.GameplayCues.RemoveAt(Index);
-				}
-			}
-
-			// Ensure that the clients are aware of these changes ASAP
-			GameplayCueContainer.MarkArrayDirty();
-			ForceReplication();
-		}
+		GameplayCueContainer.RemoveCue(GameplayCueTag);
 	}
 	else if (ScopedPredictionKey.IsLocalClientKey())
 	{
diff --git a/Engine/Plugins/Runtime/GameplayAbilities/Source/GameplayAbilities/Private/GameplayCueInterface.cpp b/Engine/Plugins/Runtime/GameplayAbilities/Source/GameplayAbilities/Private/GameplayCueInterface.cpp
--- a/Engine/Plugins/Runtime/GameplayAbilities/Source/GameplayAbilities/Private/GameplayCueInterface.cpp
+++ b/Engine/Plugins/Runtime/GameplayAbilities/Source/GameplayAbilities/Private/GameplayCueInterface.cpp
@@ -20,6 +20,9 @@ namespace GameplayCueInterfacePrivate
 	};
 	typedef TMap<FGameplayTag, TArray<FCueNameAndUFunction> > FGameplayCueTagFunctionList;
 	static TMap<FObjectKey, FGameplayCueTagFunctionList > PerClassGameplayTagToFunctionMap;
+
+	bool bUseEqualTagCountAndRemovalCallbacks = true;
+	static FAutoConsoleVariableRef CVarUseEqualTagCountAndRemovalCallbacks(TEXT("GameplayCue.Fix.UseEqualTagCountAndRemovalCallbacks"), bUseEqualTagCountAndRemovalCallbacks, TEXT("Default: True. When calling RemoveCue, get an equal number of TagCountUpdated callbacks and Removal callbacks rather than 1:many"));
 }
 
 
@@ -243,24 +246,44 @@ void FActiveGameplayCueContainer::RemoveCue(const FGameplayTag& Tag)
 		return;
 	}
 
+	int32 CountDelta = 0;
+
+	// This is the old behavior (tentatively removed in UE5.5):  A single callback per multiple removals
+	if (!GameplayCueInterfacePrivate::bUseEqualTagCountAndRemovalCallbacks)
+	{
+		for (const FActiveGameplayCue& GameplayCue : GameplayCues)
+		{
+			CountDelta += (GameplayCue.GameplayCueTag == Tag);
+		}
+
+		if (CountDelta > 0)
+		{
+			Owner->UpdateTagMap(Tag, -CountDelta);
+		}
+	}
 	
 	// Iterate backwards so we can remove during loop
-	int32 CountDelta = 0;
 	for (int32 idx=GameplayCues.Num()-1; idx >= 0; --idx)
 	{
 		FActiveGameplayCue& Cue = GameplayCues[idx];
 
 		if (Cue.GameplayCueTag == Tag)
 		{
+			if (GameplayCueInterfacePrivate::bUseEqualTagCountAndRemovalCallbacks)
+			{
+				++CountDelta;
+				Owner->UpdateTagMap(Tag, -1);
+			}
+
+			Owner->InvokeGameplayCueEvent(Tag, EGameplayCueEvent::Removed, Cue.Parameters);
 			GameplayCues.RemoveAt(idx);
-			CountDelta -= 1;
 		}
 	}
 
-	if (CountDelta < 0)
+	if (CountDelta > 0)
 	{
 		MarkArrayDirty();
-		Owner->UpdateTagMap(Tag, CountDelta);
+		Owner->ForceReplication();
 	}
 }
 
