Added invalid data warning when BlockAbilityTags is used with Instant GE

diff --git a/Engine/Plugins/Runtime/GameplayAbilities/Source/GameplayAbilities/Private/GameplayEffectComponents/BlockAbilityTagsGameplayEffectComponent.cpp b/Engine/Plugins/Runtime/GameplayAbilities/Source/GameplayAbilities/Private/GameplayEffectComponents/BlockAbilityTagsGameplayEffectComponent.cpp
--- a/Engine/Plugins/Runtime/GameplayAbilities/Source/GameplayAbilities/Private/GameplayEffectComponents/BlockAbilityTagsGameplayEffectComponent.cpp
+++ b/Engine/Plugins/Runtime/GameplayAbilities/Source/GameplayAbilities/Private/GameplayEffectComponents/BlockAbilityTagsGameplayEffectComponent.cpp
@@ -1,6 +1,9 @@
 // Copyright Epic Games, Inc. All Rights Reserved.
 
 #include "GameplayEffectComponents/BlockAbilityTagsGameplayEffectComponent.h"
+#include "Misc/DataValidation.h"
+
+#define LOCTEXT_NAMESPACE "BlockAbilityTagsGameplayEffectComponent"
 
 void UBlockAbilityTagsGameplayEffectComponent::PostInitProperties()
 {
@@ -34,6 +37,21 @@ void UBlockAbilityTagsGameplayEffectComponent::PostEditChangeProperty(FPropertyC
 		Owner->OnGameplayEffectChanged();
 	}
 }
+
+EDataValidationResult UBlockAbilityTagsGameplayEffectComponent::IsDataValid(FDataValidationContext& Context) const
+{
+	EDataValidationResult Result = Super::IsDataValid(Context);
+
+	const bool bInstantEffect = (GetOwner()->DurationPolicy == EGameplayEffectDurationType::Instant);
+	if (bInstantEffect && !InheritableBlockedAbilityTagsContainer.CombinedTags.IsEmpty())
+	{
+		Context.AddError(FText::FormatOrdered(LOCTEXT("GEInstantAndBlockAbilityTags", "GE {0} is set to Instant so {1} will not be able to function as expected."), FText::FromString(GetNameSafe(GetOwner())), FText::FromString(EditorFriendlyName)));
+		Result = EDataValidationResult::Invalid;
+	}
+
+	return Result;
+}
+
 #endif // WITH_EDITOR
 
 void UBlockAbilityTagsGameplayEffectComponent::SetAndApplyBlockedAbilityTagChanges(const FInheritedTagContainer& TagContainerMods)
@@ -52,4 +70,6 @@ void UBlockAbilityTagsGameplayEffectComponent::ApplyBlockedAbilityTagChanges() c
 	// Apply to the owning Gameplay Effect Component
 	UGameplayEffect* Owner = GetOwner();
 	InheritableBlockedAbilityTagsContainer.ApplyTo(Owner->CachedBlockedAbilityTags);
-}
\ No newline at end of file
+}
+
+#undef LOCTEXT_NAMESPACE
\ No newline at end of file
diff --git a/Engine/Plugins/Runtime/GameplayAbilities/Source/GameplayAbilities/Public/GameplayEffectComponents/BlockAbilityTagsGameplayEffectComponent.h b/Engine/Plugins/Runtime/GameplayAbilities/Source/GameplayAbilities/Public/GameplayEffectComponents/BlockAbilityTagsGameplayEffectComponent.h
--- a/Engine/Plugins/Runtime/GameplayAbilities/Source/GameplayAbilities/Public/GameplayEffectComponents/BlockAbilityTagsGameplayEffectComponent.h
+++ b/Engine/Plugins/Runtime/GameplayAbilities/Source/GameplayAbilities/Public/GameplayEffectComponents/BlockAbilityTagsGameplayEffectComponent.h
@@ -29,6 +29,9 @@ public:
 	/** Needed to properly update FInheritedTagContainer properties */
 	virtual void PostEditChangeProperty(FPropertyChangedEvent& PropertyChangedEvent) override;
 
+	/** Validate incompatible configurations */
+	virtual EDataValidationResult IsDataValid(class FDataValidationContext& Context) const override;
+
 private:
 	/** Get a cached version of the FProperty Name for PostEditChangeProperty */
 	static const FName& GetInheritableBlockedAbilityTagsContainerPropertyName()
