Minor GAS debugging improvements

#rnx

diff --git a/Engine/Plugins/Runtime/GameplayAbilities/Source/GameplayAbilities/Private/AbilitySystemComponent_Abilities.cpp b/Engine/Plugins/Runtime/GameplayAbilities/Source/GameplayAbilities/Private/AbilitySystemComponent_Abilities.cpp
--- a/Engine/Plugins/Runtime/GameplayAbilities/Source/GameplayAbilities/Private/AbilitySystemComponent_Abilities.cpp
+++ b/Engine/Plugins/Runtime/GameplayAbilities/Source/GameplayAbilities/Private/AbilitySystemComponent_Abilities.cpp
@@ -2215,7 +2215,8 @@ void UAbilitySystemComponent::ClientActivateAbilitySucceedWithEventData_Implemen
 
 	Spec->ActivationInfo.SetActivationConfirmed();
 
-	// ABILITY_LOG(Verbose, TEXT("ClientActivateAbilitySucceedWithEventData_Implementation. PredictionKey :%s Ability: %s"), *PredictionKey.ToString(), *GetNameSafe(Spec->Ability));
+	UE_LOG(LogAbilitySystem, Verbose, TEXT("%s: Server Confirmed [%s] %s. PredictionKey: %s"), *GetNameSafe(GetOwner()), *Handle.ToString(), *GetNameSafe(AbilityToActivate), *PredictionKey.ToString());
+	UE_VLOG(GetOwner(), VLogAbilitySystem, Verbose, TEXT("Server Confirmed [%s] %s. PredictionKey: %s"), *Handle.ToString(), *GetNameSafe(AbilityToActivate), *PredictionKey.ToString());
 
 	// Fixme: We need a better way to link up/reconcile predictive replicated abilities. It would be ideal if we could predictively spawn an
 	// ability and then replace/link it with the server spawned one once the server has confirmed it.
diff --git a/Engine/Plugins/Runtime/GameplayAbilities/Source/GameplayAbilities/Private/GameplayDebuggerCategory_Abilities.cpp b/Engine/Plugins/Runtime/GameplayAbilities/Source/GameplayAbilities/Private/GameplayDebuggerCategory_Abilities.cpp
--- a/Engine/Plugins/Runtime/GameplayAbilities/Source/GameplayAbilities/Private/GameplayDebuggerCategory_Abilities.cpp
+++ b/Engine/Plugins/Runtime/GameplayAbilities/Source/GameplayAbilities/Private/GameplayDebuggerCategory_Abilities.cpp
@@ -734,21 +734,26 @@ void FGameplayDebuggerCategory_Abilities::DrawGameplayAttributes(FGameplayDebugg
 		}
 		else
 		{
+			const bool bDisplayServerValue = (AttributeData.NetworkStatus != FRepData::ENetworkStatus::LocalOnly);
+			const bool bDisplayClientValue = (AttributeData.NetworkStatus != FRepData::ENetworkStatus::ServerOnly) && !bNetworkValueMatch;
+
 			// Append status that is happening on the Server
-			if (AttributeData.NetworkStatus != FRepData::ENetworkStatus::LocalOnly)
+			if (bDisplayServerValue)
 			{
+				// If we are going to also display a client value, color it in the ServerColor, otherwise white.
+				const TCHAR* ServerValueColor = bDisplayClientValue ? ServerColor : TEXT("{white}");
 				if (bServerValueMatch)
 				{
-					AttributeValueStr.Appendf(TEXT("{white}%.4g "), AttributeData.ServerCurrentValue);
+					AttributeValueStr.Appendf(TEXT("%s%.4g "), ServerValueColor, AttributeData.ServerCurrentValue);
 				}
 				else
 				{
-					AttributeValueStr.Appendf(TEXT("{white}%.4g [%.4g] "), AttributeData.ServerCurrentValue, AttributeData.ServerBaseValue);
+					AttributeValueStr.Appendf(TEXT("%s%.4g [%.4g] "), ServerValueColor, AttributeData.ServerCurrentValue, AttributeData.ServerBaseValue);
 				}
 			}
 
 			// Append status that is happening on the Client (locally)
-			if (AttributeData.NetworkStatus != FRepData::ENetworkStatus::ServerOnly)
+			if (bDisplayClientValue)
 			{
 				if (bLocalValueMatch)
 				{
