Skip the check for unmapped references when receiving the replicated Server Initiated Gameplay Effects.  This ensures the Gameplay Cues will always fire properly.

diff --git a/Engine/Plugins/Runtime/GameplayAbilities/Source/GameplayAbilities/Private/GameplayEffect.cpp b/Engine/Plugins/Runtime/GameplayAbilities/Source/GameplayAbilities/Private/GameplayEffect.cpp
--- a/Engine/Plugins/Runtime/GameplayAbilities/Source/GameplayAbilities/Private/GameplayEffect.cpp
+++ b/Engine/Plugins/Runtime/GameplayAbilities/Source/GameplayAbilities/Private/GameplayEffect.cpp
@@ -88,6 +88,11 @@ namespace UE::GameplayEffect
 	bool bUseModifierTagRequirementsOnAllGameplayEffects = true;
 	FAutoConsoleVariableRef CVarUseModTagReqsOnAllGE{ TEXT("AbilitySystem.Fix.UseModTagReqsOnAllGE"), bUseModifierTagRequirementsOnAllGameplayEffects, TEXT("Fix an issue where MustHave/MustNotHave tags did not apply to Instant and Periodic Gameplay Effects"), ECVF_Default };
 
+	// Fix introduced in UE5.4
+	bool bSkipUnmappedReferencesCheckForGameplayCues = true;
+	FAutoConsoleVariableRef CVarSkipUnmappedReferencesCheckForGameplayCues{ TEXT("AbilitySystem.Fix.SkipUnmappedReferencesCheckForGameplayCues"), bSkipUnmappedReferencesCheckForGameplayCues,
+		TEXT("Skip the bHasMoreUnmappedReferences check for GameplayCues which never worked as intended and causes issues when set properly (may be deprecated soon)"), ECVF_Default };
+
 #if WITH_EDITOR
 	namespace EditorOnly
 	{
@@ -4865,9 +4870,8 @@ void FActiveGameplayEffectsContainer::PostReplicatedReceive(const FFastArraySeri
 	{
 		QUICK_SCOPE_CYCLE_COUNTER(STAT_ActiveGameplayEffectsContainer_NetDeltaSerialize_CheckRepGameplayCues);
 
-		if (!Parameters.bHasMoreUnmappedReferences) // Do not invoke GCs when we have missing information (like AActor*s in EffectContext)
+		if ( LIKELY(UE::GameplayEffect::bSkipUnmappedReferencesCheckForGameplayCues) )
 		{
-			NumConsecutiveUnmappedReferencesDebug = 0;
 			if (Owner->IsReadyForGameplayCues())
 			{
 				Owner->HandleDeferredGameplayCues(this);
@@ -4875,11 +4879,22 @@ void FActiveGameplayEffectsContainer::PostReplicatedReceive(const FFastArraySeri
 		}
 		else
 		{
-			++NumConsecutiveUnmappedReferencesDebug;
+			if (!Parameters.bHasMoreUnmappedReferences) // Do not invoke GCs when we have missing information (like AActor*s in EffectContext)
+			{
+				NumConsecutiveUnmappedReferencesDebug = 0;
+				if (Owner->IsReadyForGameplayCues())
+				{
+					Owner->HandleDeferredGameplayCues(this);
+				}
+			}
+			else
+			{
+				++NumConsecutiveUnmappedReferencesDebug;
 
-			constexpr uint32 HighNumberOfConsecutiveUnmappedRefs = 30;
-			ensureMsgf(NumConsecutiveUnmappedReferencesDebug < HighNumberOfConsecutiveUnmappedRefs, TEXT("%hs: bHasMoreUnmappedReferences is preventing GameplayCues from firing"), __func__);
-			UE_CLOG((NumConsecutiveUnmappedReferencesDebug % HighNumberOfConsecutiveUnmappedRefs) == 0, LogAbilitySystem, Error, TEXT("%hs: bHasMoreUnmappedReferences is preventing GameplayCues from firing (%u consecutive misses)"), __func__, NumConsecutiveUnmappedReferencesDebug);
+				constexpr uint32 HighNumberOfConsecutiveUnmappedRefs = 30;
+				ensureMsgf(NumConsecutiveUnmappedReferencesDebug < HighNumberOfConsecutiveUnmappedRefs, TEXT("%hs: bHasMoreUnmappedReferences is preventing GameplayCues from firing"), __func__);
+				UE_CLOG((NumConsecutiveUnmappedReferencesDebug % HighNumberOfConsecutiveUnmappedRefs) == 0, LogAbilitySystem, Error, TEXT("%hs: bHasMoreUnmappedReferences is preventing GameplayCues from firing (%u consecutive misses)"), __func__, NumConsecutiveUnmappedReferencesDebug);
+			}
 		}
 	}
 }
