PR #11457: [GAS] fix up a comment about PostGameplayEffectExecute in AttributeSet.h

diff --git a/Engine/Plugins/Runtime/GameplayAbilities/Source/GameplayAbilities/Public/AttributeSet.h b/Engine/Plugins/Runtime/GameplayAbilities/Source/GameplayAbilities/Public/AttributeSet.h
--- a/Engine/Plugins/Runtime/GameplayAbilities/Source/GameplayAbilities/Public/AttributeSet.h
+++ b/Engine/Plugins/Runtime/GameplayAbilities/Source/GameplayAbilities/Public/AttributeSet.h
@@ -195,7 +195,7 @@ public:
 	virtual bool PreGameplayEffectExecute(struct FGameplayEffectModCallbackData &Data) { return true; }
 	
 	/**
-	 *	Called just before a GameplayEffect is executed to modify the base value of an attribute. No more changes can be made.
+	 *	Called just after a GameplayEffect is executed to modify the base value of an attribute. No more changes can be made.
 	 *	Note this is only called during an 'execute'. E.g., a modification to the 'base value' of an attribute. It is not called during an application of a GameplayEffect, such as a 5 ssecond +10 movement speed buff.
 	 */
 	virtual void PostGameplayEffectExecute(const struct FGameplayEffectModCallbackData &Data) { }
