Iris
- Fixed bugs in MinimalGameplayCueReplicationProxyReplicationFragment and MinimalReplicationTagCountMapReplicationFragments that would lead to excessive polling

diff --git a/Engine/Plugins/Runtime/GameplayAbilities/Source/GameplayAbilities/Private/Serialization/MinimalGameplayCueReplicationProxyReplicationFragment.cpp b/Engine/Plugins/Runtime/GameplayAbilities/Source/GameplayAbilities/Private/Serialization/MinimalGameplayCueReplicationProxyReplicationFragment.cpp
--- a/Engine/Plugins/Runtime/GameplayAbilities/Source/GameplayAbilities/Private/Serialization/MinimalGameplayCueReplicationProxyReplicationFragment.cpp
+++ b/Engine/Plugins/Runtime/GameplayAbilities/Source/GameplayAbilities/Private/Serialization/MinimalGameplayCueReplicationProxyReplicationFragment.cpp
@@ -90,11 +90,12 @@ bool FMinimalGameplayCueReplicationProxyReplicationFragment::PollReplicatedState
 		const uint8* ExternalStateBuffer = reinterpret_cast<uint8*>(Owner) + ReplicationStateDescriptor->MemberProperties[0]->GetOffset_ForGC();
 		const FMinimalGameplayCueReplicationProxy* ExternalSourceState = reinterpret_cast<const FMinimalGameplayCueReplicationProxy*>(ExternalStateBuffer);
 
-		const void* CachedStateBuffer = SrcReplicationState->GetStateBuffer();
-		const FMinimalGameplayCueReplicationProxy* CachedState = reinterpret_cast<const FMinimalGameplayCueReplicationProxy*>(static_cast<const uint8*>(CachedStateBuffer) + ReplicationStateDescriptor->MemberDescriptors[0].ExternalMemberOffset);
+		void* CachedStateBuffer = SrcReplicationState->GetStateBuffer();
+		FMinimalGameplayCueReplicationProxy* CachedState = reinterpret_cast<FMinimalGameplayCueReplicationProxy*>(static_cast<uint8*>(CachedStateBuffer) + ReplicationStateDescriptor->MemberDescriptors[0].ExternalMemberOffset);
 
 		if (ExternalSourceState->LastSourceArrayReplicationKey != CachedState->LastSourceArrayReplicationKey)
 		{
+			CachedState->LastSourceArrayReplicationKey = ExternalSourceState->LastSourceArrayReplicationKey;
 			return SrcReplicationState->PollPropertyReplicationState(Owner);
 		}
 	}
@@ -168,6 +169,9 @@ void FMinimalGameplayCueReplicationProxyReplicationFragment::MimicMinimalGamepla
 			// This is a new tag, we need to invoke the WhileActive gameplay cue event
 			StateOwner->SetTagMapCount(ReplicatedTag, 1);
 			StateOwner->InvokeGameplayCueEvent(ReplicatedTag, EGameplayCueEvent::WhileActive, Parameters);
+
+			// The demo recorder needs to believe that this structure is dirty so it will get saved into the demo stream
+			ExternalSourceState->LastSourceArrayReplicationKey++;
 		}
 	}
 
diff --git a/Engine/Plugins/Runtime/GameplayAbilities/Source/GameplayAbilities/Private/Serialization/MinimalReplicationTagCountMapReplicationFragment.cpp b/Engine/Plugins/Runtime/GameplayAbilities/Source/GameplayAbilities/Private/Serialization/MinimalReplicationTagCountMapReplicationFragment.cpp
--- a/Engine/Plugins/Runtime/GameplayAbilities/Source/GameplayAbilities/Private/Serialization/MinimalReplicationTagCountMapReplicationFragment.cpp
+++ b/Engine/Plugins/Runtime/GameplayAbilities/Source/GameplayAbilities/Private/Serialization/MinimalReplicationTagCountMapReplicationFragment.cpp
@@ -91,10 +91,14 @@ bool FMinimalReplicationTagCountMapReplicationFragment::PollReplicatedState(ERep
 		const uint8* ExternalStateBuffer = reinterpret_cast<uint8*>(Owner) + ReplicationStateDescriptor->MemberProperties[0]->GetOffset_ForGC();
 		const FMinimalReplicationTagCountMap* ExternalSourceState = reinterpret_cast<const FMinimalReplicationTagCountMap*>(ExternalStateBuffer);
 
-		const void* CachedStateBuffer = SrcReplicationState->GetStateBuffer();
-		const FMinimalReplicationTagCountMap* CachedState = reinterpret_cast<const FMinimalReplicationTagCountMap*>(static_cast<const uint8*>(CachedStateBuffer) + ReplicationStateDescriptor->MemberDescriptors[0].ExternalMemberOffset);
+		void* CachedStateBuffer = SrcReplicationState->GetStateBuffer();
+		FMinimalReplicationTagCountMap* CachedState = reinterpret_cast<FMinimalReplicationTagCountMap*>(static_cast<uint8*>(CachedStateBuffer) + ReplicationStateDescriptor->MemberDescriptors[0].ExternalMemberOffset);
 
-		return SrcReplicationState->PollPropertyReplicationState(Owner);
+		if (ExternalSourceState->MapID != CachedState->MapID)
+		{
+			CachedState->MapID = ExternalSourceState->MapID;
+			return SrcReplicationState->PollPropertyReplicationState(Owner);
+		}
 	}
 
 	return SrcReplicationState->IsDirty(0);
