Properly dirtying the ActiveGampelayEffects array when granting a new Gameplay Ability Spec Handle.
Adding some visual logging to help track down these complex issues.

diff --git a/Engine/Plugins/Runtime/GameplayAbilities/Source/GameplayAbilities/Private/AbilitySystemComponent_Abilities.cpp b/Engine/Plugins/Runtime/GameplayAbilities/Source/GameplayAbilities/Private/AbilitySystemComponent_Abilities.cpp
--- a/Engine/Plugins/Runtime/GameplayAbilities/Source/GameplayAbilities/Private/AbilitySystemComponent_Abilities.cpp
+++ b/Engine/Plugins/Runtime/GameplayAbilities/Source/GameplayAbilities/Private/AbilitySystemComponent_Abilities.cpp
@@ -543,6 +543,7 @@ void UAbilitySystemComponent::OnGiveAbility(FGameplayAbilitySpec& Spec)
 			if (SourceActiveGE)
 			{
 				SourceActiveGE->GrantedAbilityHandles.AddUnique(Spec.Handle);
+				SourceASC->ActiveGameplayEffects.MarkItemDirty(*SourceActiveGE);
 			}
 		}
 	}
@@ -597,6 +598,9 @@ void UAbilitySystemComponent::OnRemoveAbility(FGameplayAbilitySpec& Spec)
 		return;
 	}
 
+	UE_LOG(LogAbilitySystem, Log, TEXT("%s: Removing Ability [%s] %s Level: %d"), *GetNameSafe(GetOwner()), *Spec.Handle.ToString(), *GetNameSafe(Spec.Ability), Spec.Level);
+	UE_VLOG(GetOwner(), VLogAbilitySystem, Log, TEXT("Removing Ability [%s] %s Level: %d"), *Spec.Handle.ToString(), *GetNameSafe(Spec.Ability), Spec.Level);
+
 	for (const FAbilityTriggerData& TriggerData : Spec.Ability->AbilityTriggers)
 	{
 		FGameplayTag EventTag = TriggerData.TriggerTag;
@@ -686,9 +690,6 @@ void UAbilitySystemComponent::OnRemoveAbility(FGameplayAbilitySpec& Spec)
 
 	Spec.ReplicatedInstances.Empty();
 	Spec.NonReplicatedInstances.Empty();
-
-	UE_LOG(LogAbilitySystem, Log, TEXT("%s: Removed Ability %s [%s] Level: %d"), *GetNameSafe(GetOwner()), *GetNameSafe(Spec.Ability), *Spec.Handle.ToString(), Spec.Level);
-	UE_VLOG(GetOwner(), VLogAbilitySystem, Log, TEXT("Removed Ability %s [%s] Level: %d"), *GetNameSafe(Spec.Ability), *Spec.Handle.ToString(), Spec.Level);
 }
 
 void UAbilitySystemComponent::CheckForClearedAbilities()
diff --git a/Engine/Plugins/Runtime/GameplayAbilities/Source/GameplayAbilities/Private/GameplayAbilityTypes.cpp b/Engine/Plugins/Runtime/GameplayAbilities/Source/GameplayAbilities/Private/GameplayAbilityTypes.cpp
--- a/Engine/Plugins/Runtime/GameplayAbilities/Source/GameplayAbilities/Private/GameplayAbilityTypes.cpp
+++ b/Engine/Plugins/Runtime/GameplayAbilities/Source/GameplayAbilities/Private/GameplayAbilityTypes.cpp
@@ -232,6 +232,9 @@ void FGameplayAbilitySpec::PreReplicatedRemove(const struct FGameplayAbilitySpec
 {
 	if (InArraySerializer.Owner)
 	{
+		UE_LOG(LogAbilitySystem, Verbose, TEXT("%s: OnRemoveAbility (Non-Auth): [%s] %s. Level: %d"), *GetNameSafe(InArraySerializer.Owner->GetOwner()), *Handle.ToString(), *GetNameSafe(Ability), Level)
+		UE_VLOG(InArraySerializer.Owner->GetOwner(), VLogAbilitySystem, Verbose, TEXT("OnRemoveAbility (Non-Auth): [%s] %s. Level: %d"), *Handle.ToString(), *GetNameSafe(Ability), Level);
+
 		FScopedAbilityListLock AblityListLock(*InArraySerializer.Owner);
 		InArraySerializer.Owner->OnRemoveAbility(*this);
 	}
@@ -241,6 +244,9 @@ void FGameplayAbilitySpec::PostReplicatedAdd(const struct FGameplayAbilitySpecCo
 {
 	if (InArraySerializer.Owner)
 	{
+		UE_LOG(LogAbilitySystem, Verbose, TEXT("%s: OnGiveAbility (Non-Auth): [%s] %s. Level: %d"), *GetNameSafe(InArraySerializer.Owner->GetOwner()), *Handle.ToString(), *GetNameSafe(Ability), Level)
+		UE_VLOG(InArraySerializer.Owner->GetOwner(), VLogAbilitySystem, Verbose, TEXT("OnGiveAbility (Non-Auth): [%s] %s. Level: %d"), *Handle.ToString(), *GetNameSafe(Ability), Level);
+
 		InArraySerializer.Owner->OnGiveAbility(*this);
 	}
 }
