Modified PR #11483: Adding Debug support for GameplayAbilityTargetActor_Radius

diff --git a/Engine/Plugins/Runtime/GameplayAbilities/Source/GameplayAbilities/Private/Abilities/GameplayAbilityTargetActor_Radius.cpp b/Engine/Plugins/Runtime/GameplayAbilities/Source/GameplayAbilities/Private/Abilities/GameplayAbilityTargetActor_Radius.cpp
--- a/Engine/Plugins/Runtime/GameplayAbilities/Source/GameplayAbilities/Private/Abilities/GameplayAbilityTargetActor_Radius.cpp
+++ b/Engine/Plugins/Runtime/GameplayAbilities/Source/GameplayAbilities/Private/Abilities/GameplayAbilityTargetActor_Radius.cpp
@@ -4,6 +4,8 @@
 #include "GameFramework/Pawn.h"
 #include "WorldCollision.h"
 #include "Abilities/GameplayAbility.h"
+#include "AbilitySystemLog.h"
+#include "DrawDebugHelpers.h"
 #include "Engine/OverlapResult.h"
 #include "Engine/World.h"
 
@@ -74,6 +76,18 @@ TArray<TWeakObjectPtr<AActor> >	AGameplayAbilityTargetActor_Radius::PerformOverl
 		}
 	}
 
+	if (bDebug)
+	{
+		const FColor& DebugColor = HitActors.Num() == 0 ? FColor::Red : FColor::Green;
+		UE_VLOG_SPHERE(SourceActor, LogAbilitySystem, Log, Origin, Radius, DebugColor, TEXT("TargetActor Radius"));
+
+#if UE_ENABLE_DEBUG_DRAWING
+		constexpr bool bPersistent = false;
+		constexpr float LifeTime = 2.0f;
+		DrawDebugSphere(GetWorld(), Origin, Radius, 16, DebugColor, bPersistent, LifeTime);
+#endif // UE_ENABLE_DEBUG_DRAWING
+	}
+
 	return HitActors;
 }
 
