PR #11505: Add on gameplay effect inhibition changed delegate

diff --git a/Engine/Plugins/Runtime/GameplayAbilities/Source/GameplayAbilities/Private/AbilitySystemComponent.cpp b/Engine/Plugins/Runtime/GameplayAbilities/Source/GameplayAbilities/Private/AbilitySystemComponent.cpp
--- a/Engine/Plugins/Runtime/GameplayAbilities/Source/GameplayAbilities/Private/AbilitySystemComponent.cpp
+++ b/Engine/Plugins/Runtime/GameplayAbilities/Source/GameplayAbilities/Private/AbilitySystemComponent.cpp
@@ -626,6 +626,12 @@ FOnActiveGameplayEffectTimeChange* UAbilitySystemComponent::OnGameplayEffectTime
 	return ActiveEffect ? &ActiveEffect->EventSet.OnTimeChanged : nullptr;
 }
 
+FOnActiveGameplayEffectInhibitionChanged* UAbilitySystemComponent::OnGameplayEffectInhibitionChangedDelegate(FActiveGameplayEffectHandle Handle)
+{
+	FActiveGameplayEffect* ActiveEffect = ActiveGameplayEffects.GetActiveGameplayEffect(Handle);
+	return ActiveEffect ? &ActiveEffect->EventSet.OnInhibitionChanged : nullptr;
+}
+
 FOnGivenActiveGameplayEffectRemoved& UAbilitySystemComponent::OnAnyGameplayEffectRemovedDelegate()
 {
 	return ActiveGameplayEffects.OnActiveGameplayEffectRemovedDelegate;
diff --git a/Engine/Plugins/Runtime/GameplayAbilities/Source/GameplayAbilities/Public/AbilitySystemComponent.h b/Engine/Plugins/Runtime/GameplayAbilities/Source/GameplayAbilities/Public/AbilitySystemComponent.h
--- a/Engine/Plugins/Runtime/GameplayAbilities/Source/GameplayAbilities/Public/AbilitySystemComponent.h
+++ b/Engine/Plugins/Runtime/GameplayAbilities/Source/GameplayAbilities/Public/AbilitySystemComponent.h
@@ -551,6 +551,7 @@ class GAMEPLAYABILITIES_API UAbilitySystemComponent : public UGameplayTasksCompo
 	FOnActiveGameplayEffectRemoved_Info* OnGameplayEffectRemoved_InfoDelegate(FActiveGameplayEffectHandle Handle);
 	FOnActiveGameplayEffectStackChange* OnGameplayEffectStackChangeDelegate(FActiveGameplayEffectHandle Handle);
 	FOnActiveGameplayEffectTimeChange* OnGameplayEffectTimeChangeDelegate(FActiveGameplayEffectHandle Handle);
+	FOnActiveGameplayEffectInhibitionChanged* OnGameplayEffectInhibitionChangedDelegate(FActiveGameplayEffectHandle Handle);
 
 	// ----------------------------------------------------------------------------------------------------------------
 	//  Gameplay tag operations
