Fix typo in TagNetSerializer IsEqual

#rnx

diff --git a/Engine/Plugins/Runtime/GameplayAbilities/Source/GameplayAbilities/Private/Serialization/MinimalReplicationTagCountNetSerializer.cpp b/Engine/Plugins/Runtime/GameplayAbilities/Source/GameplayAbilities/Private/Serialization/MinimalReplicationTagCountNetSerializer.cpp
--- a/Engine/Plugins/Runtime/GameplayAbilities/Source/GameplayAbilities/Private/Serialization/MinimalReplicationTagCountNetSerializer.cpp
+++ b/Engine/Plugins/Runtime/GameplayAbilities/Source/GameplayAbilities/Private/Serialization/MinimalReplicationTagCountNetSerializer.cpp
@@ -171,7 +171,7 @@ bool FMinimalReplicationTagCountMapNetSerializer::IsEqual(FNetSerializationConte
 		TempValue0.ClampTagCount(MaxTagCount);
 
 		TempValue1.CopyReplicatedFieldsFrom(Source1);
-		TempValue0.ClampTagCount(MaxTagCount);
+		TempValue1.ClampTagCount(MaxTagCount);
 
 		FNetIsEqualArgs IsEqualArgs = Args;
 		IsEqualArgs.Source0 = NetSerializerValuePointer(&TempValue0);
