Gameplay Debugger now reconciles Server vs. Client state for Gameplay Effects, similar to how Tags and Attributes work.
Recolored the Gameplay Debugger to prefer Yellow (the common color) for properly networked attributes/effects.
Aligned some text better.  Small refactor for the Attributes.

diff --git a/Engine/Plugins/Runtime/GameplayAbilities/Source/GameplayAbilities/Private/GameplayDebuggerCategory_Abilities.cpp b/Engine/Plugins/Runtime/GameplayAbilities/Source/GameplayAbilities/Private/GameplayDebuggerCategory_Abilities.cpp
--- a/Engine/Plugins/Runtime/GameplayAbilities/Source/GameplayAbilities/Private/GameplayDebuggerCategory_Abilities.cpp
+++ b/Engine/Plugins/Runtime/GameplayAbilities/Source/GameplayAbilities/Private/GameplayDebuggerCategory_Abilities.cpp
@@ -24,10 +24,21 @@ namespace UE::AbilitySystem::Debug
 	const FString LongestDebugObjectName{ TEXT("ABCDEFGHIJKLMNOPQRSTUVWXYZ_ ABCDEFGH") };
 
 	// Let's define some names for consistent look across all of the categories
-	const TCHAR* BothColor = TEXT("{cyan}");
-	const TCHAR* ServerColor = TEXT("{yellow}");
+	const TCHAR* BothColor = TEXT("{yellow}");
+	const TCHAR* ServerColor = TEXT("{cyan}");
 	const TCHAR* LocalColor = TEXT("{green}");
 	const TCHAR* NonReplicatedColor = TEXT("{violetred}");
+
+	/** Given a string, print it out using the color legend based on the NetworkStatus */
+	inline FString ColorNetworkString(FGameplayDebuggerCategory_Abilities::ENetworkStatus NetworkStatus, const FStringView DisplayString)
+	{
+		const TCHAR* Colors[+FGameplayDebuggerCategory_Abilities::ENetworkStatus::MAX] =
+		{
+			ServerColor, LocalColor, BothColor, NonReplicatedColor
+		};
+
+		return FString::Printf(TEXT("%s%.*s"), Colors[+NetworkStatus], DisplayString.Len(), DisplayString.GetData());
+	}
 }
 
 FGameplayDebuggerCategory_Abilities::FGameplayDebuggerCategory_Abilities()
@@ -102,12 +113,19 @@ void FGameplayDebuggerCategory_Abilities::FRepData::Serialize(FArchive& Ar)
 
 	for (int32 Idx = 0; Idx < NumGE; Idx++)
 	{
+		if (Ar.IsLoading())
+		{
+			GameplayEffects[Idx].PredictionKey = FPredictionKey();
+		}
+		
+		GameplayEffects[Idx].PredictionKey.NetSerialize(Ar, ClientPackageMap.Get(), bSuccess);
 		Ar << GameplayEffects[Idx].Effect;
 		Ar << GameplayEffects[Idx].Context;
 		Ar << GameplayEffects[Idx].Duration;
 		Ar << GameplayEffects[Idx].Period;
 		Ar << GameplayEffects[Idx].Stacks;
 		Ar << GameplayEffects[Idx].Level;
+		Ar << GameplayEffects[Idx].bInhibited;
 	}
 
 	int32 NumAttrib = Attributes.Num();
@@ -143,27 +161,6 @@ void FGameplayDebuggerCategory_Abilities::CollectData(APlayerController* OwnerPC
 			DataPack.TagCounts.Add(AbilityComp->GetTagCount(Tag));
 		}
 
-		// Gameplay Effects
-		TArray<FGameplayEffectSpec> ActiveEffectSpecs;
-		AbilityComp->GetAllActiveGameplayEffectSpecs(ActiveEffectSpecs);
-		for (int32 Idx = 0; Idx < ActiveEffectSpecs.Num(); Idx++)
-		{
-			const FGameplayEffectSpec& EffectSpec = ActiveEffectSpecs[Idx];
-			FRepData::FGameplayEffectDebug ItemData;
-
-			ItemData.Effect = EffectSpec.ToSimpleString();
-			ItemData.Effect.RemoveFromStart(DEFAULT_OBJECT_PREFIX);
-			ItemData.Effect.RemoveFromEnd(TEXT("_C"));
-
-			ItemData.Context = EffectSpec.GetContext().ToString();
-			ItemData.Duration = EffectSpec.GetDuration();
-			ItemData.Period = EffectSpec.GetPeriod();
-			ItemData.Stacks = EffectSpec.GetStackCount();
-			ItemData.Level = EffectSpec.GetLevel();
-
-			DataPack.GameplayEffects.Add(ItemData);
-		}
-
 		// Abilities
 		const TArray<FGameplayAbilitySpec>& AbilitySpecs = AbilityComp->GetActivatableAbilities();
 		for (int32 Idx = 0; Idx < AbilitySpecs.Num(); Idx++)
@@ -184,11 +181,41 @@ void FGameplayDebuggerCategory_Abilities::CollectData(APlayerController* OwnerPC
 			DataPack.Abilities.Add(ItemData);
 		}
 
+		// Gameplay Effects
+		DataPack.GameplayEffects = CollectEffectsData(OwnerPC, AbilityComp);
+
 		// Attributes
 		DataPack.Attributes = CollectAttributeData(OwnerPC, AbilityComp);
 	}
 }
 
+TArray<FGameplayDebuggerCategory_Abilities::FRepData::FGameplayEffectDebug> FGameplayDebuggerCategory_Abilities::CollectEffectsData(const APlayerController* OwnerPC, const UAbilitySystemComponent* AbilityComp) const
+{
+	TArray<FRepData::FGameplayEffectDebug> DebugEffects;
+
+	for (const FActiveGameplayEffect& ActiveGE : &AbilityComp->GetActiveGameplayEffects())
+	{
+		FRepData::FGameplayEffectDebug& ItemData = DebugEffects.AddDefaulted_GetRef();
+		ItemData.bInhibited = ActiveGE.bIsInhibited;
+		ItemData.PredictionKey = ActiveGE.PredictionKey;
+		ItemData.Duration = ActiveGE.GetDuration();
+		ItemData.Period = ActiveGE.GetPeriod();
+
+		const FGameplayEffectSpec& EffectSpec = ActiveGE.Spec;
+		ItemData.Effect = EffectSpec.ToSimpleString();
+		ItemData.Effect.RemoveFromStart(DEFAULT_OBJECT_PREFIX);
+		ItemData.Effect.RemoveFromEnd(TEXT("_C"));
+
+		ItemData.Context = EffectSpec.GetContext().ToString();
+		ItemData.Stacks = EffectSpec.GetStackCount();
+		ItemData.Level = EffectSpec.GetLevel();
+
+		ItemData.NetworkStatus = OwnerPC->HasAuthority() ? ENetworkStatus::ServerOnly : ENetworkStatus::LocalOnly;
+	}
+
+	return DebugEffects;
+}
+
 TArray<FGameplayDebuggerCategory_Abilities::FRepData::FGameplayAttributeDebug> FGameplayDebuggerCategory_Abilities::CollectAttributeData(const APlayerController* OwnerPC, const UAbilitySystemComponent* AbilityComp) const
 {
 	TArray<FRepData::FGameplayAttributeDebug> DebugAttributes;
@@ -222,50 +249,33 @@ TArray<FGameplayDebuggerCategory_Abilities::FRepData::FGameplayAttributeDebug> F
 	}
 
 	// Grab the AttributeSet rather than the Attributes themselves so we can check the network functionality
-	static TMap<UClass*, TArray<FGameplayAttribute>> AttributeSetClassToGameplayAttributes;
 	for (const UAttributeSet* AttributeSet : AbilityComp->GetSpawnedAttributes())
 	{
 		const TSubclassOf<UAttributeSet> AttributeSetClass = AttributeSet ? AttributeSet->GetClass() : nullptr;
-		if (!AttributeSet || !AttributeSetClass)
+		if (!AttributeSetClass)
 		{
 			continue;
 		}
 
-		// Keep the FGameplayAttributes in a static map as we'll be looking these up (at least) twice per frame
-		TArray<FGameplayAttribute>* LocalAttributes = AttributeSetClassToGameplayAttributes.Find(AttributeSetClass.Get());
-		if (!LocalAttributes)
-		{
-			LocalAttributes = &AttributeSetClassToGameplayAttributes.Add(AttributeSetClass.Get());
-			UAttributeSet::GetAttributesFromSetClass(AttributeSetClass, *LocalAttributes);
-		}
-
 		// These are all of the replication conditions per variable
 		TArray<FLifetimeProperty> LifetimeProps;
 		AttributeSet->GetLifetimeReplicatedProps(LifetimeProps);
 
 		// Network status can change per AttributeSet, so figure it out
-		FRepData::ENetworkStatus NetworkStatus = (LocalRole == ENetRole::ROLE_Authority) ? FRepData::ENetworkStatus::ServerOnly : FRepData::ENetworkStatus::LocalOnly;
+		ENetworkStatus LocalNetworkStatus = (LocalRole == ENetRole::ROLE_Authority) ? ENetworkStatus::ServerOnly : ENetworkStatus::LocalOnly;
 		const bool bAttributeSetReplicates = bASCReplicates && AttributeSet->IsSupportedForNetworking();
 
 		// Now just gather the debug data
-		for (const FGameplayAttribute& Attrib : *LocalAttributes)
+		TArray<FGameplayAttribute> LocalAttributes;
+		UAttributeSet::GetAttributesFromSetClass(AttributeSetClass, LocalAttributes);
+		for (const FGameplayAttribute& Attrib : LocalAttributes)
 		{
-			if (!Attrib.IsValid())
-			{
-				continue;
-			}
-
-			FRepData::FGameplayAttributeDebug& DebugAttribute = DebugAttributes.AddDefaulted_GetRef();
-
-			DebugAttribute.AttributeName = Attrib.AttributeName;
-			DebugAttribute.BaseValue = AbilityComp->GetNumericAttributeBase(Attrib);
-			DebugAttribute.CurrentValue = AbilityComp->GetNumericAttribute(Attrib);
-			DebugAttribute.NetworkStatus = NetworkStatus;
+			ENetworkStatus NetworkStatus = LocalNetworkStatus;
 
 			// Override the status to network if all replication tests pass
 			if (bAttributeSetReplicates)
 			{
-				if (FProperty* Property = Attrib.GetUProperty())
+				if (const FProperty* Property = Attrib.GetUProperty())
 				{
 					if (Property->HasAnyPropertyFlags(EPropertyFlags::CPF_Net))
 					{
@@ -273,11 +283,20 @@ TArray<FGameplayDebuggerCategory_Abilities::FRepData::FGameplayAttributeDebug> F
 						if (RepProperty && RepProperty->Condition != ELifetimeCondition::COND_Never)
 						{
 							// InvalidRepIndex maps to INDEX_NONE even though unsigned (wrap-around)
-							DebugAttribute.NetworkStatus = FRepData::ENetworkStatus::Networked;
+							NetworkStatus = ENetworkStatus::Networked;
 						}
 					}
 				}
 			}
+
+			// Add the data to our returned array
+			FRepData::FGameplayAttributeDebug& DebugAttribute = DebugAttributes.Add_GetRef(
+				{
+					.AttributeName = Attrib.AttributeName,
+					.BaseValue = AbilityComp->GetNumericAttributeBase(Attrib),
+					.CurrentValue = Attrib.GetNumericValue(AttributeSet),
+					.NetworkStatus = NetworkStatus
+				});
 		}
 	}
 
@@ -322,10 +341,10 @@ void FGameplayDebuggerCategory_Abilities::DrawData(APlayerController* OwnerPC, F
 {
 	// Draw the sub-category bindings inline with the category header
 	{
-		CanvasContext.CursorX += 250.0f;
+		CanvasContext.CursorX += 200.0f;
 		CanvasContext.CursorY -= CanvasContext.GetLineHeight();
 		const TCHAR* Active = TEXT("{green}");
-		const TCHAR* Inactive = TEXT("{yellow}");
+		const TCHAR* Inactive = TEXT("{grey}");
 		CanvasContext.Printf(TEXT("Tags [%s%s{white}]\tAbilities [%s%s{white}]\tEffects [%s%s{white}]\tAttributes [%s%s{white}]"),
 			bShowGameplayTags ? Active : Inactive, *GetInputHandlerDescription(0),
 			bShowGameplayAbilities ? Active : Inactive, *GetInputHandlerDescription(1),
@@ -333,7 +352,12 @@ void FGameplayDebuggerCategory_Abilities::DrawData(APlayerController* OwnerPC, F
 			bShowGameplayAttributes ? Active : Inactive, *GetInputHandlerDescription(3));
 	}
 
-	static float LastDrawDataEndSize = CanvasContext.Canvas->SizeY - CanvasContext.CursorY - CanvasContext.CursorX;
+	if (LastDrawDataEndSize <= 0.0f)
+	{
+		// Default to the full frame size
+		LastDrawDataEndSize = CanvasContext.Canvas->SizeY - CanvasContext.CursorY - CanvasContext.CursorX;
+	}
+
 	float ThisDrawDataStartPos = CanvasContext.CursorY;
 
 	const FLinearColor BackgroundColor(0.1f, 0.1f, 0.1f, 0.8f);
@@ -373,8 +397,9 @@ void FGameplayDebuggerCategory_Abilities::DrawGameplayTags(FGameplayDebuggerCanv
 {
 	using namespace UE::AbilitySystem::Debug;
 	const float CanvasWidth = CanvasContext.Canvas->SizeX;
+	constexpr float Padding = 10.0f;
 
-	AActor* LocalDebugActor = FindLocalDebugActor();
+	const AActor* LocalDebugActor = FindLocalDebugActor();
 	if (const UAbilitySystemComponent* AbilityComp = UAbilitySystemGlobals::GetAbilitySystemComponentFromActor(LocalDebugActor))
 	{
 		FGameplayTagContainer ServerOnlyTags = DataPack.OwnedTags;
@@ -439,8 +464,6 @@ void FGameplayDebuggerCategory_Abilities::DrawGameplayTags(FGameplayDebuggerCanv
 		// and ones where they disagree.
 		if (!AbilityComp->IsOwnerActorAuthoritative())
 		{
-			CanvasContext.Printf(TEXT("Owned Tags\tLegend:  %sBoth  %sServer  %sLocal"), BothColor, ServerColor, LocalColor);
-
 			FGameplayTagContainer LocalOnlyTags;
 			AbilityComp->GetOwnedGameplayTags(LocalOnlyTags);
 
@@ -452,17 +475,30 @@ void FGameplayDebuggerCategory_Abilities::DrawGameplayTags(FGameplayDebuggerCanv
 			const FString MatchingTagsStr = BuildTagsString(MatchingTags, BothColor);
 			const FString ServerOnlyTagsStr = BuildTagsString(ServerOnlyTags, ServerColor);
 			const FString LocalOnlyTagsStr = BuildTagsString(LocalOnlyTags, LocalColor);
-			
+
+			CanvasContext.Printf(TEXT("Owned Tags [%d]:"), MatchingTags.Num() + ServerOnlyTags.Num() + LocalOnlyTags.Num());
+			CanvasContext.CursorX += 200.0f;
+			CanvasContext.CursorY -= CanvasContext.GetLineHeight();
+			CanvasContext.Printf(TEXT("Legend:  %sBoth [%d]    %sServer [%d]    %sLocal [%d]"),
+				BothColor, MatchingTags.Num(),
+				ServerColor, ServerOnlyTags.Num(),
+				LocalColor, LocalOnlyTags.Num());
+
 			FString WrappedDebugText;
 			WrapStringAccordingToViewport(FString::Printf(TEXT("%s%s%s"), *MatchingTagsStr, *ServerOnlyTagsStr, *LocalOnlyTagsStr), WrappedDebugText, CanvasContext, CanvasWidth);
+
+			CanvasContext.CursorX += Padding;
 			CanvasContext.Print(WrappedDebugText);
 		}
 		else
 		{
 			// As the authority, the source of the truth should be the ServerTags we already gathered
 			FString ServerOnlyTagsStr;
-			WrapStringAccordingToViewport(BuildTagsString(ServerOnlyTags, ServerColor), ServerOnlyTagsStr, CanvasContext, CanvasWidth);
-			CanvasContext.Printf(TEXT("Owned Tags: \n%s"), *ServerOnlyTagsStr);
+			WrapStringAccordingToViewport(BuildTagsString(ServerOnlyTags, BothColor), ServerOnlyTagsStr, CanvasContext, CanvasWidth);
+			CanvasContext.Printf(TEXT("Owned Tags [%d]"), ServerOnlyTags.Num());
+
+			CanvasContext.CursorX += Padding;
+			CanvasContext.Print(ServerOnlyTagsStr);
 		}
 	}
 
@@ -473,6 +509,7 @@ void FGameplayDebuggerCategory_Abilities::DrawGameplayTags(FGameplayDebuggerCanv
 void FGameplayDebuggerCategory_Abilities::DrawGameplayEffects(FGameplayDebuggerCanvasContext& CanvasContext, const APlayerController* OwnerPC) const
 {
 	using namespace UE::AbilitySystem::Debug;
+	using FGameplayEffectDebug = FRepData::FGameplayEffectDebug;
 
 	// Find some stable naming sizes
 	constexpr float Padding = 10.0f;
@@ -481,22 +518,70 @@ void FGameplayDebuggerCategory_Abilities::DrawGameplayEffects(FGameplayDebuggerC
 	{
 		float TempSizeY = 0.0f;
 		CanvasContext.MeasureString(*LongestDebugObjectName, ObjNameSize, TempSizeY);
-		CanvasContext.MeasureString(FString::Printf(TEXT("source: %.35s"), *LongestDebugObjectName), SrcNameSize, TempSizeY);
+		CanvasContext.MeasureString(FString::Printf(TEXT("source: %.30s"), *LongestDebugObjectName), SrcNameSize, TempSizeY);
 		ObjNameSize += Padding;
 		SrcNameSize += Padding;
 	}
 
-	CanvasContext.Printf(TEXT("Gameplay Effects: {yellow}%d"), DataPack.GameplayEffects.Num());
+	// Merge together the server and local data
+	const UAbilitySystemComponent* LocalASC = UAbilitySystemGlobals::GetAbilitySystemComponentFromActor(FindLocalDebugActor());
+	const bool bConsiderLocalStatus = !OwnerPC->IsNetMode(ENetMode::NM_Standalone) && LocalASC && (LocalASC->ReplicationMode != EGameplayEffectReplicationMode::Full);
+	const TArray<FGameplayEffectDebug>& ServerEffects = DataPack.GameplayEffects;
+	TArray<FGameplayEffectDebug> LocalEffects = bConsiderLocalStatus ? CollectEffectsData(OwnerPC, LocalASC) : ServerEffects;
+
+	auto GetKeyHash = [](const FGameplayEffectDebug& GameplayEffectDebug) -> uint32
+		{
+			return HashCombineFast(
+				GetTypeHash(GameplayEffectDebug.PredictionKey),
+				GetTypeHash(GameplayEffectDebug.Effect));
+		};
+
+	int NumEffectCounts[+ENetworkStatus::MAX] = { 0 };
+
+	TMap<uint32, FGameplayEffectDebug> Effects;
+	for (const FGameplayEffectDebug& ServerEffect : ServerEffects)
+	{
+		Effects.Add(GetKeyHash(ServerEffect)) = ServerEffect;
+	}
+
+	for (const FGameplayEffectDebug& LocalEffect : LocalEffects)
+	{
+		FGameplayEffectDebug& Effect = Effects.FindOrAdd(GetKeyHash(LocalEffect));
+		if (!Effect.Effect.IsEmpty())
+		{
+			Effect.NetworkStatus = ENetworkStatus::Networked;
+		}
+		else
+		{
+			Effect = LocalEffect;
+			Effect.NetworkStatus = ENetworkStatus::LocalOnly;
+		}
+
+		++NumEffectCounts[+Effect.NetworkStatus];
+	}
+
+	// Easier to do this calculation than to keep track during those loops
+	NumEffectCounts[+ENetworkStatus::ServerOnly] = Effects.Num() - NumEffectCounts[+ENetworkStatus::Networked] - NumEffectCounts[+ENetworkStatus::LocalOnly];
+
+	CanvasContext.Printf(TEXT("Gameplay Effects [%d]:"), Effects.Num());
+	CanvasContext.CursorX += 200.0f;
+	CanvasContext.CursorY -= CanvasContext.GetLineHeight();
+	CanvasContext.Printf(TEXT("Legend:  %sBoth [%d]    %sServer [%d]    %sLocal [%d]    {red}!! Inhibited"),
+		BothColor, NumEffectCounts[+ENetworkStatus::Networked],
+		ServerColor, NumEffectCounts[+ENetworkStatus::ServerOnly],
+		LocalColor, NumEffectCounts[+ENetworkStatus::LocalOnly]);
+
 	CanvasContext.CursorX += Padding;
-	for (int32 Idx = 0; Idx < DataPack.GameplayEffects.Num(); Idx++)
+	for (const auto& Pair : Effects)
 	{
-		const FRepData::FGameplayEffectDebug& ItemData = DataPack.GameplayEffects[Idx];
+		const FGameplayEffectDebug& ItemData = Pair.Value;
 
 		float CursorX = CanvasContext.CursorX;
 		float CursorY = CanvasContext.CursorY;
 
-		CanvasContext.PrintAt(CursorX, CursorY, FColor::Yellow, *ItemData.Effect.Left(35));
-		CanvasContext.PrintfAt(CursorX + ObjNameSize, CursorY, FColor::Silver, TEXT("source: {white}%.35s"), *ItemData.Context);
+		const FString EffectName = ColorNetworkString(ItemData.NetworkStatus, *ItemData.Effect.Left(LongestDebugObjectName.Len()));
+		CanvasContext.PrintfAt(CursorX, CursorY, TEXT("%s%s"), ItemData.bInhibited ? TEXT("{red}!!") : TEXT(""), *EffectName);
+		CanvasContext.PrintfAt(CursorX + ObjNameSize, CursorY, FColor::Silver, TEXT("source: {white}%.30s"), *ItemData.Context);
 
 		TStringBuilder<1024> Desc;
 		Desc.Appendf(TEXT(" {grey}duration: {white}"), *ItemData.Context);
@@ -563,7 +648,11 @@ void FGameplayDebuggerCategory_Abilities::DrawGameplayAbilities(FGameplayDebugge
 	const float ColumnWidth = ObjNameSize * 2 + SourceNameSize + LevelNameSize;
 	const int NumColumns = FMath::Max(1, FMath::FloorToInt(CanvasWidth / ColumnWidth));
 
-	CanvasContext.Printf(TEXT("Gameplay Abilities: \t{yellow}Granted[%d] \t{cyan}Active[%d]"), DataPack.Abilities.Num(), NumActive);
+	CanvasContext.Print(TEXT("Gameplay Abilities:"));
+	CanvasContext.CursorX += 200.0f;
+	CanvasContext.CursorY -= CanvasContext.GetLineHeight();
+	CanvasContext.Printf(TEXT("Legend:  {yellow}Granted [%d]    {cyan}Active [%d]"), DataPack.Abilities.Num(), NumActive);
+
 	CanvasContext.CursorX += Padding;
 	for (const FRepData::FGameplayAbilityDebug& ItemData : DataPack.Abilities)
 	{
@@ -603,28 +692,22 @@ void FGameplayDebuggerCategory_Abilities::DrawGameplayAttributes(FGameplayDebugg
 
 	struct FDebugAttributeData
 	{
-		FString AttributeName;
+		FStringView AttributeName;
 
 		float ServerBaseValue = 0.0f;
 		float ServerCurrentValue = 0.0f;
 		float LocalBaseValue = 0.0f;
 		float LocalCurrentValue = 0.0f;
 
-		FRepData::ENetworkStatus NetworkStatus = FRepData::ENetworkStatus::LocalOnly;
+		ENetworkStatus NetworkStatus = ENetworkStatus::LocalOnly;
 	};
 
 	const bool bConsiderNetworkStatus = !OwnerPC->IsNetMode(ENetMode::NM_Standalone);
+	const UAbilitySystemComponent* LocalASC = bConsiderNetworkStatus ? UAbilitySystemGlobals::GetAbilitySystemComponentFromActor(FindLocalDebugActor()) : nullptr;
+	const TArray<FRepData::FGameplayAttributeDebug> LocalAttributes = LocalASC ? CollectAttributeData(OwnerPC, LocalASC) : TArray<FRepData::FGameplayAttributeDebug>{};
 	TArray<FRepData::FGameplayAttributeDebug> ServerAttributes = DataPack.Attributes;
-	TArray<FRepData::FGameplayAttributeDebug> LocalAttributes;
-	if (bConsiderNetworkStatus)
-	{
-		if (UAbilitySystemComponent* LocalASC = UAbilitySystemGlobals::GetAbilitySystemComponentFromActor(FindLocalDebugActor()))
-		{
-			LocalAttributes = CollectAttributeData(OwnerPC, LocalASC);
-		}
-	}
 
-	bool bAllAttributesNetworked = true;
+	int NumAttributeCounts[+ENetworkStatus::MAX] = { 0 };
 
 	// Reverse the order of iteration so RemoveAt has a better chance at removing near the end
 	TArray<FDebugAttributeData> AttributeDebugData;
@@ -632,11 +715,13 @@ void FGameplayDebuggerCategory_Abilities::DrawGameplayAttributes(FGameplayDebugg
 	{
 		const FRepData::FGameplayAttributeDebug& LocalAttribute = LocalAttributes[Index];
 
-		FDebugAttributeData& DebugDatum = AttributeDebugData.AddDefaulted_GetRef();
-		DebugDatum.AttributeName = LocalAttribute.AttributeName;
-		DebugDatum.LocalBaseValue = LocalAttribute.BaseValue;
-		DebugDatum.LocalCurrentValue = LocalAttribute.CurrentValue;
-		DebugDatum.NetworkStatus = FRepData::ENetworkStatus::LocalOnly;
+		FDebugAttributeData& DebugDatum = AttributeDebugData.Add_GetRef(
+			FDebugAttributeData{
+				.AttributeName = LocalAttribute.AttributeName,
+				.LocalBaseValue = LocalAttribute.BaseValue,
+				.LocalCurrentValue = LocalAttribute.CurrentValue,
+				.NetworkStatus = ENetworkStatus::LocalOnly
+			});
 
 		int32 ServerIndex = ServerAttributes.FindLastByPredicate([FindName = LocalAttribute.AttributeName](const FRepData::FGameplayAttributeDebug& Item) { return FindName == Item.AttributeName; });
 		if (ServerIndex != INDEX_NONE)
@@ -645,26 +730,23 @@ void FGameplayDebuggerCategory_Abilities::DrawGameplayAttributes(FGameplayDebugg
 
 			DebugDatum.ServerBaseValue = ServerAttribute.BaseValue;
 			DebugDatum.ServerCurrentValue = ServerAttribute.CurrentValue;
-			DebugDatum.NetworkStatus = ServerAttribute.NetworkStatus == (FRepData::ENetworkStatus::Networked) ? FRepData::ENetworkStatus::Networked : FRepData::ENetworkStatus::Detached;
+			DebugDatum.NetworkStatus = (ServerAttribute.NetworkStatus == ENetworkStatus::Networked) ? ENetworkStatus::Networked : ENetworkStatus::Detached;
 
 			// Remove them from server-only array
-			ServerAttributes.RemoveAt(ServerIndex);
+			ServerAttributes.RemoveAtSwap(ServerIndex, EAllowShrinking::No);
 		}
-
-		if (DebugDatum.NetworkStatus != FRepData::ENetworkStatus::Networked)
-			bAllAttributesNetworked = false;
+		++NumAttributeCounts[+DebugDatum.NetworkStatus];
 	}
 
 	// If any entries are still in this array, they are server only and mark them as such
 	for (const FRepData::FGameplayAttributeDebug& ServerAttribute : ServerAttributes)
 	{
-		bAllAttributesNetworked = false;
-
 		FDebugAttributeData& DebugDatum = AttributeDebugData.AddDefaulted_GetRef();
 		DebugDatum.AttributeName = ServerAttribute.AttributeName;
 		DebugDatum.ServerBaseValue = ServerAttribute.BaseValue;
 		DebugDatum.ServerCurrentValue = ServerAttribute.CurrentValue;
-		DebugDatum.NetworkStatus = FRepData::ENetworkStatus::ServerOnly;
+		DebugDatum.NetworkStatus = ENetworkStatus::ServerOnly;
+		++NumAttributeCounts[+DebugDatum.NetworkStatus];
 	}
 
 	// Finally sort to keep everything in alphabetical order
@@ -686,47 +768,41 @@ void FGameplayDebuggerCategory_Abilities::DrawGameplayAttributes(FGameplayDebugg
 	const int NumColumns = FMath::Max(1, FMath::FloorToInt(CanvasWidth / ColumnWidth));
 
 	// Figure out the colors we're going to be using and print the legend
-	const TCHAR* ExpectedColor = BothColor;
-	const TCHAR* Modified = TEXT("*");
-	FString HeaderText = FString::Printf(TEXT("Attributes [%d]"), AttributeDebugData.Num());
+	FString LegendText;
+	const bool bAllAttributesNetworked = (0 == (NumAttributeCounts[+ENetworkStatus::ServerOnly] + NumAttributeCounts[+ENetworkStatus::LocalOnly] + NumAttributeCounts[+ENetworkStatus::Detached]));
 	if (!bConsiderNetworkStatus || bAllAttributesNetworked)
 	{
-		// Normal case: no legend required, single color which is the local case
-		ExpectedColor = ServerColor;
+		// Normal case: no legend required
 	}
 	else if (!bAllAttributesNetworked)
 	{
 		// Not all attributes are networked; we should display all cases for easier debugging
-		HeaderText.Appendf(TEXT("\tLegend:  %sBoth  %sServer  %sLocal  %sNonReplicated"), BothColor, ServerColor, LocalColor, NonReplicatedColor);
-		ExpectedColor = BothColor;
+		LegendText = FString::Printf(TEXT("Legend:  %sBoth [%d]    %sServer [%d]    %sLocal [%d]    %sNonReplicated [%d]"),
+			BothColor, NumAttributeCounts[+ENetworkStatus::Networked],
+			ServerColor, NumAttributeCounts[+ENetworkStatus::ServerOnly],
+			LocalColor, NumAttributeCounts[+ENetworkStatus::LocalOnly],
+			NonReplicatedColor, NumAttributeCounts[+ENetworkStatus::Detached]);
 	}
 	else
 	{
 		// All attributes are networked; the values may differ but assume the server is correct
-		HeaderText.Appendf(TEXT("\tLegend:  %sServer  %sLocal"), ServerColor, LocalColor);
-		ExpectedColor = ServerColor;
+		LegendText = FString::Printf(TEXT("Legend:  %sServer    %sLocal"), ServerColor, LocalColor);
 	}
+	CanvasContext.Printf(TEXT("Attributes [%d]:"), AttributeDebugData.Num());
+	CanvasContext.CursorX += 200.0f;
+	CanvasContext.CursorY -= CanvasContext.GetLineHeight();
+	CanvasContext.Print(LegendText);
 
-	const TCHAR* StatusTexts[+FRepData::ENetworkStatus::MAX];
-	StatusTexts[+FRepData::ENetworkStatus::Networked] = ExpectedColor; // This one is going to change based on 'expected' value (e.g. if all local, keep the UI consistently colored)
-	StatusTexts[+FRepData::ENetworkStatus::ServerOnly] = ServerColor;
-	StatusTexts[+FRepData::ENetworkStatus::LocalOnly] = LocalColor;
-	StatusTexts[+FRepData::ENetworkStatus::Detached] = NonReplicatedColor;
-
-	CanvasContext.Print(HeaderText);
 	CanvasContext.CursorX += Padding;
-
 	for (const FDebugAttributeData& AttributeData : AttributeDebugData)
 	{
-		const bool bModified = (AttributeData.ServerBaseValue != AttributeData.ServerCurrentValue) || (AttributeData.LocalBaseValue != AttributeData.LocalCurrentValue);
-		const TCHAR* StatusText = StatusTexts[+AttributeData.NetworkStatus];
-
 		const bool bServerValueMatch = (AttributeData.ServerBaseValue == AttributeData.ServerCurrentValue);
 		const bool bLocalValueMatch = (AttributeData.LocalBaseValue == AttributeData.LocalCurrentValue);
+		const bool bModified = !bServerValueMatch || !bLocalValueMatch;
 		const bool bNetworkValueMatch = (AttributeData.LocalBaseValue == AttributeData.ServerBaseValue) && (AttributeData.LocalCurrentValue == AttributeData.ServerCurrentValue);
 
 		// Let's build up the attribute value string which is just trying to represent the four states: srv cur [srv base] local cur [local base] in as little text as possible
-		TStringBuilder<1024> AttributeValueStr;
+		TStringBuilder<64> AttributeValueStr;
 		if (bNetworkValueMatch && bServerValueMatch && bLocalValueMatch)
 		{
 			// Everything matches, let's choose white and any one of the values (since they match)
@@ -734,8 +810,8 @@ void FGameplayDebuggerCategory_Abilities::DrawGameplayAttributes(FGameplayDebugg
 		}
 		else
 		{
-			const bool bDisplayServerValue = (AttributeData.NetworkStatus != FRepData::ENetworkStatus::LocalOnly);
-			const bool bDisplayClientValue = (AttributeData.NetworkStatus != FRepData::ENetworkStatus::ServerOnly) && !bNetworkValueMatch;
+			const bool bDisplayServerValue = (AttributeData.NetworkStatus != ENetworkStatus::LocalOnly);
+			const bool bDisplayClientValue = (AttributeData.NetworkStatus != ENetworkStatus::ServerOnly) && !bNetworkValueMatch;
 
 			// Append status that is happening on the Server
 			if (bDisplayServerValue)
@@ -770,7 +846,9 @@ void FGameplayDebuggerCategory_Abilities::DrawGameplayAttributes(FGameplayDebugg
 		const float CursorX = CanvasContext.CursorX;
 		const float CursorY = CanvasContext.CursorY;
 
-		const FString AttributeDebugText = FString::Printf(TEXT("%s%s%.35s: %s"), bModified ? Modified : TEXT(""), StatusText, *AttributeData.AttributeName, AttributeValueStr.ToString());
+		const ENetworkStatus NetworkStatus = bConsiderNetworkStatus ? AttributeData.NetworkStatus : ENetworkStatus::Networked;
+		const FString ColoredAttributeName = ColorNetworkString(NetworkStatus, AttributeData.AttributeName.Left(LongestDebugObjectName.Len()));
+		const FString AttributeDebugText = FString::Printf(TEXT("%s%s: %.*s"), bModified ? TEXT("*") : TEXT(""), *ColoredAttributeName, AttributeValueStr.Len(), AttributeValueStr.GetData());
 		CanvasContext.PrintAt(CursorX, CursorY, AttributeDebugText);
 
 		// PrintAt would have reset these values, restore them.
diff --git a/Engine/Plugins/Runtime/GameplayAbilities/Source/GameplayAbilities/Private/GameplayDebuggerCategory_Abilities.h b/Engine/Plugins/Runtime/GameplayAbilities/Source/GameplayAbilities/Private/GameplayDebuggerCategory_Abilities.h
--- a/Engine/Plugins/Runtime/GameplayAbilities/Source/GameplayAbilities/Private/GameplayDebuggerCategory_Abilities.h
+++ b/Engine/Plugins/Runtime/GameplayAbilities/Source/GameplayAbilities/Private/GameplayDebuggerCategory_Abilities.h
@@ -7,6 +7,7 @@
 #if WITH_GAMEPLAY_DEBUGGER_MENU
 #include "GameplayTagContainer.h"
 #include "GameplayDebuggerCategory.h"
+#include "GameplayPrediction.h"
 
 class AActor;
 class APlayerController;
@@ -28,6 +29,17 @@ public:
 	void OnShowGameplayEffectsToggle();
 	void OnShowGameplayAttributesToggle();
 		
+
+	// Some GAS features such as Attributes can exist on the server, client, or both.  We can also get 'detached' if both sides have the same values (such as Attributes) that aren't networked.
+	// We are reusing this same concept to try and reconcile Gameplay Effects predicted locally, or triggered server-side.
+	enum class ENetworkStatus : uint8
+	{
+		ServerOnly, LocalOnly, Networked, Detached, MAX
+	};
+
+	// Unary operator + for quick conversion from enum class to int32
+	friend constexpr int32 operator+(const ENetworkStatus& value) { return static_cast<int32>(value); }
+
 protected:
 
 	void DrawGameplayTags(FGameplayDebuggerCanvasContext& CanvasContext, const APlayerController* OwnerPC) const;
@@ -54,20 +66,18 @@ protected:
 
 		struct FGameplayEffectDebug
 		{
+			FPredictionKey PredictionKey;
 			FString Effect;
 			FString Context;
 			float Duration = 0.0f;
 			float Period = 0.0f;
 			int32 Stacks = 0;
 			float Level = 0.0f;
-		};
-		TArray<FGameplayEffectDebug> GameplayEffects;
 
-		// Attributes can exist purposefully on ServerOnly or LocalOnly.  They are by default Networked on both.  We can get 'detached' if both sides have attributes that aren't networked.
-		enum class ENetworkStatus : uint8
-		{
-			ServerOnly, LocalOnly, Networked, Detached, MAX
+			ENetworkStatus NetworkStatus = ENetworkStatus::ServerOnly;
+			bool bInhibited = false;
 		};
+		TArray<FGameplayEffectDebug> GameplayEffects;
 
 		struct FGameplayAttributeDebug
 		{
@@ -79,16 +89,17 @@ protected:
 		TArray<FGameplayAttributeDebug> Attributes;
 
 		void Serialize(FArchive& Ar);
-
-		// Unary operator + for quick conversion from enum class to int32
-		friend constexpr int32 operator+(const FGameplayDebuggerCategory_Abilities::FRepData::ENetworkStatus& value) { return static_cast<int32>(value); }
 	};
 	FRepData DataPack;
 
 	bool WrapStringAccordingToViewport(const FString& iStr, FString& oStr, FGameplayDebuggerCanvasContext& CanvasContext, float ViewportWitdh) const;
 
 private:
-	TArray<FRepData::FGameplayAttributeDebug> CollectAttributeData(const APlayerController* OwnerPC, const UAbilitySystemComponent* ASC) const;
+	TArray<FRepData::FGameplayAttributeDebug> CollectAttributeData(const APlayerController* OwnerPC, const UAbilitySystemComponent* AbilityComp) const;
+	TArray<FRepData::FGameplayEffectDebug> CollectEffectsData(const APlayerController* OwnerPC, const UAbilitySystemComponent* AbilityComp) const;
+
+	// Save off the last expected draw size so that we can draw a border around it next frame (and hope we're the same size)
+	float LastDrawDataEndSize = 0.0f;
 
 	bool bShowGameplayTags = true;
 	bool bShowGameplayAbilities = true;
