Static analysis fix.

#rnx

diff --git a/Engine/Plugins/Runtime/GameplayAbilities/Source/GameplayAbilities/Private/GameplayDebuggerCategory_Abilities.cpp b/Engine/Plugins/Runtime/GameplayAbilities/Source/GameplayAbilities/Private/GameplayDebuggerCategory_Abilities.cpp
--- a/Engine/Plugins/Runtime/GameplayAbilities/Source/GameplayAbilities/Private/GameplayDebuggerCategory_Abilities.cpp
+++ b/Engine/Plugins/Runtime/GameplayAbilities/Source/GameplayAbilities/Private/GameplayDebuggerCategory_Abilities.cpp
@@ -252,7 +252,7 @@ TArray<FGameplayDebuggerCategory_Abilities::FRepData::FGameplayAttributeDebug> F
 	for (const UAttributeSet* AttributeSet : AbilityComp->GetSpawnedAttributes())
 	{
 		const TSubclassOf<UAttributeSet> AttributeSetClass = AttributeSet ? AttributeSet->GetClass() : nullptr;
-		if (!AttributeSetClass)
+		if (!AttributeSet || !AttributeSetClass) // !AttributeSet needed for static analysis
 		{
 			continue;
 		}
