Gameplay FPredictionKey Fixes.

1.  Previously, when calling a CaughtUp or Rejected callback, the FPredictionKey Replication Array could realloc while you were holding on to the reference.
2.  Introduced some debug functionality that will attempt to catch these extremely hard to find cases where an FPredictionKey could be created and never acknowledged.
3.  Added CVars for controlling behavior on how FPredictionKey Dependency Chains should behave.

diff --git a/Engine/Plugins/Runtime/GameplayAbilities/Source/GameplayAbilities/Private/GameplayPrediction.cpp b/Engine/Plugins/Runtime/GameplayAbilities/Source/GameplayAbilities/Private/GameplayPrediction.cpp
--- a/Engine/Plugins/Runtime/GameplayAbilities/Source/GameplayAbilities/Private/GameplayPrediction.cpp
+++ b/Engine/Plugins/Runtime/GameplayAbilities/Source/GameplayAbilities/Private/GameplayPrediction.cpp
@@ -3,9 +3,67 @@
 #include "GameplayPrediction.h"
 #include "AbilitySystemComponent.h"
 #include "AbilitySystemLog.h"
+#include "Engine/World.h"
+#include "Engine/NetDriver.h"
 
 #include UE_INLINE_GENERATED_CPP_BY_NAME(GameplayPrediction)
 
+DEFINE_LOG_CATEGORY_STATIC(LogPredictionKey, Display, All);
+
+namespace UE::AbilitySystem::Private
+{
+	// This controls the client-side code that catches up missed/unacknowledged keys.  In theory this can happen in bad network conditions where
+	// a server can fail to acknowledge a key due to packet loss, but it also happens due to bugs where we create & register prediction keys that
+	// are never sent to the server.  The science behind this number isn't exact: we default to enough that all local players can exhaust their ring buffer.
+	// Set it lower to catch bugs related to not sending prediction keys sooner in your logs, or higher if you have a massive local player count.
+	constexpr int32 NumExpectedLocalPlayers = 4;
+	int32 CVarMaxStaleKeysBeforeAckValue = FReplicatedPredictionKeyMap::KeyRingBufferSize * NumExpectedLocalPlayers;
+	FAutoConsoleVariableRef CVarMaxStaleKeysBeforeAck(TEXT("AbilitySystem.PredictionKey.MaxStaleKeysBeforeAck"), CVarMaxStaleKeysBeforeAckValue,
+		TEXT("How many prediction keys can be dropped before StaleKeyBehavior is run."));
+
+	// What should we do with these old stale FPredictionKeys?  Prior to UE5.4, we always CaughtUp.  I believe it's actually safer to drop.
+	int32 CVarStaleKeyBehaviorValue = 2;
+	FAutoConsoleVariableRef CVarStaleKeyBehavior(TEXT("AbilitySystem.PredictionKey.StaleKeyBehavior"), CVarStaleKeyBehaviorValue,
+		TEXT("How do we handle stale keys? 0 = CaughtUp. 1 = Reject. 2 = Drop"));
+
+	// How should we deal with dependent keys (in a chain)?  Prior to UE5.4, old keys implied new keys.  We introduced some new functionality (explained in the help text).
+	// 0 (no bitmask) is legacy behavior.  Logically, (1 | (1<<1)) = 3 is the correct value. We default to 1 to move towards a long-term fix (of value 3).
+	int32 CVarDependentChainBehaviorValue = 1;
+	FAutoConsoleVariableRef CVarDependentChainBehavior(TEXT("AbilitySystem.PredictionKey.DepChainBehavior"), CVarDependentChainBehaviorValue,
+		TEXT("How do we handle dependency key chains? Bitmask: 0 = Old Accept/Rejected Implies Newer Accepted/Rejected. 1 = Newer Accepted also implies Older Accepted. 2 = Old Accepted Does NOT imply Newer Accepted"));
+
+	/**
+	 * Given an FProperty Link (such as a UFunction's Properties, or a UStruct's Properties), find and return any linked FPredictionKeys.
+	 * This is useful for determining if we're sending FPredictionKeys via RPC.
+	 */
+	TArray<FPredictionKey> FindPredictionKeysInPropertyLink(const FProperty* InPropertyLink, uint8_t* StartDataOffset)
+	{
+		TArray<FPredictionKey> PredictionKeys;
+
+		// Go through the PropertyLinks and look for an FPredictionKey or FPredictionKey nested in a FStruct...
+		for (const FProperty* Property = InPropertyLink; Property; Property = Property->PropertyLinkNext)
+		{
+			if (const FStructProperty* StructProperty = CastField<FStructProperty>(Property))
+			{
+				if (FPredictionKey::StaticStruct() == StructProperty->Struct)
+				{
+					// Now if this FPredictionKey argument, get it from memory and prefer this value.
+					FPredictionKey& PassedInArg = PredictionKeys.AddDefaulted_GetRef();
+					Property->CopyCompleteValue(&PassedInArg, StartDataOffset + Property->GetOffset_ForUFunction());
+				}
+				else if (StructProperty->Struct)
+				{
+					// Find a nested prediction key, but don't return it (prefer one at the top-level)
+					TArray<FPredictionKey> InnerKeys = FindPredictionKeysInPropertyLink(StructProperty->Struct->PropertyLink, StartDataOffset + StructProperty->GetOffset_ForUFunction());
+					PredictionKeys.Append( MoveTempIfPossible(InnerKeys) );
+				}
+			}
+		}
+
+		return PredictionKeys;
+	}
+}
+
 FReplicatedPredictionKeyItem::FReplicatedPredictionKeyItem()
 {
 };
@@ -111,7 +169,6 @@ void FPredictionKey::GenerateNewPredictionKey()
 	{
 		GKey = 1;
 	}
-	bIsStale = false;
 }
 
 void FPredictionKey::GenerateDependentPredictionKey()
@@ -130,13 +187,15 @@ void FPredictionKey::GenerateDependentPredictionKey()
 
 	GenerateNewPredictionKey();
 
+	ensureAlwaysMsgf((Base == 0) || (Current - Base < 20), TEXT("Deep PredictionKey Chain Detected.  It's likely there's circular logic that could stack overflow."));
+
 	if (Previous > 0)
 	{
 		FPredictionKeyDelegates::AddDependency(Current, Previous);
 	}
 }
 
-FPredictionKey FPredictionKey::CreateNewPredictionKey(UAbilitySystemComponent* OwningComponent)
+FPredictionKey FPredictionKey::CreateNewPredictionKey(const UAbilitySystemComponent* OwningComponent)
 {
 	FPredictionKey NewKey;
 	
@@ -148,7 +207,7 @@ FPredictionKey FPredictionKey::CreateNewPredictionKey(UAbilitySystemComponent* O
 	return NewKey;
 }
 
-FPredictionKey FPredictionKey::CreateNewServerInitiatedKey(UAbilitySystemComponent* OwningComponent)
+FPredictionKey FPredictionKey::CreateNewServerInitiatedKey(const UAbilitySystemComponent* OwningComponent)
 {
 	FPredictionKey NewKey;
 
@@ -235,12 +294,15 @@ void FPredictionKeyDelegates::Reject(FPredictionKey::KeyType Key)
 	FDelegates* DelPtr = Get().DelegateMap.Find(Key);
 	if (DelPtr)
 	{
-		for (auto& Delegate : DelPtr->RejectedDelegates)
+		// Copy & Remove first, because during the delegate, we're likely to call this again, 
+		// thereby invalidating DelPtr.
+		TArray<FPredictionKeyEvent> RejectedDelegates = MoveTemp(DelPtr->RejectedDelegates);
+		Get().DelegateMap.Remove(Key);
+
+		for (auto& Delegate : RejectedDelegates)
 		{
 			Delegate.ExecuteIfBound();
 		}
-
-		Get().DelegateMap.Remove(Key);
 	}
 }
 
@@ -249,31 +311,40 @@ void FPredictionKeyDelegates::CatchUpTo(FPredictionKey::KeyType Key)
 	FDelegates* DelPtr = Get().DelegateMap.Find(Key);
 	if (DelPtr)
 	{
-		for (auto& Delegate : DelPtr->CaughtUpDelegates)
-		{
-			Delegate.ExecuteIfBound();
-		}
+		// Copy & Remove first, because during the delegate, we're likely to call this again, 
+		// thereby invalidating DelPtr.
+		TArray<FPredictionKeyEvent> CaughtUpDelegates = MoveTemp(DelPtr->CaughtUpDelegates);
 		Get().DelegateMap.Remove(Key);
-	}
-
-
-#if 0 
-	// Sanity checking
-	for (auto MapIt = Get().DelegateMap.CreateIterator(); MapIt; ++MapIt)
-	{
-		if (MapIt.Key() <= Key)
-		{	
-			ABILITY_LOG(Warning, TEXT("PredictionKey is stale: %d after CatchUp to: %d"), MapIt.Key(), Key);
 
+		for (auto& Delegate : CaughtUpDelegates)
+		{
+			Delegate.ExecuteIfBound();
 		}
 	}
-#endif
 }
 
 void FPredictionKeyDelegates::AddDependency(FPredictionKey::KeyType ThisKey, FPredictionKey::KeyType DependsOn)
 {
+	// If we Reject the BaseKey, then notify ThisKey it has also been Rejected.
 	NewRejectedDelegate(DependsOn).BindStatic(&FPredictionKeyDelegates::Reject, ThisKey);
-	NewCaughtUpDelegate(DependsOn).BindStatic(&FPredictionKeyDelegates::CatchUpTo, ThisKey);
+
+	// 1. If we receive confirmation that a newer key happened, notify the previous one it also happened.
+	//	This allows a FPredictionKey to _not_ be sent to the server, yet be acknowledged as long as a newer dependent key was sent & acknowledged.
+	//	This is logically correct and newly introduced behavior.
+	if ((UE::AbilitySystem::Private::CVarDependentChainBehaviorValue & 1) != 0)
+	{
+		NewCaughtUpDelegate(ThisKey).BindStatic(&FPredictionKeyDelegates::CatchUpTo, DependsOn);
+	}
+
+	// 2. If we receive an earlier key in the chain of dependencies, then the later ones should also be acknowledged. 
+	//	This is unintuitive, but it's possible to create FPredictionKeys that are never sent to the server.
+	//	This is the case with FScopedServerAbilityRPCBatcher.  It sends only the BaseKey but needs to notify dependents.
+	//	This isn't logically correct, so we're introducing a way to disable this legacy functionality once Engine fixes are finished.
+	if ((UE::AbilitySystem::Private::CVarDependentChainBehaviorValue & 2) == 0)
+	{
+		NewCaughtUpDelegate(DependsOn).BindStatic(&FPredictionKeyDelegates::CatchUpTo, ThisKey);
+	}
+	
 }
 
 // -------------------------------------
@@ -320,19 +391,80 @@ FScopedPredictionWindow::FScopedPredictionWindow(UAbilitySystemComponent* InAbil
 	{
 		return;
 	}
-	
-	// InAbilitySystemComponent->GetPredictionKey().IsValidForMorePrediction() == false && 
+
+	// Because of the check above, this is expected to only run on the client.
 	if (bCanGenerateNewKey)
 	{
 		check(InAbilitySystemComponent != NULL); // Should have bailed above with ensure(Owner.IsValid())
 		ClearScopedPredictionKey = true;
 		RestoreKey = InAbilitySystemComponent->ScopedPredictionKey;
-		InAbilitySystemComponent->ScopedPredictionKey.GenerateDependentPredictionKey();		
+		InAbilitySystemComponent->ScopedPredictionKey.GenerateDependentPredictionKey();
+	}
+
+#if !UE_BUILD_SHIPPING
+	// Add some debugging functionality if we're the first scoped prediction window (will become a BaseKey value)
+	if (bCanGenerateNewKey && !RestoreKey.IsValidKey())
+	{
+		// Intercept any RPC's of FPredictionKeys so that we can mark them as sent to the server.
+		// For RPC's not marked as sent-to-server, we can log them and track them down with breakpoints.
+		const UWorld* NetWorldPtr = Owner->GetWorld();
+		UNetDriver* NetDriver = NetWorldPtr ? NetWorldPtr->GetNetDriver() : nullptr;
+		if (NetDriver && NetDriver->GetNetMode() == NM_Client)
+		{
+			DebugSavedNetDriver = NetDriver;
+			DebugSavedOnSendRPC = NetDriver->SendRPCDel;
+			DebugBaseKeyOfChain = InAbilitySystemComponent->ScopedPredictionKey.Current;
+
+			NetDriver->SendRPCDel.BindWeakLambda(InAbilitySystemComponent,
+				[this, SendRPCDel = DebugSavedOnSendRPC](AActor* Actor, UFunction* Function, void* Parameters, FOutParmRec* OutParms, FFrame* Stack, UObject* SubObject, bool& bBlockSendRPC)
+				{
+					// Chain the previous RPC callback (e.g. packet loss simulation)
+					SendRPCDel.ExecuteIfBound(Actor, Function, Parameters, OutParms, Stack, SubObject, bBlockSendRPC);
+
+					// If we've already reset this, it's an indication this PredictionKey Chain has already been sent.
+					if (!DebugBaseKeyOfChain.IsSet())
+					{
+						return;
+					}
+
+					const FPredictionKey::KeyType BaseKey = DebugBaseKeyOfChain.GetValue();
+
+					// See if we are communicating a key based on this scope, if so, reset our debug value so as not to warn that it was never sent on Scope destruct.
+					TArray<FPredictionKey> PassedInKeys = UE::AbilitySystem::Private::FindPredictionKeysInPropertyLink(Function->PropertyLink, static_cast<uint8_t*>(Parameters));
+					for (const FPredictionKey& PassedInKey : PassedInKeys)
+					{
+						if (PassedInKey.Current == BaseKey || PassedInKey.Base == BaseKey)
+						{
+							UE_LOG(LogPredictionKey, Verbose, TEXT("Sent key %s to confirm BaseKey %d"), *PassedInKey.ToString(), BaseKey);
+							DebugBaseKeyOfChain.Reset();
+							break;
+						}
+					}
+				});
+		}
 	}
+#endif
 }
 
 FScopedPredictionWindow::~FScopedPredictionWindow()
 {
+#if !UE_BUILD_SHIPPING
+	// Stop intercepting RPC calls
+	if (UNetDriver* NetDriver = DebugSavedNetDriver.Get())
+	{
+		if (DebugBaseKeyOfChain.IsSet())
+		{
+			const FPredictionKey::KeyType BaseKey = DebugBaseKeyOfChain.GetValue();
+			UE_CLOG(FPredictionKeyDelegates::Get().DelegateMap.Contains(BaseKey), LogPredictionKey, Warning, TEXT("No key based off PredictionKey %d was communicated to the server during ScopedPredictionWindow. This will likely leak a FPredictionKey with dangling CaughtUpTo/Rejected delegates."), BaseKey);
+		}
+		
+		NetDriver->SendRPCDel = DebugSavedOnSendRPC;
+		DebugSavedOnSendRPC.Unbind();
+		DebugSavedNetDriver.Reset();
+		DebugBaseKeyOfChain.Reset();
+	}
+#endif
+
 	if (UAbilitySystemComponent* OwnerPtr = Owner.Get())
 	{
 		if (SetReplicatedPredictionKey)
@@ -355,34 +487,85 @@ FScopedPredictionWindow::~FScopedPredictionWindow()
 
 // -----------------------------------
 
-void FReplicatedPredictionKeyItem::OnRep()
+void FReplicatedPredictionKeyItem::OnRep(const FReplicatedPredictionKeyMap& InArray)
 {
+	using namespace UE::AbilitySystem::Private;
+
+	// If either of these logs are set to Verbose, log on one (and only one) of them.
 	ABILITY_LOG(Verbose, TEXT("FReplicatedPredictionKeyItem::OnRep %s"), *PredictionKey.ToString());
-	
-	// Every predictive action we've done up to and including the current value of ReplicatedPredictionKey needs to be wiped
+	if (!UE_LOG_ACTIVE(LogAbilitySystem, Verbose))
+	{
+		UE_LOG(LogPredictionKey, Verbose, TEXT("FReplicatedPredictionKeyItem::OnRep %s"), *PredictionKey.ToString());
+	}
+
+	// Every predictive action we've done in the current chain of dependencies (including the current value) of ReplicatedPredictionKey needs to be acknowledged
 	FPredictionKeyDelegates::CatchUpTo(PredictionKey.Current);
 
-	// Sanity checking
-	int32 Index = PredictionKey.Current % FReplicatedPredictionKeyMap::KeyRingBufferSize;
+	// Remove Stale Prediction Keys
+	//
+	// Let's check for older keys that should have been cleaned up.  This is an indication that these FPredictionKeys were never acknowledged
+	// and one way that can happen, is if no FPredictionKeys in a dependency chain are sent to the server.  In such a case, you've locally predicted
+	// a bunch of events that the server has no way of accepting/rejecting.  A decision was made long ago to auto-accept these, so let's keep that functionality.
+	const int MaxKeyLagValue = FMath::Max(CVarMaxStaleKeysBeforeAckValue, FReplicatedPredictionKeyMap::KeyRingBufferSize);
+
+	// We now define a range from [Min,Max] but since it's a circular buffer, it can also be [0,Min] [Max,KeyMax]
+	constexpr int32 MaxKeyValue = std::numeric_limits<FPredictionKey::KeyType>::max();
+	int32 RangeMin = (PredictionKey.Current - MaxKeyLagValue) > 0 ? (PredictionKey.Current - MaxKeyLagValue) : MaxKeyValue + (PredictionKey.Current - MaxKeyLagValue);
+	int32 RangeMax = (PredictionKey.Current + MaxKeyLagValue) % MaxKeyValue;
+
+	// If Max > Min, then we're not wrapped and the safe range is [Min,Max]
+	const bool bRangeIsSafeZone = (RangeMax > RangeMin);
+	if (!bRangeIsSafeZone)
+	{
+		// Otherwise, let's swap so [Min, Max] is the unsafe zone and we'll pair it with !bRangeIsSafeZone.
+		Swap(RangeMin, RangeMax);
+	}
+
+	// Go through the unordered delegates and look for old entries
+	TArray<FPredictionKey::KeyType> StalePredictionKeys;
 	for (auto MapIt = FPredictionKeyDelegates::Get().DelegateMap.CreateIterator(); MapIt; ++MapIt)
 	{
-		// If older key
-		if (MapIt.Key() <= PredictionKey.Current)
-		{	
-			// Older key that would have gone in this slot
-			if (MapIt.Key() % FReplicatedPredictionKeyMap::KeyRingBufferSize == Index)
-			{
-				// Message the log, this can happen during normal gameplay due to replication order, but can also indicate an ability-specific issue
-				ABILITY_LOG(Log, TEXT("Passed PredictionKey %d in Delegate map while OnRep'ing %s"), MapIt.Key(), *PredictionKey.ToString());
+		const FPredictionKey::KeyType ExistingKey = MapIt.Key();
 
-				// Execute CaughtUp delegates
-				for (auto& Delegate : MapIt.Value().CaughtUpDelegates)
-				{
-					Delegate.ExecuteIfBound();
-				}
+		// Now check if we're within the range specified.  If we are, and the range is not the safe zone (or vice versa) catch-up.
+		const bool bWithinRange = (ExistingKey >= RangeMin) && (ExistingKey <= RangeMax);
+		if (bWithinRange ^ bRangeIsSafeZone)
+		{
+			StalePredictionKeys.Add(ExistingKey);
+		}
+	}
 
-				// Cleanup
-				MapIt.RemoveCurrent();
+	// If there were any unacknowledged keys who have since timed out
+	if (StalePredictionKeys.Num() > 0)
+	{
+		// Sort so we do the acks in the correct order
+		StalePredictionKeys.Sort();
+		for (FPredictionKey::KeyType Key : StalePredictionKeys)
+		{
+			const int32 KeyIndex = Key % InArray.KeyRingBufferSize;
+			if (!ensure(InArray.PredictionKeys.IsValidIndex(KeyIndex)))
+			{
+				continue;
+			}
+
+			// We don't want to be creating these FPredictionKeys that aren't sent to the server.  We should warn the user that they need to accept/reject them based on their own logic.
+			const FPredictionKey& ExistingKeyInSlot = InArray.PredictionKeys[KeyIndex].PredictionKey;
+			UE_LOG(LogPredictionKey, Warning, TEXT("UnAck'd PredictionKey %d in DelegateMap while OnRep %s. This indicates keychains are created w/o being sent to & ack'd by the server. Last ack'd key in that slot was %s."), Key, *PredictionKey.ToString(), *ExistingKeyInSlot.ToString());
+			
+			if (CVarStaleKeyBehaviorValue == 0)
+			{
+				// This is legacy functionality.
+				FPredictionKeyDelegates::CatchUpTo(Key);
+			}
+			else if (CVarStaleKeyBehaviorValue == 1)
+			{
+				// Alternatively, this may be more appropriate as the server didn't specifically acknowledge us.
+				FPredictionKeyDelegates::Reject(Key);
+			}
+			else
+			{
+				// This is the safest.  It could leave Abilities waiting in async tasks forever, but should never crash.
+				FPredictionKeyDelegates::Get().DelegateMap.Remove(Key);
 			}
 		}
 	}
diff --git a/Engine/Plugins/Runtime/GameplayAbilities/Source/GameplayAbilities/Public/GameplayPrediction.h b/Engine/Plugins/Runtime/GameplayAbilities/Source/GameplayAbilities/Public/GameplayPrediction.h
--- a/Engine/Plugins/Runtime/GameplayAbilities/Source/GameplayAbilities/Public/GameplayPrediction.h
+++ b/Engine/Plugins/Runtime/GameplayAbilities/Source/GameplayAbilities/Public/GameplayPrediction.h
@@ -2,6 +2,7 @@
 
 #pragma once
 
+#include "Engine/NetDriver.h"
 #include "Engine/NetSerialization.h"
 #include "Net/Serialization/FastArraySerializer.h"
 #include "GameplayPrediction.generated.h"
@@ -295,11 +296,6 @@ struct GAMEPLAYABILITIES_API FPredictionKey
 	typedef int16 KeyType;
 	FPredictionKey() = default;
 
-	/** On the server, what network connection this was serialized on. */
-	UE_DEPRECATED(5.0, "No longer used in favor of GetPredictiveConnectionKey, to avoid holding a direct object reference.")
-	UPROPERTY(NotReplicated)
-	TObjectPtr<UPackageMap> PredictiveConnection = nullptr;
-
 	/** The unique ID of this prediction key */
 	UPROPERTY()
 	int16	Current = 0;
@@ -308,19 +304,15 @@ struct GAMEPLAYABILITIES_API FPredictionKey
 	UPROPERTY(NotReplicated)
 	int16	Base = 0;
 
-	/** If stale, this key cannot be used for more prediction */
-	UPROPERTY(NotReplicated)
-	bool bIsStale = false;
-
 	/** True if this was created as a server initiated activation key, used to identify server activations but cannot be used for prediction */
 	UPROPERTY()
 	bool bIsServerInitiated = false;
 
 	/** Construct a new prediction key with no dependencies */
-	static FPredictionKey CreateNewPredictionKey(class UAbilitySystemComponent*);
+	static FPredictionKey CreateNewPredictionKey(const UAbilitySystemComponent*);
 
 	/** Construct a new server initiation key, for abilities activated on the server */
-	static FPredictionKey CreateNewServerInitiatedKey(class UAbilitySystemComponent*);
+	static FPredictionKey CreateNewServerInitiatedKey(const UAbilitySystemComponent*);
 
 	/** Create a new dependent prediction key: keep our existing base or use the current key as the base. */
 	void GenerateDependentPredictionKey();
@@ -357,7 +349,7 @@ struct GAMEPLAYABILITIES_API FPredictionKey
 	/** Can this key be used for more predictive actions, or has it already been sent off to the server? */
 	bool IsValidForMorePrediction() const
 	{
-		return Current > 0 && bIsStale == false && bIsServerInitiated == false;
+		return IsLocalClientKey();
 	}
 
 	/** Was this PredictionKey received from a NetSerialize or created locally? */
@@ -492,6 +484,12 @@ private:
 	bool ClearScopedPredictionKey;
 	bool SetReplicatedPredictionKey;
 	FPredictionKey RestoreKey;
+
+#if !UE_BUILD_SHIPPING
+	FOnSendRPC DebugSavedOnSendRPC;
+	TWeakObjectPtr<UNetDriver> DebugSavedNetDriver;
+	TOptional<FPredictionKey::KeyType> DebugBaseKeyOfChain;
+#endif
 };
 
 // -----------------------------------------------------------------
@@ -529,14 +527,14 @@ struct FReplicatedPredictionKeyItem : public FFastArraySerializerItem
 	UPROPERTY()
 	FPredictionKey PredictionKey;
 	
-	void PostReplicatedAdd(const struct FReplicatedPredictionKeyMap &InArray) { OnRep(); }
-	void PostReplicatedChange(const struct FReplicatedPredictionKeyMap &InArray) { OnRep(); }
+	void PostReplicatedAdd(const struct FReplicatedPredictionKeyMap &InArray) { OnRep(InArray); }
+	void PostReplicatedChange(const struct FReplicatedPredictionKeyMap &InArray) { OnRep(InArray); }
 
 	FString GetDebugString() { return PredictionKey.ToString(); }
 
 private:
 
-	void OnRep();
+	void OnRep(const struct FReplicatedPredictionKeyMap& InArray);
 };
 
 USTRUCT()
