Presumable fix for GameplayAbilities not always being removed from the ReplicatedSubObjectsList prior to being marked as garbage

diff --git a/Engine/Plugins/Runtime/GameplayAbilities/Source/GameplayAbilities/Private/AbilitySystemComponent_Abilities.cpp b/Engine/Plugins/Runtime/GameplayAbilities/Source/GameplayAbilities/Private/AbilitySystemComponent_Abilities.cpp
--- a/Engine/Plugins/Runtime/GameplayAbilities/Source/GameplayAbilities/Private/AbilitySystemComponent_Abilities.cpp
+++ b/Engine/Plugins/Runtime/GameplayAbilities/Source/GameplayAbilities/Private/AbilitySystemComponent_Abilities.cpp
@@ -654,6 +654,12 @@ void UAbilitySystemComponent::OnRemoveAbility(FGameplayAbilitySpec& Spec)
 	if (PrimaryInstance)
 	{
 		PrimaryInstance->OnRemoveAbility(AbilityActorInfo.Get(), Spec);
+		
+		// Make sure we remove this before marking it as garbage.
+		if (GetOwnerRole() == ROLE_Authority)
+		{
+			RemoveReplicatedInstancedAbility(PrimaryInstance);
+		}
 		PrimaryInstance->MarkAsGarbage();
 	}
 	else
