FScopedServerAbilityRPCBatcher should ensure all keys generated in its batch are dependent.  Otherwise, it's possible to leak FPredictionKeys as some could be generated and never acknowledged because the batch chooses only a single key to be sent (when it can be acknowledging as many as three different events).  By making the keys dependent, just a single acknowledgement will acknowledge all of those events.

diff --git a/Engine/Plugins/Runtime/GameplayAbilities/Source/GameplayAbilities/Private/AbilitySystemComponent_Abilities.cpp b/Engine/Plugins/Runtime/GameplayAbilities/Source/GameplayAbilities/Private/AbilitySystemComponent_Abilities.cpp
--- a/Engine/Plugins/Runtime/GameplayAbilities/Source/GameplayAbilities/Private/AbilitySystemComponent_Abilities.cpp
+++ b/Engine/Plugins/Runtime/GameplayAbilities/Source/GameplayAbilities/Private/AbilitySystemComponent_Abilities.cpp
@@ -3891,7 +3891,7 @@ int32 AbilitySystemLogServerRPCBatching = 0;
 static FAutoConsoleVariableRef CVarAbilitySystemLogServerRPCBatching(TEXT("AbilitySystem.ServerRPCBatching.Log"), AbilitySystemLogServerRPCBatching, TEXT(""), ECVF_Default	);
 
 FScopedServerAbilityRPCBatcher::FScopedServerAbilityRPCBatcher(UAbilitySystemComponent* InASC, FGameplayAbilitySpecHandle InAbilityHandle) 
-: ASC(InASC), AbilityHandle(InAbilityHandle)
+	: ASC(InASC), AbilityHandle(InAbilityHandle), ScopedPredictionWindow(InASC)
 {
 	if (ASC && AbilityHandle.IsValid() && ASC->ShouldDoServerAbilityRPCBatch())
 	{
diff --git a/Engine/Plugins/Runtime/GameplayAbilities/Source/GameplayAbilities/Public/Abilities/GameplayAbilityTypes.h b/Engine/Plugins/Runtime/GameplayAbilities/Source/GameplayAbilities/Public/Abilities/GameplayAbilityTypes.h
--- a/Engine/Plugins/Runtime/GameplayAbilities/Source/GameplayAbilities/Public/Abilities/GameplayAbilityTypes.h
+++ b/Engine/Plugins/Runtime/GameplayAbilities/Source/GameplayAbilities/Public/Abilities/GameplayAbilityTypes.h
@@ -494,6 +494,7 @@ private:
 
 	UAbilitySystemComponent* ASC;
 	FGameplayAbilitySpecHandle AbilityHandle;
+	FScopedPredictionWindow ScopedPredictionWindow;
 };
 
 
