AbilitiesGameplayEffectComponent could fail to grant and remove abilities if those abilities were granted/removed during the same scope lock (due to the PendingAdds/PendingRemoves).  Making this more clear by allowing users to query the pending lists in some interfaces.

diff --git a/Engine/Plugins/Runtime/GameplayAbilities/Source/GameplayAbilities/Private/AbilitySystemComponent_Abilities.cpp b/Engine/Plugins/Runtime/GameplayAbilities/Source/GameplayAbilities/Private/AbilitySystemComponent_Abilities.cpp
--- a/Engine/Plugins/Runtime/GameplayAbilities/Source/GameplayAbilities/Private/AbilitySystemComponent_Abilities.cpp
+++ b/Engine/Plugins/Runtime/GameplayAbilities/Source/GameplayAbilities/Private/AbilitySystemComponent_Abilities.cpp
@@ -862,7 +862,7 @@ void UAbilitySystemComponent::DecrementAbilityListLock()
 	}
 }
 
-FGameplayAbilitySpec* UAbilitySystemComponent::FindAbilitySpecFromHandle(FGameplayAbilitySpecHandle Handle) const
+FGameplayAbilitySpec* UAbilitySystemComponent::FindAbilitySpecFromHandle(FGameplayAbilitySpecHandle Handle, EConsiderPending ConsiderPending) const
 {
 	SCOPE_CYCLE_COUNTER(STAT_FindAbilitySpecFromHandle);
 
@@ -870,7 +870,21 @@ FGameplayAbilitySpec* UAbilitySystemComponent::FindAbilitySpecFromHandle(FGamepl
 	{
 		if (Spec.Handle == Handle)
 		{
-			return const_cast<FGameplayAbilitySpec*>(&Spec);
+			if (!Spec.PendingRemove || EnumHasAnyFlags(ConsiderPending, EConsiderPending::PendingRemove))
+			{
+				return const_cast<FGameplayAbilitySpec*>(&Spec);
+			}
+		}
+	}
+
+	if (EnumHasAnyFlags(ConsiderPending, EConsiderPending::PendingAdd))
+	{
+		for (const FGameplayAbilitySpec& Spec : AbilityPendingAdds)
+		{
+			if (!Spec.PendingRemove || EnumHasAnyFlags(ConsiderPending, EConsiderPending::PendingRemove))
+			{
+				return const_cast<FGameplayAbilitySpec*>(&Spec);
+			}
 		}
 	}
 
@@ -882,6 +896,42 @@ FGameplayAbilitySpec* UAbilitySystemComponent::FindAbilitySpecFromGEHandle(FActi
 	return nullptr;
 }
 
+TArray<const FGameplayAbilitySpec*> UAbilitySystemComponent::FindAbilitySpecsFromGEHandle(const FScopedAbilityListLock& /*Used as a Contract*/, FActiveGameplayEffectHandle ActiveGEHandle, EConsiderPending ConsiderPending) const
+{
+	TArray<const FGameplayAbilitySpec*> ReturnValue;
+
+	if (!ensureMsgf(IsOwnerActorAuthoritative(), TEXT("%hs is only valid on authority as FGameplayAbilitySpec::GameplayEffectHandle is not replicated and ability granting only happens on the server"), __func__))
+	{
+		return ReturnValue;
+	}
+
+	auto GatherGAsByGEHandle = [ActiveGEHandle, ConsiderPending, &ReturnValue](const TArrayView<const FGameplayAbilitySpec> AbilitiesToConsider)
+		{
+			for (const FGameplayAbilitySpec& GASpec : AbilitiesToConsider)
+			{
+				if (GASpec.GameplayEffectHandle == ActiveGEHandle)
+				{
+					if (!GASpec.PendingRemove || EnumHasAnyFlags(ConsiderPending, EConsiderPending::PendingRemove))
+					{
+						ReturnValue.Emplace(&GASpec);
+					}
+				}
+			}
+		};
+
+	// All activatable abilities (which will include abilities that are in AbilityPendingRemoves
+	GatherGAsByGEHandle(GetActivatableAbilities());
+
+	// If requested, specifically look for abilities that are pending add
+	if (EnumHasAnyFlags(ConsiderPending,EConsiderPending::PendingAdd))
+	{
+		GatherGAsByGEHandle(AbilityPendingAdds);
+	}
+
+	return ReturnValue;
+}
+
+
 FGameplayAbilitySpec* UAbilitySystemComponent::FindAbilitySpecFromClass(TSubclassOf<UGameplayAbility> InAbilityClass) const
 {
 	SCOPE_CYCLE_COUNTER(STAT_FindAbilitySpecFromHandle);
diff --git a/Engine/Plugins/Runtime/GameplayAbilities/Source/GameplayAbilities/Private/GameplayEffectComponents/AbilitiesGameplayEffectComponent.cpp b/Engine/Plugins/Runtime/GameplayAbilities/Source/GameplayAbilities/Private/GameplayEffectComponents/AbilitiesGameplayEffectComponent.cpp
--- a/Engine/Plugins/Runtime/GameplayAbilities/Source/GameplayAbilities/Private/GameplayEffectComponents/AbilitiesGameplayEffectComponent.cpp
+++ b/Engine/Plugins/Runtime/GameplayAbilities/Source/GameplayAbilities/Private/GameplayEffectComponents/AbilitiesGameplayEffectComponent.cpp
@@ -104,7 +104,8 @@ void UAbilitiesGameplayEffectComponent::RemoveAbilities(FActiveGameplayEffectHan
 		return;
 	}
 
-	const TArray<FGameplayAbilitySpec>& AllAbilities = ASC->GetActivatableAbilities();
+	FScopedAbilityListLock ScopedAbilityListLock(*ASC);
+	const TArray<const FGameplayAbilitySpec*> GrantedAbilities = ASC->FindAbilitySpecsFromGEHandle(ScopedAbilityListLock, ActiveGEHandle, EConsiderPending::All);
 	for (const FGameplayAbilitySpecConfig& AbilityConfig : GrantAbilityConfigs)
 	{
 		// Check that we're configured
@@ -115,22 +116,23 @@ void UAbilitiesGameplayEffectComponent::RemoveAbilities(FActiveGameplayEffectHan
 		}
 
 		// See if we were granted, and if so we can remove it
-		const FGameplayAbilitySpec* AbilitySpecDef = AllAbilities.FindByPredicate([AbilityCDO, &ActiveGEHandle](FGameplayAbilitySpec& Spec) { return Spec.Ability == AbilityCDO && Spec.GameplayEffectHandle == ActiveGEHandle; });
-		if (!AbilitySpecDef)
+		const FGameplayAbilitySpec* const* AbilitySpecItem = GrantedAbilities.FindByPredicate([AbilityCDO](const FGameplayAbilitySpec* Spec) { return Spec->Ability == AbilityCDO; });
+		if (!AbilitySpecItem || !(*AbilitySpecItem))
 		{
 			continue;
 		}
+		const FGameplayAbilitySpec& AbilitySpecDef = (**AbilitySpecItem);
 
 		switch (AbilityConfig.RemovalPolicy)
 		{
 			case EGameplayEffectGrantedAbilityRemovePolicy::CancelAbilityImmediately:
 			{
-				ASC->ClearAbility(AbilitySpecDef->Handle);
+				ASC->ClearAbility(AbilitySpecDef.Handle);
 				break;
 			}
 			case EGameplayEffectGrantedAbilityRemovePolicy::RemoveAbilityOnEnd:
 			{
-				ASC->SetRemoveAbilityOnEnd(AbilitySpecDef->Handle);
+				ASC->SetRemoveAbilityOnEnd(AbilitySpecDef.Handle);
 				break;
 			}
 			default:
diff --git a/Engine/Plugins/Runtime/GameplayAbilities/Source/GameplayAbilities/Public/AbilitySystemComponent.h b/Engine/Plugins/Runtime/GameplayAbilities/Source/GameplayAbilities/Public/AbilitySystemComponent.h
--- a/Engine/Plugins/Runtime/GameplayAbilities/Source/GameplayAbilities/Public/AbilitySystemComponent.h
+++ b/Engine/Plugins/Runtime/GameplayAbilities/Source/GameplayAbilities/Public/AbilitySystemComponent.h
@@ -92,6 +92,22 @@ enum class EGameplayEffectReplicationMode : uint8
 	Full,
 };
 
+/** When performing actions (such as gathering activatable abilities), how do we deal with Pending items (e.g. abilities not yet added or removed) */
+enum class EConsiderPending : uint8
+{
+	/** Don't consider any Pending actions (such as Pending Abilities Added or Removed) */
+	None = 0,
+
+	/** Consider Pending Adds when performing the action */
+	PendingAdd = (1 << 0),
+
+	/** Consider Pending Removes when performing the action */
+	PendingRemove = (1 << 1),
+
+	All = PendingAdd | PendingRemove
+};
+ENUM_CLASS_FLAGS(EConsiderPending)
+
 /** The core ActorComponent for interfacing with the GameplayAbilities System */
 UCLASS(ClassGroup=AbilitySystem, hidecategories=(Object,LOD,Lighting,Transform,Sockets,TextureStreaming), editinlinenew, meta=(BlueprintSpawnableComponent))
 class GAMEPLAYABILITIES_API UAbilitySystemComponent : public UGameplayTasksComponent, public IGameplayTagAssetInterface, public IAbilitySystemReplicationProxyInterface
@@ -1123,13 +1139,21 @@ class GAMEPLAYABILITIES_API UAbilitySystemComponent : public UGameplayTasksCompo
 	/** Returns local world time that an ability was activated. Valid on authority (server) and autonomous proxy (controlling client).  */
 	float GetAbilityLastActivatedTime() const { return AbilityLastActivatedTime; }
 
-	/** Returns an ability spec from a handle. If modifying call MarkAbilitySpecDirty */
-	FGameplayAbilitySpec* FindAbilitySpecFromHandle(FGameplayAbilitySpecHandle Handle) const;
+	/** Returns an ability spec from a handle. If modifying call MarkAbilitySpecDirty. Treat the return value as ephemeral as the pointer will potentially be invalidated on any subsequent call into AbilitySystemComponent. */
+	FGameplayAbilitySpec* FindAbilitySpecFromHandle(FGameplayAbilitySpecHandle Handle, EConsiderPending ConsiderPending = EConsiderPending::PendingRemove) const;
 	
 	/** Returns an ability spec from a GE handle. If modifying call MarkAbilitySpecDirty */
 	UE_DEPRECATED(5.3, "FindAbilitySpecFromGEHandle was never accurate because a GameplayEffect can grant multiple GameplayAbilities. It now returns nullptr.")
 	FGameplayAbilitySpec* FindAbilitySpecFromGEHandle(FActiveGameplayEffectHandle Handle) const;
 
+	/**
+	* Returns all ability spec handles granted from a GE handle. Only the server may call this function.
+	* @param ScopeLock - The lock to communicate to the caller that the return value is only valid for as long as this lock is in scope
+	* @param Handle - The handle of the Active Gameplay Effect which granted the abilities we are looking for
+	* @param ConsiderPending - Are we returning AbilitySpecs that are pending for addition/removal?
+	*/
+	TArray<const FGameplayAbilitySpec*> FindAbilitySpecsFromGEHandle(const FScopedAbilityListLock& ScopeLock, FActiveGameplayEffectHandle Handle, EConsiderPending ConsiderPending = EConsiderPending::PendingRemove) const;
+
 	/** Returns an ability spec corresponding to given ability class. If modifying call MarkAbilitySpecDirty */
 	FGameplayAbilitySpec* FindAbilitySpecFromClass(TSubclassOf<UGameplayAbility> InAbilityClass) const;
 
