Ensure FPredictionKeys only attempt to "catch-up" if they were locally generated.  If they are Server Initiated, then ignore them for catch-up purposes.

This would have had very subtle bugs in the past, where a server-initiated prediction key would "catch-up" client keys, possibly before (or after) rejecting them.  The new FPredictionKey handling makes these bugs more apparent.

diff --git a/Engine/Plugins/Runtime/GameplayAbilities/Source/GameplayAbilities/Private/GameplayPrediction.cpp b/Engine/Plugins/Runtime/GameplayAbilities/Source/GameplayAbilities/Private/GameplayPrediction.cpp
--- a/Engine/Plugins/Runtime/GameplayAbilities/Source/GameplayAbilities/Private/GameplayPrediction.cpp
+++ b/Engine/Plugins/Runtime/GameplayAbilities/Source/GameplayAbilities/Private/GameplayPrediction.cpp
@@ -216,6 +216,16 @@ FPredictionKey FPredictionKey::CreateNewServerInitiatedKey(const UAbilitySystemC
 	{
 		NewKey.GenerateNewPredictionKey();
 		NewKey.bIsServerInitiated = true;
+
+#if !UE_BUILD_SHIPPING
+		// Make sure the Server and Client aren't synchronized in terms of key generation or it can hide bugs.
+		const uint32 ServerDifferential = static_cast<uint32>(reinterpret_cast<intptr_t>(OwningComponent));
+		NewKey.Current = (NewKey.Current + ServerDifferential) % std::numeric_limits<FPredictionKey::KeyType>::max();
+		if (NewKey.Current < 1)
+		{
+			NewKey.Current = 1;
+		}
+#endif
 	}
 	return NewKey;
 }
@@ -498,6 +508,13 @@ void FReplicatedPredictionKeyItem::OnRep(const FReplicatedPredictionKeyMap& InAr
 		UE_LOG(LogPredictionKey, Verbose, TEXT("FReplicatedPredictionKeyItem::OnRep %s"), *PredictionKey.ToString());
 	}
 
+	// We should only run catch-up logic to locally predicted keys.  Otherwise, there will eventually be a case where the server
+	// initiates a key that matches a local value.  Then we catch-up to a local value, even though the server was not specifically acknowledging it.
+	if (PredictionKey.bIsServerInitiated)
+	{
+		return;
+	}
+
 	// Every predictive action we've done in the current chain of dependencies (including the current value) of ReplicatedPredictionKey needs to be acknowledged
 	FPredictionKeyDelegates::CatchUpTo(PredictionKey.Current);
 
diff --git a/Engine/Plugins/Runtime/GameplayAbilities/Source/GameplayAbilities/Public/GameplayPrediction.h b/Engine/Plugins/Runtime/GameplayAbilities/Source/GameplayAbilities/Public/GameplayPrediction.h
--- a/Engine/Plugins/Runtime/GameplayAbilities/Source/GameplayAbilities/Public/GameplayPrediction.h
+++ b/Engine/Plugins/Runtime/GameplayAbilities/Source/GameplayAbilities/Public/GameplayPrediction.h
@@ -376,7 +376,7 @@ struct GAMEPLAYABILITIES_API FPredictionKey
 
 	FString ToString() const
 	{
-		return FString::Printf(TEXT("[%d/%d]"), Current, Base);
+		return bIsServerInitiated ? FString::Printf(TEXT("[Srv: %d]"), Current) : FString::Printf(TEXT("[%d/%d]"), Current, Base);
 	}
 
 	friend uint32 GetTypeHash(const FPredictionKey& InKey)
