Make it clear that we currently do not support replication when using InstancedPerExecution GameplayAbilities.

We will look at adding support in the near future.  Track issue UE-210273.

diff --git a/Engine/Plugins/Runtime/GameplayAbilities/Source/GameplayAbilities/Private/Abilities/GameplayAbility.cpp b/Engine/Plugins/Runtime/GameplayAbilities/Source/GameplayAbilities/Private/Abilities/GameplayAbility.cpp
--- a/Engine/Plugins/Runtime/GameplayAbilities/Source/GameplayAbilities/Private/Abilities/GameplayAbility.cpp
+++ b/Engine/Plugins/Runtime/GameplayAbilities/Source/GameplayAbilities/Private/Abilities/GameplayAbility.cpp
@@ -215,6 +215,14 @@ EDataValidationResult UGameplayAbility::IsDataValid(FDataValidationContext& Cont
 			Result = EDataValidationResult::Invalid;
 		}
 	}
+	else
+	{
+		if (GetInstancingPolicy() != EGameplayAbilityInstancingPolicy::InstancedPerActor)
+		{
+			Context.AddError(FText::Format(LOCTEXT("ReplicatedInstancePolicyNotSupported", "Instancing Policy '{0}' is not supported for Replication.  Either change the Replication Policy or the Instancing Policy"), UEnum::GetDisplayValueAsText(GetInstancingPolicy())));
+			Result = EDataValidationResult::Invalid;
+		}
+	}
 
 	for (TFieldIterator<const UFunction> FuncIter(GetClass(), EFieldIterationFlags::IncludeSuper); FuncIter; ++FuncIter)
 	{
diff --git a/Engine/Plugins/Runtime/GameplayAbilities/Source/GameplayAbilities/Public/Abilities/GameplayAbilityTypes.h b/Engine/Plugins/Runtime/GameplayAbilities/Source/GameplayAbilities/Public/Abilities/GameplayAbilityTypes.h
--- a/Engine/Plugins/Runtime/GameplayAbilities/Source/GameplayAbilities/Public/Abilities/GameplayAbilityTypes.h
+++ b/Engine/Plugins/Runtime/GameplayAbilities/Source/GameplayAbilities/Public/Abilities/GameplayAbilityTypes.h
@@ -48,7 +48,7 @@ namespace EGameplayAbilityInstancingPolicy
 		// Each actor gets their own instance of this ability. State can be saved, replication is possible.
 		InstancedPerActor,
 
-		// We instance this ability each time it is executed. Replication possible but not recommended.
+		// We instance this ability each time it is executed. Replication currently unsupported.
 		InstancedPerExecution,
 	};
 }
