Fix for GC error with UAbilityTask_WaitVelocityChange.

This change just avoids the issue by making the UMovementComponent a TWeakObjectPtr.

diff --git a/Engine/Plugins/Runtime/GameplayAbilities/Source/GameplayAbilities/Private/Abilities/Tasks/AbilityTask_WaitVelocityChange.cpp b/Engine/Plugins/Runtime/GameplayAbilities/Source/GameplayAbilities/Private/Abilities/Tasks/AbilityTask_WaitVelocityChange.cpp
--- a/Engine/Plugins/Runtime/GameplayAbilities/Source/GameplayAbilities/Private/Abilities/Tasks/AbilityTask_WaitVelocityChange.cpp
+++ b/Engine/Plugins/Runtime/GameplayAbilities/Source/GameplayAbilities/Private/Abilities/Tasks/AbilityTask_WaitVelocityChange.cpp
@@ -14,9 +14,9 @@ UAbilityTask_WaitVelocityChange::UAbilityTask_WaitVelocityChange(const FObjectIn
 
 void UAbilityTask_WaitVelocityChange::TickTask(float DeltaTime)
 {
-	if (CachedMovementComponent)
+	if (UMovementComponent* LocalMoveComp = CachedMovementComponent.Get())
 	{
-		float dot = FVector::DotProduct(Direction, CachedMovementComponent->Velocity);
+		float dot = FVector::DotProduct(Direction, LocalMoveComp->Velocity);
 
 		if (dot > MinimumMagnitude)
 		{
@@ -48,7 +48,7 @@ UAbilityTask_WaitVelocityChange* UAbilityTask_WaitVelocityChange::CreateWaitVelo
 void UAbilityTask_WaitVelocityChange::Activate()
 {
 	const FGameplayAbilityActorInfo* ActorInfo = Ability->GetCurrentActorInfo();
-	CachedMovementComponent = ActorInfo->MovementComponent.Get();
+	CachedMovementComponent = ActorInfo->MovementComponent;
 	SetWaitingOnAvatar();
 }
 
diff --git a/Engine/Plugins/Runtime/GameplayAbilities/Source/GameplayAbilities/Public/Abilities/Tasks/AbilityTask_WaitVelocityChange.h b/Engine/Plugins/Runtime/GameplayAbilities/Source/GameplayAbilities/Public/Abilities/Tasks/AbilityTask_WaitVelocityChange.h
--- a/Engine/Plugins/Runtime/GameplayAbilities/Source/GameplayAbilities/Public/Abilities/Tasks/AbilityTask_WaitVelocityChange.h
+++ b/Engine/Plugins/Runtime/GameplayAbilities/Source/GameplayAbilities/Public/Abilities/Tasks/AbilityTask_WaitVelocityChange.h
@@ -30,7 +30,7 @@ class GAMEPLAYABILITIES_API UAbilityTask_WaitVelocityChange: public UAbilityTask
 protected:
 
 	UPROPERTY()
-	TObjectPtr<UMovementComponent>	CachedMovementComponent;
+	TWeakObjectPtr<UMovementComponent>	CachedMovementComponent;
 
 	float	MinimumMagnitude;
 	FVector Direction;
