Reverting some "AbilitySystem.PredictionKey" Console Variable defaults to use legacy behavior.  Some bugs were discovered in the implementation right before UE5.4 is shipping.  By setting these to legacy, we can reduce the visibility of those bugs, but a patch will be needed to eliminate them altogether.

diff --git a/Engine/Plugins/Runtime/GameplayAbilities/Source/GameplayAbilities/Private/GameplayPrediction.cpp b/Engine/Plugins/Runtime/GameplayAbilities/Source/GameplayAbilities/Private/GameplayPrediction.cpp
--- a/Engine/Plugins/Runtime/GameplayAbilities/Source/GameplayAbilities/Private/GameplayPrediction.cpp
+++ b/Engine/Plugins/Runtime/GameplayAbilities/Source/GameplayAbilities/Private/GameplayPrediction.cpp
@@ -21,16 +21,16 @@ namespace UE::AbilitySystem::Private
 	FAutoConsoleVariableRef CVarMaxStaleKeysBeforeAck(TEXT("AbilitySystem.PredictionKey.MaxStaleKeysBeforeAck"), CVarMaxStaleKeysBeforeAckValue,
 		TEXT("How many prediction keys can be dropped before StaleKeyBehavior is run."));
 
-	// What should we do with these old stale FPredictionKeys?  Prior to UE5.4, we always CaughtUp.  I believe it's actually safer to drop.
-	int32 CVarStaleKeyBehaviorValue = 2;
+	// What should we do with these old stale FPredictionKeys?  Prior to UE5.5, we always CaughtUp.  I believe it's actually safer to drop.
+	int32 CVarStaleKeyBehaviorValue = 0;
 	FAutoConsoleVariableRef CVarStaleKeyBehavior(TEXT("AbilitySystem.PredictionKey.StaleKeyBehavior"), CVarStaleKeyBehaviorValue,
 		TEXT("How do we handle stale keys? 0 = CaughtUp. 1 = Reject. 2 = Drop"));
 
-	// How should we deal with dependent keys (in a chain)?  Prior to UE5.4, old keys implied new keys.  We introduced some new functionality (explained in the help text).
-	// 0 (no bitmask) is legacy behavior.  Logically, (1 | (1<<1)) = 3 is the correct value. We default to 1 to move towards a long-term fix (of value 3).
-	int32 CVarDependentChainBehaviorValue = 1;
+	// How should we deal with dependent keys (in a chain)?  Prior to UE5.5, old keys implied new keys.  We introduced some new functionality (explained in the help text).
+	// 0 (no bitmask) is legacy behavior.  Logically, (0x1 | 0x2) = 3 is the correct value. The long-term fix will be a value of 3.
+	int32 CVarDependentChainBehaviorValue = 0;
 	FAutoConsoleVariableRef CVarDependentChainBehavior(TEXT("AbilitySystem.PredictionKey.DepChainBehavior"), CVarDependentChainBehaviorValue,
-		TEXT("How do we handle dependency key chains? Bitmask: 0 = Old Accept/Rejected Implies Newer Accepted/Rejected. 1 = Newer Accepted also implies Older Accepted. 2 = Old Accepted Does NOT imply Newer Accepted"));
+		TEXT("How do we handle dependency key chains? Bitmask: 0 = Old Accept/Rejected Implies Newer Accepted/Rejected. 0x1 = Newer Accepted also implies Older Accepted. 0x2 = Old Accepted Does NOT imply Newer Accepted"));
 
 	/**
 	 * Given an FProperty Link (such as a UFunction's Properties, or a UStruct's Properties), find and return any linked FPredictionKeys.
