Fixes for issues where Gameplay Effects & PredictionKeys.  Solidifying rules around the client/server communication.  It includes:

1.  GE's could be without communicating a prediction key if it went through SendGameplayEventToActor.
2.  GE's could be removed without being the authority through certain paths.
3.  Server Initiated PredictionKeys were being sent down to the client, but it takes a spot in a circular buffer and thus it's possible it stomps an actual client FPredictionKey.
4.  Extra logging and ensures to clarify these rules.
5.  Changes to the Abilities Gameplay Debugger to catch these issues (previously multiple applications of a GE could go undetected).
6.  Blueprint function MakeSpecHandle has been deprecated in favor of MakeSpecHandleByClass since you need the CDO instance.

diff --git a/Engine/Plugins/Runtime/GameplayAbilities/Source/GameplayAbilities/Private/Abilities/GameplayAbility.cpp b/Engine/Plugins/Runtime/GameplayAbilities/Source/GameplayAbilities/Private/Abilities/GameplayAbility.cpp
--- a/Engine/Plugins/Runtime/GameplayAbilities/Source/GameplayAbilities/Private/Abilities/GameplayAbility.cpp
+++ b/Engine/Plugins/Runtime/GameplayAbilities/Source/GameplayAbilities/Private/Abilities/GameplayAbility.cpp
@@ -9,6 +9,7 @@
 #include "AbilitySystemStats.h"
 #include "AbilitySystemGlobals.h"
 #include "AbilitySystemComponent.h"
+#include "AbilitySystemPrivate.h"
 #include "Abilities/Tasks/AbilityTask.h"
 #include "Components/SkeletalMeshComponent.h"
 #include "GameplayCue_Types.h"
@@ -143,8 +144,16 @@ void UGameplayAbility::SendGameplayEvent(FGameplayTag EventTag, FGameplayEventDa
 	UAbilitySystemComponent* const AbilitySystemComponent = GetAbilitySystemComponentFromActorInfo_Ensured();
 	if (AbilitySystemComponent)
 	{
-		FScopedPredictionWindow NewScopedWindow(AbilitySystemComponent, true);
-		AbilitySystemComponent->HandleGameplayEvent(EventTag, &Payload);
+		using namespace UE::AbilitySystem::Private;
+		if (EnumHasAnyFlags(static_cast<EAllowPredictiveGEFlags>(CVarAllowPredictiveGEFlagsValue), EAllowPredictiveGEFlags::AllowGameplayEventToApplyGE))
+		{
+			FScopedPredictionWindow NewScopedWindow(AbilitySystemComponent, true);
+			AbilitySystemComponent->HandleGameplayEvent(EventTag, &Payload);
+		}
+		else
+		{
+			AbilitySystemComponent->HandleGameplayEvent(EventTag, &Payload);
+		}
 	}
 }
 
diff --git a/Engine/Plugins/Runtime/GameplayAbilities/Source/GameplayAbilities/Private/AbilitySystemBlueprintLibrary.cpp b/Engine/Plugins/Runtime/GameplayAbilities/Source/GameplayAbilities/Private/AbilitySystemBlueprintLibrary.cpp
--- a/Engine/Plugins/Runtime/GameplayAbilities/Source/GameplayAbilities/Private/AbilitySystemBlueprintLibrary.cpp
+++ b/Engine/Plugins/Runtime/GameplayAbilities/Source/GameplayAbilities/Private/AbilitySystemBlueprintLibrary.cpp
@@ -1,12 +1,13 @@
 // Copyright Epic Games, Inc. All Rights Reserved.
 
 #include "AbilitySystemBlueprintLibrary.h"
-#include "GameplayEffectAggregator.h"
-#include "AbilitySystemGlobals.h"
 #include "AbilitySystemComponent.h"
+#include "AbilitySystemGlobals.h"
 #include "AbilitySystemLog.h"
+#include "AbilitySystemPrivate.h"
 #include "Engine/World.h"
 #include "GameplayEffect.h"
+#include "GameplayEffectAggregator.h"
 #include "GameplayEffectComponents/AdditionalEffectsGameplayEffectComponent.h"
 #include "GameplayEffectUIData.h"
 #include "GameplayAbilitySpec.h"
@@ -30,8 +31,16 @@ void UAbilitySystemBlueprintLibrary::SendGameplayEventToActor(AActor* Actor, FGa
 		UAbilitySystemComponent* AbilitySystemComponent = GetAbilitySystemComponent(Actor);
 		if (AbilitySystemComponent != nullptr && IsValidChecked(AbilitySystemComponent))
 		{
-			FScopedPredictionWindow NewScopedWindow(AbilitySystemComponent, true);
-			AbilitySystemComponent->HandleGameplayEvent(EventTag, &Payload);
+			using namespace UE::AbilitySystem::Private;
+			if (EnumHasAnyFlags(static_cast<EAllowPredictiveGEFlags>(CVarAllowPredictiveGEFlagsValue), EAllowPredictiveGEFlags::AllowGameplayEventToApplyGE))
+			{
+				FScopedPredictionWindow NewScopedWindow(AbilitySystemComponent, true);
+				AbilitySystemComponent->HandleGameplayEvent(EventTag, &Payload);
+			}
+			else
+			{
+				AbilitySystemComponent->HandleGameplayEvent(EventTag, &Payload);
+			}
 		}
 		else
 		{
@@ -268,6 +277,20 @@ FGameplayEffectSpecHandle UAbilitySystemBlueprintLibrary::MakeSpecHandle(UGamepl
 	return FGameplayEffectSpecHandle();
 }
 
+FGameplayEffectSpecHandle UAbilitySystemBlueprintLibrary::MakeSpecHandleByClass(TSubclassOf<UGameplayEffect> GameplayEffect, AActor* Instigator, AActor* EffectCauser, float Level)
+{
+	if (const UGameplayEffect* GameplayEffectCDO = GameplayEffect.GetDefaultObject())
+	{
+		FGameplayEffectContext* EffectContext = UAbilitySystemGlobals::Get().AllocGameplayEffectContext();
+		EffectContext->AddInstigator(Instigator, EffectCauser);
+
+		return FGameplayEffectSpecHandle(new FGameplayEffectSpec(GameplayEffectCDO, FGameplayEffectContextHandle(EffectContext), Level));
+	}
+
+	ABILITY_LOG(Warning, TEXT("%hs was called with invalid GameplayEffect"), __func__);
+	return FGameplayEffectSpecHandle();
+}
+
 FGameplayEffectSpecHandle UAbilitySystemBlueprintLibrary::CloneSpecHandle(AActor* InNewInstigator, AActor* InEffectCauser, FGameplayEffectSpecHandle GameplayEffectSpecHandle_Clone)
 {
 	FGameplayEffectContext* EffectContext = UAbilitySystemGlobals::Get().AllocGameplayEffectContext();
diff --git a/Engine/Plugins/Runtime/GameplayAbilities/Source/GameplayAbilities/Private/AbilitySystemComponent.cpp b/Engine/Plugins/Runtime/GameplayAbilities/Source/GameplayAbilities/Private/AbilitySystemComponent.cpp
--- a/Engine/Plugins/Runtime/GameplayAbilities/Source/GameplayAbilities/Private/AbilitySystemComponent.cpp
+++ b/Engine/Plugins/Runtime/GameplayAbilities/Source/GameplayAbilities/Private/AbilitySystemComponent.cpp
@@ -11,6 +11,7 @@
 #include "AbilitySystemStats.h"
 #include "AbilitySystemGlobals.h"
 #include "GameplayCueManager.h"
+#include "AbilitySystemPrivate.h"
 
 #include "Net/UnrealNetwork.h"
 #include "Engine/ActorChannel.h"
@@ -1039,9 +1040,30 @@ const UGameplayEffect* UAbilitySystemComponent::GetGameplayEffectDefForHandle(FA
 
 bool UAbilitySystemComponent::RemoveActiveGameplayEffect(FActiveGameplayEffectHandle Handle, int32 StacksToRemove)
 {
+	using namespace UE::AbilitySystem::Private;
+	if (!EnumHasAnyFlags(static_cast<EAllowPredictiveGEFlags>(CVarAllowPredictiveGEFlagsValue), EAllowPredictiveGEFlags::AllowNativeRemoveByHandle))
+	{
+		if (!IsOwnerActorAuthoritative())
+		{
+			UE_LOG(LogAbilitySystem, Warning, TEXT("%hs called without Authority. Fix-up code, or temporarily patch using AbilitySystem.Fix.AllowPredictiveGEFlags"), __func__);
+			return false;
+		}
+	}
+
 	return ActiveGameplayEffects.RemoveActiveGameplayEffect(Handle, StacksToRemove);
 }
 
+void UAbilitySystemComponent::RemoveActiveGameplayEffect_NoReturn(FActiveGameplayEffectHandle Handle, int32 StacksToRemove)
+{
+	// Legacy version allowed client removal, so let's keep that for now.
+	RemoveActiveGameplayEffect_AllowClientRemoval(Handle, StacksToRemove);
+}
+
+void UAbilitySystemComponent::RemoveActiveGameplayEffect_AllowClientRemoval(FActiveGameplayEffectHandle Handle, int32 StacksToRemove)
+{
+	ActiveGameplayEffects.RemoveActiveGameplayEffect(Handle, StacksToRemove);
+}
+
 void UAbilitySystemComponent::RemoveActiveGameplayEffectBySourceEffect(TSubclassOf<UGameplayEffect> GameplayEffect, UAbilitySystemComponent* InstigatorAbilitySystemComponent, int32 StacksToRemove /*= -1*/)
 {
 	if (GameplayEffect)
diff --git a/Engine/Plugins/Runtime/GameplayAbilities/Source/GameplayAbilities/Private/AbilitySystemComponent_Abilities.cpp b/Engine/Plugins/Runtime/GameplayAbilities/Source/GameplayAbilities/Private/AbilitySystemComponent_Abilities.cpp
--- a/Engine/Plugins/Runtime/GameplayAbilities/Source/GameplayAbilities/Private/AbilitySystemComponent_Abilities.cpp
+++ b/Engine/Plugins/Runtime/GameplayAbilities/Source/GameplayAbilities/Private/AbilitySystemComponent_Abilities.cpp
@@ -1971,7 +1971,7 @@ void UAbilitySystemComponent::InternalServerTryActivateAbility(FGameplayAbilityS
 	if (!Spec)
 	{
 		// Can potentially happen in race conditions where client tries to activate ability that is removed server side before it is received.
-		ABILITY_LOG(Display, TEXT("InternalServerTryActivateAbility. Rejecting ClientActivation of ability with invalid SpecHandle!"));
+		UE_VLOG_UELOG(GetOwner(), LogAbilitySystem, Display, TEXT("%hs: Rejecting ClientActivation of ability with invalid SpecHandle!"), __func__);
 		ClientActivateAbilityFailed(Handle, PredictionKey.Current);
 		return;
 	}
@@ -1980,7 +1980,7 @@ void UAbilitySystemComponent::InternalServerTryActivateAbility(FGameplayAbilityS
 
 	if (!ensure(AbilityToActivate))
 	{
-		ABILITY_LOG(Error, TEXT("InternalServerTryActiveAbility. Rejecting ClientActivation of unconfigured spec ability!"));
+		UE_VLOG_UELOG(GetOwner(), LogAbilitySystem, Error, TEXT("%hs: Rejecting ClientActivation of unconfigured spec ability!"), __func__);
 		ClientActivateAbilityFailed(Handle, PredictionKey.Current);
 		return;
 	}
@@ -1989,7 +1989,7 @@ void UAbilitySystemComponent::InternalServerTryActivateAbility(FGameplayAbilityS
 	if (AbilityToActivate->GetNetSecurityPolicy() == EGameplayAbilityNetSecurityPolicy::ServerOnlyExecution ||
 		AbilityToActivate->GetNetSecurityPolicy() == EGameplayAbilityNetSecurityPolicy::ServerOnly)
 	{
-		ABILITY_LOG(Display, TEXT("InternalServerTryActiveAbility. Rejecting ClientActivation of %s due to security policy violation."), *GetNameSafe(AbilityToActivate));
+		UE_VLOG_UELOG(GetOwner(), LogAbilitySystem, Display, TEXT("%hs: Rejecting ClientActivation of %s due to security policy violation."), __func__, *GetNameSafe(AbilityToActivate));
 		ClientActivateAbilityFailed(Handle, PredictionKey.Current);
 		return;
 	}
@@ -2010,11 +2010,14 @@ void UAbilitySystemComponent::InternalServerTryActivateAbility(FGameplayAbilityS
 	// Attempt to activate the ability (server side) and tell the client if it succeeded or failed.
 	if (InternalTryActivateAbility(Handle, PredictionKey, &InstancedAbility, nullptr, TriggerEventData))
 	{
-		// TryActivateAbility handles notifying the client of success
+		// TryActivateAbility handles notifying the client of success, but let's still log it
+		UE_VLOG_UELOG(GetOwner(), LogAbilitySystem, Verbose, TEXT("%hs: Accepted ClientActivation of %s with PredictionKey %s."), __func__, *GetNameSafe(Spec->Ability), *PredictionKey.ToString());
 	}
 	else
 	{
-		ABILITY_LOG(Display, TEXT("InternalServerTryActivateAbility. Rejecting ClientActivation of %s. InternalTryActivateAbility failed: %s"), *GetNameSafe(Spec->Ability), *InternalTryActivateAbilityFailureTags.ToStringSimple() );
+		UE_VLOG_UELOG(GetOwner(), LogAbilitySystem, Display, TEXT("%hs: Rejecting ClientActivation of %s with PredictionKey %s. InternalTryActivateAbility failed: %s"),
+			__func__, *GetNameSafe(Spec->Ability), *PredictionKey.ToString(), *InternalTryActivateAbilityFailureTags.ToStringSimple() );
+
 		ClientActivateAbilityFailed(Handle, PredictionKey.Current);
 		Spec->InputPressed = false;
 
diff --git a/Engine/Plugins/Runtime/GameplayAbilities/Source/GameplayAbilities/Private/AbilitySystemPrivate.h b/Engine/Plugins/Runtime/GameplayAbilities/Source/GameplayAbilities/Private/AbilitySystemPrivate.h
new file mode 100644
--- /dev/null
+++ b/Engine/Plugins/Runtime/GameplayAbilities/Source/GameplayAbilities/Private/AbilitySystemPrivate.h
@@ -0,0 +1,24 @@
+// Copyright Epic Games, Inc. All Rights Reserved.
+
+#pragma once
+
+#include "Misc/EnumClassFlags.h"
+#include "HAL/IConsoleManager.h"
+
+namespace UE::AbilitySystem::Private
+{
+	/** These are flags to rollback fixes in live production */
+	enum class EAllowPredictiveGEFlags
+	{
+		None = 0,								// No intentional loopholes, try our best to enforce this rule (preferred, default).
+		AllowNativeRemoveByHandle = 1 << 0,		// Some games may be relying on RemoveActiveGameplayEffect having only the UFUNCTION BlueprintAuthority check; it was previously not checked in native.
+		AllowRemovalByTagRequirements = 1 << 1,	// A bug introduced in UE5.3 allowed target tag requirements to remove GE's client-side due to the above.
+		AllowGameplayEventToApplyGE = 1 << 2	// Legacy bug allowed Gameplay Events to create a prediction window with which one could apply GE's locally. That causes issues since it's not guaranteed to be communicated to the server.
+	};
+	ENUM_CLASS_FLAGS(EAllowPredictiveGEFlags);
+
+	// CVar that controls which fixes to the predictive GE code we want to disable (for legacy reasons)
+	inline int32 CVarAllowPredictiveGEFlagsValue = 0;
+	inline FAutoConsoleVariableRef CVarAllowPredictiveGEFlags(TEXT("AbilitySystem.Fix.AllowPredictiveGEFlags"), CVarAllowPredictiveGEFlagsValue, TEXT("Default: 0 (all fixes, no exceptions).\n \
+		Use flag 0x1 to allow when removing by Handle in native (legacy).  Use flag 0x2 to allow RequirementsTags to remove (bug introduced in UE5.3).  Use flag 0x4 to allow Gameplay Events to predictively execute GE's (legacy bug)."));
+}
diff --git a/Engine/Plugins/Runtime/GameplayAbilities/Source/GameplayAbilities/Private/GameplayAbilityTypes.cpp b/Engine/Plugins/Runtime/GameplayAbilities/Source/GameplayAbilities/Private/GameplayAbilityTypes.cpp
--- a/Engine/Plugins/Runtime/GameplayAbilities/Source/GameplayAbilities/Private/GameplayAbilityTypes.cpp
+++ b/Engine/Plugins/Runtime/GameplayAbilities/Source/GameplayAbilities/Private/GameplayAbilityTypes.cpp
@@ -240,6 +240,15 @@ void FGameplayAbilitySpec::PreReplicatedRemove(const struct FGameplayAbilitySpec
 	}
 }
 
+void FGameplayAbilitySpec::PostReplicatedChange(const struct FGameplayAbilitySpecContainer& InArraySerializer)
+{
+	if (InArraySerializer.Owner)
+	{
+		UE_LOG(LogAbilitySystem, Verbose, TEXT("%s: AbilitySpecChanged (Non-Auth): [%s] %s. Level: %d"), *GetNameSafe(InArraySerializer.Owner->GetOwner()), *Handle.ToString(), *GetNameSafe(Ability), Level)
+		UE_VLOG(InArraySerializer.Owner->GetOwner(), VLogAbilitySystem, Verbose, TEXT("AbilitySpecChanged (Non-Auth): [%s] %s. Level: %d"), *Handle.ToString(), *GetNameSafe(Ability), Level);
+	}
+}
+
 void FGameplayAbilitySpec::PostReplicatedAdd(const struct FGameplayAbilitySpecContainer& InArraySerializer)
 {
 	if (InArraySerializer.Owner)
diff --git a/Engine/Plugins/Runtime/GameplayAbilities/Source/GameplayAbilities/Private/GameplayDebuggerCategory_Abilities.cpp b/Engine/Plugins/Runtime/GameplayAbilities/Source/GameplayAbilities/Private/GameplayDebuggerCategory_Abilities.cpp
--- a/Engine/Plugins/Runtime/GameplayAbilities/Source/GameplayAbilities/Private/GameplayDebuggerCategory_Abilities.cpp
+++ b/Engine/Plugins/Runtime/GameplayAbilities/Source/GameplayAbilities/Private/GameplayDebuggerCategory_Abilities.cpp
@@ -113,12 +113,7 @@ void FGameplayDebuggerCategory_Abilities::FRepData::Serialize(FArchive& Ar)
 
 	for (int32 Idx = 0; Idx < NumGE; Idx++)
 	{
-		if (Ar.IsLoading())
-		{
-			GameplayEffects[Idx].PredictionKey = FPredictionKey();
-		}
-		
-		GameplayEffects[Idx].PredictionKey.NetSerialize(Ar, ClientPackageMap.Get(), bSuccess);
+		Ar << GameplayEffects[Idx].ReplicationID;
 		Ar << GameplayEffects[Idx].Effect;
 		Ar << GameplayEffects[Idx].Context;
 		Ar << GameplayEffects[Idx].Duration;
@@ -196,8 +191,8 @@ TArray<FGameplayDebuggerCategory_Abilities::FRepData::FGameplayEffectDebug> FGam
 	for (const FActiveGameplayEffect& ActiveGE : &AbilityComp->GetActiveGameplayEffects())
 	{
 		FRepData::FGameplayEffectDebug& ItemData = DebugEffects.AddDefaulted_GetRef();
+		ItemData.ReplicationID = ActiveGE.ReplicationID;
 		ItemData.bInhibited = ActiveGE.bIsInhibited;
-		ItemData.PredictionKey = ActiveGE.PredictionKey;
 		ItemData.Duration = ActiveGE.GetDuration();
 		ItemData.Period = ActiveGE.GetPeriod();
 
@@ -529,24 +524,28 @@ void FGameplayDebuggerCategory_Abilities::DrawGameplayEffects(FGameplayDebuggerC
 	const TArray<FGameplayEffectDebug>& ServerEffects = DataPack.GameplayEffects;
 	TArray<FGameplayEffectDebug> LocalEffects = bConsiderLocalStatus ? CollectEffectsData(OwnerPC, LocalASC) : ServerEffects;
 
-	auto GetKeyHash = [](const FGameplayEffectDebug& GameplayEffectDebug) -> uint32
+	int NumEffectCounts[+ENetworkStatus::MAX] = { 0 };
+
+	auto GetKeyHash = [](const FGameplayEffectDebug& GameplayEffectDebug, TArray<FGameplayEffectDebug> SourceArray) -> uint32
 		{
-			return HashCombineFast(
-				GetTypeHash(GameplayEffectDebug.PredictionKey),
-				GetTypeHash(GameplayEffectDebug.Effect));
-		};
+			// If we're networked, this will be the same between the client and server (assuming it's been replicated)
+			if (GameplayEffectDebug.ReplicationID != INDEX_NONE)
+				return GameplayEffectDebug.ReplicationID;
 
-	int NumEffectCounts[+ENetworkStatus::MAX] = { 0 };
+			// If we're not networked/not yet replicated, rely on the index of the item in the array & name for good measure
+			ptrdiff_t ArrayIdx = (&GameplayEffectDebug - SourceArray.GetData());
+			return HashCombineFast(GetTypeHash(GameplayEffectDebug.Effect), ArrayIdx);
+		};
 
 	TMap<uint32, FGameplayEffectDebug> Effects;
 	for (const FGameplayEffectDebug& ServerEffect : ServerEffects)
 	{
-		Effects.Add(GetKeyHash(ServerEffect)) = ServerEffect;
+		Effects.Add(GetKeyHash(ServerEffect, ServerEffects)) = ServerEffect;
 	}
 
 	for (const FGameplayEffectDebug& LocalEffect : LocalEffects)
 	{
-		FGameplayEffectDebug& Effect = Effects.FindOrAdd(GetKeyHash(LocalEffect));
+		FGameplayEffectDebug& Effect = Effects.FindOrAdd(GetKeyHash(LocalEffect, LocalEffects));
 		if (!Effect.Effect.IsEmpty())
 		{
 			Effect.NetworkStatus = ENetworkStatus::Networked;
diff --git a/Engine/Plugins/Runtime/GameplayAbilities/Source/GameplayAbilities/Private/GameplayDebuggerCategory_Abilities.h b/Engine/Plugins/Runtime/GameplayAbilities/Source/GameplayAbilities/Private/GameplayDebuggerCategory_Abilities.h
--- a/Engine/Plugins/Runtime/GameplayAbilities/Source/GameplayAbilities/Private/GameplayDebuggerCategory_Abilities.h
+++ b/Engine/Plugins/Runtime/GameplayAbilities/Source/GameplayAbilities/Private/GameplayDebuggerCategory_Abilities.h
@@ -66,7 +66,7 @@ protected:
 
 		struct FGameplayEffectDebug
 		{
-			FPredictionKey PredictionKey;
+			int32 ReplicationID = INDEX_NONE; // unique & shared between server/client (or INDEX_NONE if local)
 			FString Effect;
 			FString Context;
 			float Duration = 0.0f;
diff --git a/Engine/Plugins/Runtime/GameplayAbilities/Source/GameplayAbilities/Private/GameplayEffect.cpp b/Engine/Plugins/Runtime/GameplayAbilities/Source/GameplayAbilities/Private/GameplayEffect.cpp
--- a/Engine/Plugins/Runtime/GameplayAbilities/Source/GameplayAbilities/Private/GameplayEffect.cpp
+++ b/Engine/Plugins/Runtime/GameplayAbilities/Source/GameplayAbilities/Private/GameplayEffect.cpp
@@ -2621,6 +2621,11 @@ bool FActiveGameplayEffect::CheckRemovalTagRequirements(const FGameplayTagContai
 
 void FActiveGameplayEffect::PreReplicatedRemove(const struct FActiveGameplayEffectsContainer &InArray)
 {
+	ensureMsgf(InArray.ScopedLockCount > 0, TEXT("%hs needs to be ScopeLocked.  Previously from UAbilitySystemComponent::PreNetReceive."), __func__);
+
+	const int32 PreRemoveIdx = this - InArray.GameplayEffects_Internal.GetData();
+	ensureMsgf(InArray.GameplayEffects_Internal.IsValidIndex(PreRemoveIdx), TEXT("%hs had invalid 'this' pointer which was not part of the passed-in FActiveGameplayEffectsContainer"), __func__);
+
 	if (Spec.Def == nullptr)
 	{
 		UE_LOG(LogGameplayEffects, Error, TEXT("Received PreReplicatedRemove with no UGameplayEffect def."));
@@ -2642,13 +2647,21 @@ void FActiveGameplayEffect::PreReplicatedRemove(const struct FActiveGameplayEffe
 	}
 	GameplayEffectRemovalInfo.EffectContext = Spec.GetEffectContext();
 
-	UE_VLOG_UELOG(InArray.Owner->GetOwnerActor(), LogGameplayEffects, Verbose, TEXT("%s (Non-Auth): %s. Premature: %d Inhibited: %d. Pending( Remove: %d OnActive: %d WhileActive: %d )"), ANSI_TO_TCHAR(__func__), *GetDebugString(), GameplayEffectRemovalInfo.bPrematureRemoval, bIsInhibited, IsPendingRemove, bPendingRepOnActiveGC, bPendingRepWhileActiveGC);
+	UE_VLOG_UELOG(InArray.Owner->GetOwnerActor(), LogGameplayEffects, Verbose, TEXT("%hs (Non-Auth): %s. Premature: %d Inhibited: %d. Pending( Remove: %d OnActive: %d WhileActive: %d )"), __func__, *GetDebugString(), GameplayEffectRemovalInfo.bPrematureRemoval, bIsInhibited, IsPendingRemove, bPendingRepOnActiveGC, bPendingRepWhileActiveGC);
 
 	const_cast<FActiveGameplayEffectsContainer&>(InArray).InternalOnActiveGameplayEffectRemoved(*this, !bIsInhibited, GameplayEffectRemovalInfo);	// Const cast is ok. It is there to prevent mutation of the GameplayEffects array, which this wont do.
+
+	// Make sure no calls caused a realloc of InArray while holding this pointer
+	ensureMsgf(InArray.GameplayEffects_Internal.IsValidIndex(PreRemoveIdx) && &InArray.GameplayEffects_Internal[PreRemoveIdx] == this, TEXT("%hs had realloc of the array holding the item pointed to by 'this'. This will eventually crash"), __func__);
 }
 
 void FActiveGameplayEffect::PostReplicatedAdd(const struct FActiveGameplayEffectsContainer &InArray)
 {
+	ensureMsgf(InArray.ScopedLockCount > 0, TEXT("%hs needs to be ScopeLocked.  Previously from UAbilitySystemComponent::PreNetReceive."), __func__);
+
+	const int32 PreAddIdx = this - InArray.GameplayEffects_Internal.GetData();
+	ensureMsgf(InArray.GameplayEffects_Internal.IsValidIndex(PreAddIdx), TEXT("%hs had invalid 'this' pointer which was not part of the passed-in FActiveGameplayEffectsContainer"), __func__);
+
 	if (Spec.Def == nullptr)
 	{
 		UE_LOG(LogGameplayEffects, Error, TEXT("FActiveGameplayEffect::PostReplicatedAdd Received ReplicatedGameplayEffect with no UGameplayEffect def. (%s)"), *Spec.GetEffectContext().ToString());
@@ -2716,11 +2729,18 @@ void FActiveGameplayEffect::PostReplicatedAdd(const struct FActiveGameplayEffect
 	// Do stuff for adding GEs (add mods, tags, *invoke callbacks*).  But do NOT invoke the GameplayCues as we don't know if this GE ends up inhibited or not (thus the bPendingRepOnActiveGC variables).
 	constexpr bool bInvokeGameplayCueEvents = false;
 	const_cast<FActiveGameplayEffectsContainer&>(InArray).InternalOnActiveGameplayEffectAdded(*this, bInvokeGameplayCueEvents);	// Const cast is ok. It is there to prevent mutation of the GameplayEffects array, which this wont do.
-	
+
+	// Make sure no calls caused a realloc of InArray while holding this pointer
+	ensureMsgf(InArray.GameplayEffects_Internal.IsValidIndex(PreAddIdx) && &InArray.GameplayEffects_Internal[PreAddIdx] == this, TEXT("%hs had realloc of the array holding the item pointed to by 'this'. This will eventually crash"), __func__);
 }
 
 void FActiveGameplayEffect::PostReplicatedChange(const struct FActiveGameplayEffectsContainer &InArray)
 {
+	ensureMsgf(InArray.ScopedLockCount > 0, TEXT("%hs needs to be ScopeLocked.  Previously from UAbilitySystemComponent::PreNetReceive."), __func__);
+
+	const int32 PreChangeIdx = this - InArray.GameplayEffects_Internal.GetData();
+	ensureMsgf(InArray.GameplayEffects_Internal.IsValidIndex(PreChangeIdx), TEXT("%hs had invalid 'this' pointer which was not part of the passed-in FActiveGameplayEffectsContainer"), __func__);
+
 	if (Spec.Def == nullptr)
 	{
 		UE_LOG(LogGameplayEffects, Error, TEXT("FActiveGameplayEffect::PostReplicatedChange Received ReplicatedGameplayEffect with no UGameplayEffect def. (%s)"), *Spec.GetEffectContext().ToString());
@@ -2734,12 +2754,15 @@ void FActiveGameplayEffect::PostReplicatedChange(const struct FActiveGameplayEff
 		return;
 	}
 
+	UE_VLOG_UELOG(InArray.Owner->GetOwnerActor(), LogGameplayEffects, Verbose, TEXT("%hs (Non-Auth): %s. Pending( OnActive: %d. WhileActive: %d )"), __func__, *GetDebugString(), bPendingRepOnActiveGC, bPendingRepWhileActiveGC);
+
 	// Handle potential duration refresh
 	if (CachedStartServerWorldTime != StartServerWorldTime)
 	{
 		RecomputeStartWorldTime(InArray);
 		CachedStartServerWorldTime = StartServerWorldTime;
 
+		// Const cast is ok. It is there to prevent mutation of the GameplayEffects array, which we believe this wont do due to scope lock.
 		const_cast<FActiveGameplayEffectsContainer&>(InArray).OnDurationChange(*this);
 	}
 	
@@ -2747,18 +2770,19 @@ void FActiveGameplayEffect::PostReplicatedChange(const struct FActiveGameplayEff
 	if (ClientCachedStackCount != StackCount)
 	{
 		// If its a stack count change, we just call OnStackCountChange and it will broadcast delegates and update attribute aggregators
-		// Const cast is ok. It is there to prevent mutation of the GameplayEffects array, which this wont do.
+		// Const cast is ok. It is there to prevent mutation of the GameplayEffects array, which we believe this wont do due to scope lock.
 		const_cast<FActiveGameplayEffectsContainer&>(InArray).OnStackCountChange(*this, ClientCachedStackCount, StackCount);
 		ClientCachedStackCount = StackCount;
 	}
 	else
 	{
 		// Stack count didn't change, but something did (like a modifier magnitude). We need to update our attribute aggregators
-		// Const cast is ok. It is there to prevent mutation of the GameplayEffects array, which this wont do.
+		// Const cast is ok. It is there to prevent mutation of the GameplayEffects array, which we believe this wont do due to scope lock.
 		const_cast<FActiveGameplayEffectsContainer&>(InArray).UpdateAllAggregatorModMagnitudes(*this);
 	}
 
-	UE_VLOG_UELOG(InArray.Owner->GetOwnerActor(), LogGameplayEffects, Verbose, TEXT("%s (Non-Auth): %s. Pending( OnActive: %d. WhileActive: %d )"), ANSI_TO_TCHAR(__func__), *GetDebugString(), bPendingRepOnActiveGC, bPendingRepWhileActiveGC);
+	// Make sure no calls caused a realloc of InArray while holding this pointer
+	ensureMsgf(InArray.GameplayEffects_Internal.IsValidIndex(PreChangeIdx) && &InArray.GameplayEffects_Internal[PreChangeIdx] == this, TEXT("%hs had realloc of the array holding the item pointed to by 'this'. This will eventually crash"), __func__);
 }
 
 FString FActiveGameplayEffect::GetDebugString()
@@ -4002,11 +4026,11 @@ PRAGMA_ENABLE_DEPRECATION_WARNINGS
 	{
 		FActiveGameplayEffectHandle NewHandle = FActiveGameplayEffectHandle::GenerateNewHandle(Owner);
 
-		if (ScopedLockCount > 0 && GameplayEffects_Internal.GetSlack() <= 0)
+		if (GameplayEffects_Internal.GetSlack() <= 0)
 		{
 			/**
-			 *	If we have no more slack and we are scope locked, we need to put this addition on our pending GE list, which will be moved
-			 *	onto the real active GE list once the scope lock is over.
+			 *	Since we are scope locked, if we have no more slack we need to put this addition on our pending GE list, which will be moved
+			 *	onto the real active GE list once the scope lock is over.  Otherwise we risk moving the pointer of the GE that scope-locked us.
 			 *	
 			 *	To avoid extra heap allocations, each active gameplayeffect container keeps a linked list of pending GEs. This list is allocated
 			 *	on demand and re-used in subsequent pending adds. The code below will either 1) Alloc a new pending GE 2) reuse an existing pending GE.
@@ -4164,7 +4188,7 @@ PRAGMA_ENABLE_DEPRECATION_WARNINGS
 		MarkArrayDirty();
 
 		// Once replicated state has caught up to this prediction key, we must remove this gameplay effect.
-		InPredictionKey.NewRejectOrCaughtUpDelegate(FPredictionKeyEvent::CreateUObject(Owner, &UAbilitySystemComponent::RemoveActiveGameplayEffect_NoReturn, AppliedActiveGE->Handle, -1));
+		InPredictionKey.NewRejectOrCaughtUpDelegate(FPredictionKeyEvent::CreateUObject(Owner, &UAbilitySystemComponent::RemoveActiveGameplayEffect_AllowClientRemoval, AppliedActiveGE->Handle, -1));
 		
 	}
 
diff --git a/Engine/Plugins/Runtime/GameplayAbilities/Source/GameplayAbilities/Private/GameplayEffectComponents/TargetTagRequirementsGameplayEffectComponent.cpp b/Engine/Plugins/Runtime/GameplayAbilities/Source/GameplayAbilities/Private/GameplayEffectComponents/TargetTagRequirementsGameplayEffectComponent.cpp
--- a/Engine/Plugins/Runtime/GameplayAbilities/Source/GameplayAbilities/Private/GameplayEffectComponents/TargetTagRequirementsGameplayEffectComponent.cpp
+++ b/Engine/Plugins/Runtime/GameplayAbilities/Source/GameplayAbilities/Private/GameplayEffectComponents/TargetTagRequirementsGameplayEffectComponent.cpp
@@ -5,6 +5,7 @@
 #include "AbilitySystemLog.h"
 #include "Algo/Find.h"
 #include "Misc/DataValidation.h"
+#include "AbilitySystemPrivate.h"
 
 #define LOCTEXT_NAMESPACE "TargetTagRequirementsGameplayEffectComponent"
 
@@ -25,7 +26,7 @@ bool UTargetTagRequirementsGameplayEffectComponent::CanGameplayEffectApply(const
 		return false;
 	}
 
-	if (!RemovalTagRequirements.IsEmpty() && RemovalTagRequirements.RequirementsMet(Tags) == true)
+	if (HaveRemovalRequirementsBeenMet(Tags, ActiveGEContainer.IsNetAuthority()))
 	{
 		return false;
 	}
@@ -42,6 +43,10 @@ bool UTargetTagRequirementsGameplayEffectComponent::OnActiveGameplayEffectAdded(
 		return false;
 	}
 
+	using namespace UE::AbilitySystem::Private;
+	const bool bAllowClientSideRemoval = EnumHasAnyFlags(static_cast<EAllowPredictiveGEFlags>(CVarAllowPredictiveGEFlagsValue), EAllowPredictiveGEFlags::AllowRemovalByTagRequirements);
+	const bool bAuthority = GEContainer.IsNetAuthority();
+
 	FActiveGameplayEffectHandle ActiveGEHandle = ActiveGE.Handle;
 	if (FActiveGameplayEffectEvents* EventSet = ASC->GetActiveEffectEventSet(ActiveGEHandle))
 	{
@@ -69,9 +74,13 @@ bool UTargetTagRequirementsGameplayEffectComponent::OnActiveGameplayEffectAdded(
 		AppendUnique(GameplayTagsToBind, OngoingTagRequirements.IgnoreTags.GetGameplayTagArray());
 		AppendUnique(GameplayTagsToBind, OngoingTagRequirements.RequireTags.GetGameplayTagArray());
 		AppendUnique(GameplayTagsToBind, OngoingTagRequirements.TagQuery.GetGameplayTagArray());
-		AppendUnique(GameplayTagsToBind, RemovalTagRequirements.IgnoreTags.GetGameplayTagArray());
-		AppendUnique(GameplayTagsToBind, RemovalTagRequirements.RequireTags.GetGameplayTagArray());
-		AppendUnique(GameplayTagsToBind, RemovalTagRequirements.TagQuery.GetGameplayTagArray());
+
+		if (bAuthority || bAllowClientSideRemoval)
+		{
+			AppendUnique(GameplayTagsToBind, RemovalTagRequirements.IgnoreTags.GetGameplayTagArray());
+			AppendUnique(GameplayTagsToBind, RemovalTagRequirements.RequireTags.GetGameplayTagArray());
+			AppendUnique(GameplayTagsToBind, RemovalTagRequirements.TagQuery.GetGameplayTagArray());
+		}
 
 		// Add our tag requirements to the ASC's Callbacks map. This helps filter down the amount of callbacks we'll get due to tag changes
 		// (rather than registering for the one callback whenever any tag changes).  We also need to keep track to remove those registered delegates in OnEffectRemoved.
@@ -123,11 +132,11 @@ void UTargetTagRequirementsGameplayEffectComponent::OnTagChanged(const FGameplay
 		FGameplayTagContainer OwnedTags;
 		Owner->GetOwnedGameplayTags(OwnedTags);
 
-		const bool bRemovalRequirementsMet = !RemovalTagRequirements.IsEmpty() && RemovalTagRequirements.RequirementsMet(OwnedTags);
+		const bool bRemovalRequirementsMet = HaveRemovalRequirementsBeenMet(OwnedTags, Owner->IsOwnerActorAuthoritative());
 		if (bRemovalRequirementsMet)
 		{
 			// This is slightly different functionality from pre-UE5.3, we're calling RemoveActiveGameplayEffect rather than InternalRemoveActiveGameplayEffect.
-			// The result is we set the calculated magnitudes back to zero.  This also used to only run on the Server.
+			// The result is we set the calculated magnitudes back to zero.  In UE5.3 this used to also run on the client (incorrect).
 			Owner->RemoveActiveGameplayEffect(ActiveGEHandle);
 		}
 		else
@@ -140,6 +149,15 @@ void UTargetTagRequirementsGameplayEffectComponent::OnTagChanged(const FGameplay
 	}
 }
 
+bool UTargetTagRequirementsGameplayEffectComponent::HaveRemovalRequirementsBeenMet(const FGameplayTagContainer& TargetOwnedTags, bool bNetAuthority) const
+{
+	using namespace UE::AbilitySystem::Private;
+	const bool bAllowClientSideRemoval = EnumHasAnyFlags(static_cast<EAllowPredictiveGEFlags>(CVarAllowPredictiveGEFlagsValue), EAllowPredictiveGEFlags::AllowRemovalByTagRequirements);
+	const bool bRemovalRequirementsMet = (bAllowClientSideRemoval || bNetAuthority) && !RemovalTagRequirements.IsEmpty() && RemovalTagRequirements.RequirementsMet(TargetOwnedTags);
+
+	return bRemovalRequirementsMet;
+}
+
 #if WITH_EDITOR
 /**
  * Validate incompatable configurations
diff --git a/Engine/Plugins/Runtime/GameplayAbilities/Source/GameplayAbilities/Private/GameplayPrediction.cpp b/Engine/Plugins/Runtime/GameplayAbilities/Source/GameplayAbilities/Private/GameplayPrediction.cpp
--- a/Engine/Plugins/Runtime/GameplayAbilities/Source/GameplayAbilities/Private/GameplayPrediction.cpp
+++ b/Engine/Plugins/Runtime/GameplayAbilities/Source/GameplayAbilities/Private/GameplayPrediction.cpp
@@ -22,16 +22,22 @@ namespace UE::AbilitySystem::Private
 		TEXT("How many prediction keys can be dropped before StaleKeyBehavior is run."));
 
 	// What should we do with these old stale FPredictionKeys?  Prior to UE5.5, we always CaughtUp.  I believe it's actually safer to drop.
-	int32 CVarStaleKeyBehaviorValue = 0;
+	int32 CVarStaleKeyBehaviorValue = 2;
 	FAutoConsoleVariableRef CVarStaleKeyBehavior(TEXT("AbilitySystem.PredictionKey.StaleKeyBehavior"), CVarStaleKeyBehaviorValue,
 		TEXT("How do we handle stale keys? 0 = CaughtUp. 1 = Reject. 2 = Drop"));
 
 	// How should we deal with dependent keys (in a chain)?  Prior to UE5.5, old keys implied new keys.  We introduced some new functionality (explained in the help text).
-	// 0 (no bitmask) is legacy behavior.  Logically, (0x1 | 0x2) = 3 is the correct value. The long-term fix will be a value of 3.
-	int32 CVarDependentChainBehaviorValue = 0;
+	// 0 (no bitmask) is legacy behavior.  Logically, (0x1 | 0x2) = 3 is the correct value. We default to 1 to move towards a long-term fix (of value 3).
+	int32 CVarDependentChainBehaviorValue = 1;
 	FAutoConsoleVariableRef CVarDependentChainBehavior(TEXT("AbilitySystem.PredictionKey.DepChainBehavior"), CVarDependentChainBehaviorValue,
 		TEXT("How do we handle dependency key chains? Bitmask: 0 = Old Accept/Rejected Implies Newer Accepted/Rejected. 0x1 = Newer Accepted also implies Older Accepted. 0x2 = Old Accepted Does NOT imply Newer Accepted"));
 
+	// Prior to UE5.5, we used to allow Server Initiated Replication Keys to be sent to the client as 'acknowledged' but that hardly makes sense.
+	// Unfortunately, it also causes key hash collisions since there is limited space based on the key value.
+	int32 CVarReplicateServerKeysAsAcknowledgedValue = 0;
+	FAutoConsoleVariableRef CVarReplicateServerKeysAsAcknowledged(TEXT("AbilitySystem.PredictionKey.RepServerKeysAsAcknowledged"), CVarReplicateServerKeysAsAcknowledgedValue,
+		TEXT("Do we send server initiated keys as acknowledged to the client? Default false. Was true prior to UE5.5."));
+
 	/**
 	 * Given an FProperty Link (such as a UFunction's Properties, or a UStruct's Properties), find and return any linked FPredictionKeys.
 	 * This is useful for determining if we're sending FPredictionKeys via RPC.
@@ -214,18 +220,14 @@ FPredictionKey FPredictionKey::CreateNewServerInitiatedKey(const UAbilitySystemC
 	// Only valid on the server
 	if (OwningComponent->GetOwnerRole() == ROLE_Authority)
 	{
-		NewKey.GenerateNewPredictionKey();
-		NewKey.bIsServerInitiated = true;
-
-#if !UE_BUILD_SHIPPING
 		// Make sure the Server and Client aren't synchronized in terms of key generation or it can hide bugs.
-		const uint32 ServerDifferential = static_cast<uint32>(reinterpret_cast<intptr_t>(OwningComponent));
-		NewKey.Current = (NewKey.Current + ServerDifferential) % std::numeric_limits<FPredictionKey::KeyType>::max();
-		if (NewKey.Current < 1)
+		static KeyType GServerKey = 1;
+		NewKey.bIsServerInitiated = true;
+		NewKey.Current = GServerKey++;
+		if (GServerKey <= 0)
 		{
-			NewKey.Current = 1;
+			GServerKey = 1;
 		}
-#endif
 	}
 	return NewKey;
 }
@@ -359,6 +361,10 @@ void FPredictionKeyDelegates::AddDependency(FPredictionKey::KeyType ThisKey, FPr
 
 // -------------------------------------
 
+/** 
+ * This is the Server version of FScopedPredictionWindow constructor.
+ * This exists for legacy reasons and I'm not convinced this needs to exist.  Instead we should manually accept/reject the FPredictionKey (currently done in the destructor).
+ */
 FScopedPredictionWindow::FScopedPredictionWindow(UAbilitySystemComponent* AbilitySystemComponent, FPredictionKey InPredictionKey, bool InSetReplicatedPredictionKey /*=true*/)
 {
 	if (AbilitySystemComponent == nullptr)
@@ -445,7 +451,7 @@ FScopedPredictionWindow::FScopedPredictionWindow(UAbilitySystemComponent* InAbil
 					{
 						if (PassedInKey.Current == BaseKey || PassedInKey.Base == BaseKey)
 						{
-							UE_LOG(LogPredictionKey, Verbose, TEXT("Sent key %s to confirm BaseKey %d"), *PassedInKey.ToString(), BaseKey);
+							UE_LOG(LogPredictionKey, Verbose, TEXT("Sent key %s with %s to confirm BaseKey %d"), *PassedInKey.ToString(), *Function->GetName(), BaseKey);
 							DebugBaseKeyOfChain.Reset();
 							break;
 						}
@@ -465,7 +471,16 @@ FScopedPredictionWindow::~FScopedPredictionWindow()
 		if (DebugBaseKeyOfChain.IsSet())
 		{
 			const FPredictionKey::KeyType BaseKey = DebugBaseKeyOfChain.GetValue();
-			UE_CLOG(FPredictionKeyDelegates::Get().DelegateMap.Contains(BaseKey), LogPredictionKey, Warning, TEXT("No key based off PredictionKey %d was communicated to the server during ScopedPredictionWindow. This will likely leak a FPredictionKey with dangling CaughtUpTo/Rejected delegates."), BaseKey);
+			if (const FPredictionKeyDelegates::FDelegates* BoundDelegates = FPredictionKeyDelegates::Get().DelegateMap.Find(BaseKey))
+			{
+				const bool bIsBound = BoundDelegates->RejectedDelegates.Num() > 0 || BoundDelegates->CaughtUpDelegates.Num() > 0;
+
+				// If you get this log, you can set a breakpoint here and inspect the Delegates to see what it's bound to and therefore what code is at fault (the binding code, or the code that never sends the key to the server). 
+				UE_CLOG(bIsBound, LogPredictionKey, Warning, TEXT("No key based off PredictionKey %d was communicated to the server during ScopedPredictionWindow. The PredictionKey has bound callbacks that can never be called (leak)."), BaseKey);
+
+				// If you get this log, you can check the call stack to see where this FScopedPredictionWindow originated from, and why it would be creating an entry in the FPredictionKeyDelegates without any bindings to it.
+				UE_CLOG(!bIsBound, LogPredictionKey, Warning, TEXT("No key based off PredictionKey %d was communicated to the server during ScopedPredictionWindow. The PredictionKey is not bound, but has leaked an entry in FPredictionKeyDelegates."), BaseKey);
+			}
 		}
 		
 		NetDriver->SendRPCDel = DebugSavedOnSendRPC;
@@ -485,7 +500,18 @@ FScopedPredictionWindow::~FScopedPredictionWindow()
 			// (for example, predict w/ key 100 -> prediction key replication dropped -> predict w/ invalid key -> next rep of prediction key is 0).
 			if (OwnerPtr->ScopedPredictionKey.IsValidKey())
 			{
-				OwnerPtr->ReplicatedPredictionKeyMap.ReplicatePredictionKey(OwnerPtr->ScopedPredictionKey);
+				const bool bServerInitiatedKey = OwnerPtr->ScopedPredictionKey.IsServerInitiatedKey();
+				const bool bAllowAckServerInitiatedKey = UE::AbilitySystem::Private::CVarReplicateServerKeysAsAcknowledgedValue > 0;
+				if (!bServerInitiatedKey || bAllowAckServerInitiatedKey)
+				{
+					UE_CLOG(bServerInitiatedKey, LogAbilitySystem, Warning, TEXT("Replicating Server Initiated PredictionKey %s (this may stomp a client key leaving it unack'd). See CVarReplicateServerKeysAsAcknowledged"), *OwnerPtr->ScopedPredictionKey.ToString());
+					UE_VLOG_UELOG(OwnerPtr->GetOwnerActor(), LogPredictionKey, Verbose, TEXT("Server: ReplicatePredictionKey %s"), *OwnerPtr->ScopedPredictionKey.ToString());
+					OwnerPtr->ReplicatedPredictionKeyMap.ReplicatePredictionKey(OwnerPtr->ScopedPredictionKey);
+				}
+				else
+				{
+					UE_LOG(LogAbilitySystem, Verbose, TEXT("Server: NOT replicating PredictionKey %s due to being server initiated. See CVarReplicateServerKeysAsAcknowledged."), *OwnerPtr->ScopedPredictionKey.ToString());
+				}
 			}
 		}
 		if (ClearScopedPredictionKey)
@@ -512,6 +538,7 @@ void FReplicatedPredictionKeyItem::OnRep(const FReplicatedPredictionKeyMap& InAr
 	// initiates a key that matches a local value.  Then we catch-up to a local value, even though the server was not specifically acknowledging it.
 	if (PredictionKey.bIsServerInitiated)
 	{
+		UE_LOG(LogPredictionKey, Warning, TEXT("FReplicatedPredictionKeyItem::OnRep received Server Initiated Key %s (which likely stomped a local key due to keymap hash collisions)."), *PredictionKey.ToString());
 		return;
 	}
 
diff --git a/Engine/Plugins/Runtime/GameplayAbilities/Source/GameplayAbilities/Public/AbilitySystemBlueprintLibrary.h b/Engine/Plugins/Runtime/GameplayAbilities/Source/GameplayAbilities/Public/AbilitySystemBlueprintLibrary.h
--- a/Engine/Plugins/Runtime/GameplayAbilities/Source/GameplayAbilities/Public/AbilitySystemBlueprintLibrary.h
+++ b/Engine/Plugins/Runtime/GameplayAbilities/Source/GameplayAbilities/Public/AbilitySystemBlueprintLibrary.h
@@ -119,9 +119,14 @@ class GAMEPLAYABILITIES_API UAbilitySystemBlueprintLibrary : public UBlueprintFu
 	static FGameplayTargetDataFilterHandle MakeFilterHandle(FGameplayTargetDataFilter Filter, AActor* FilterActor);
 
 	/** Create a spec handle, filling out all fields */
-	UFUNCTION(BlueprintPure, Category = "Spec")
+	UE_DEPRECATED(5.5, "Use MakeSpecHandleByClass. It's safer as InGameplayEffect needs to be a CDO")
+	UFUNCTION(BlueprintPure, Category = "Spec", meta=(DeprecatedFunction, DeprecatedMessage="Use Make Spec Handle (By Class)"))
 	static FGameplayEffectSpecHandle MakeSpecHandle(UGameplayEffect* InGameplayEffect, AActor* InInstigator, AActor* InEffectCauser, float InLevel = 1.0f);
 
+	/** Create a spec handle, filling out all fields */
+	UFUNCTION(BlueprintCallable, Category = "Spec", meta=(DisplayName="Make Spec Handle (By Class)"))
+	static FGameplayEffectSpecHandle MakeSpecHandleByClass(TSubclassOf<UGameplayEffect> GameplayEffect, AActor* Instigator, AActor* EffectCauser, float Level = 1.0f);
+
 	/** Create a spec handle, cloning another */
 	UFUNCTION(BlueprintPure, Category = "Spec")
 	static FGameplayEffectSpecHandle CloneSpecHandle(AActor* InNewInstigator, AActor* InEffectCauser, FGameplayEffectSpecHandle GameplayEffectSpecHandle_Clone);
diff --git a/Engine/Plugins/Runtime/GameplayAbilities/Source/GameplayAbilities/Public/AbilitySystemComponent.h b/Engine/Plugins/Runtime/GameplayAbilities/Source/GameplayAbilities/Public/AbilitySystemComponent.h
--- a/Engine/Plugins/Runtime/GameplayAbilities/Source/GameplayAbilities/Public/AbilitySystemComponent.h
+++ b/Engine/Plugins/Runtime/GameplayAbilities/Source/GameplayAbilities/Public/AbilitySystemComponent.h
@@ -392,10 +392,8 @@ class GAMEPLAYABILITIES_API UAbilitySystemComponent : public UGameplayTasksCompo
 	int32 GetAggregatedStackCount(const FGameplayEffectQuery& Query) const;
 
 	/** This only exists so it can be hooked up to a multicast delegate */
-	void RemoveActiveGameplayEffect_NoReturn(FActiveGameplayEffectHandle Handle, int32 StacksToRemove=-1)
-	{
-		RemoveActiveGameplayEffect(Handle, StacksToRemove);
-	}
+	UE_DEPRECATED(5.5, "This shouldn't be public.  Use RemoveActiveGameplayEffect")
+	void RemoveActiveGameplayEffect_NoReturn(FActiveGameplayEffectHandle Handle, int32 StacksToRemove = -1);
 
 	/** Called for predictively added gameplay cue. Needs to remove tag count and possible invoke OnRemove event if misprediction */
 	virtual void OnPredictiveGameplayCueCatchup(FGameplayTag Tag);
@@ -1868,7 +1866,10 @@ protected:
 	void CheckDurationExpired(FActiveGameplayEffectHandle Handle);
 		
 	TArray<TObjectPtr<UGameplayTask>>&	GetAbilityActiveTasks(UGameplayAbility* Ability);
-	
+
+	/** A version of RemoveActiveGameplayEffect that allows us to remove it even when we don't have authority */
+	void RemoveActiveGameplayEffect_AllowClientRemoval(FActiveGameplayEffectHandle Handle, int32 StacksToRemove = -1);
+
 	/** Contains all of the gameplay effects that are currently active on this component */
 	UPROPERTY(Replicated)
 	FActiveGameplayEffectsContainer ActiveGameplayEffects;
diff --git a/Engine/Plugins/Runtime/GameplayAbilities/Source/GameplayAbilities/Public/GameplayAbilitySpec.h b/Engine/Plugins/Runtime/GameplayAbilities/Source/GameplayAbilities/Public/GameplayAbilitySpec.h
--- a/Engine/Plugins/Runtime/GameplayAbilities/Source/GameplayAbilities/Public/GameplayAbilitySpec.h
+++ b/Engine/Plugins/Runtime/GameplayAbilities/Source/GameplayAbilities/Public/GameplayAbilitySpec.h
@@ -275,6 +275,7 @@ PRAGMA_ENABLE_DEPRECATION_WARNINGS
 	bool IsActive() const;
 
 	void PreReplicatedRemove(const struct FGameplayAbilitySpecContainer& InArraySerializer);
+	void PostReplicatedChange(const struct FGameplayAbilitySpecContainer& InArraySerializer);
 	void PostReplicatedAdd(const struct FGameplayAbilitySpecContainer& InArraySerializer);
 
 	FString GetDebugString();
diff --git a/Engine/Plugins/Runtime/GameplayAbilities/Source/GameplayAbilities/Public/GameplayEffectComponents/TargetTagRequirementsGameplayEffectComponent.h b/Engine/Plugins/Runtime/GameplayAbilities/Source/GameplayAbilities/Public/GameplayEffectComponents/TargetTagRequirementsGameplayEffectComponent.h
--- a/Engine/Plugins/Runtime/GameplayAbilities/Source/GameplayAbilities/Public/GameplayEffectComponents/TargetTagRequirementsGameplayEffectComponent.h
+++ b/Engine/Plugins/Runtime/GameplayAbilities/Source/GameplayAbilities/Public/GameplayEffectComponents/TargetTagRequirementsGameplayEffectComponent.h
@@ -38,6 +38,9 @@ private:
 	/** We need to register a callback for when the owner changes its tags.  When that happens, we need to figure out if our GE should continue to execute */
 	void OnTagChanged(const FGameplayTag GameplayTag, int32 NewCount, FActiveGameplayEffectHandle ActiveGEHandle) const;
 
+	/** Have requirements to remove this Gameplay Effect been met? */
+	bool HaveRemovalRequirementsBeenMet(const FGameplayTagContainer& TargetOwnedTags, bool bNetAuthority) const;
+
 public:
 	/** Tag requirements the target must have for this GameplayEffect to be applied. This is pass/fail at the time of application. If fail, this GE fails to apply. */
 	UPROPERTY(EditDefaultsOnly, Category = Tags)
