Fixed failing Gameplay Cue unit tests.

diff --git a/Engine/Plugins/Runtime/GameplayAbilities/Source/GameplayAbilities/Private/Tests/GameplayEffectTests.cpp b/Engine/Plugins/Runtime/GameplayAbilities/Source/GameplayAbilities/Private/Tests/GameplayEffectTests.cpp
--- a/Engine/Plugins/Runtime/GameplayAbilities/Source/GameplayAbilities/Private/Tests/GameplayEffectTests.cpp
+++ b/Engine/Plugins/Runtime/GameplayAbilities/Source/GameplayAbilities/Private/Tests/GameplayEffectTests.cpp
@@ -488,7 +488,7 @@ public: // the tests
 				Test->TestEqual(TEXT("  OwnerTagCount"), DestComponent->GetTagCount(UE::GameplayTags::GameplayCue_Test), ExpectedOwnerTagCount);
 
 				// Now try to remove it
-				DestComponent->RemoveActiveGameplayEffect(ActiveGEHandle);
+				DestComponent->RemoveActiveGameplayEffect_NoReturn(ActiveGEHandle);
 
 				Test->TestFalse(TEXT("  IsGameplayCueActive (After RemoveActiveEffects)"), DestComponent->IsGameplayCueActive(UE::GameplayTags::GameplayCue_Test));
 				Test->TestEqual(TEXT("  OnRemove Calls"), GCNotify_Test_CDO->NumOnRemoveCalls, ExpectedOnRemove);
@@ -547,7 +547,7 @@ public: // the tests
 					for (int Index = 0; Index < ExpectedOnActive; ++Index)
 					{
 						Test->TestEqual(TEXT("  IsGameplayCueActive before RemoveActiveEffects Unstack"), DestComponent->IsGameplayCueActive(UE::GameplayTags::GameplayCue_Test), (ExpectedOnActive > 0));
-						DestComponent->RemoveActiveGameplayEffect(ActiveGEHandle, StacksToRemove);
+						DestComponent->RemoveActiveGameplayEffect_NoReturn(ActiveGEHandle, StacksToRemove);
 					}
 				}
 				else
diff --git a/Engine/Plugins/Runtime/GameplayAbilities/Source/GameplayAbilities/Public/AbilitySystemComponent.h b/Engine/Plugins/Runtime/GameplayAbilities/Source/GameplayAbilities/Public/AbilitySystemComponent.h
--- a/Engine/Plugins/Runtime/GameplayAbilities/Source/GameplayAbilities/Public/AbilitySystemComponent.h
+++ b/Engine/Plugins/Runtime/GameplayAbilities/Source/GameplayAbilities/Public/AbilitySystemComponent.h
@@ -392,7 +392,7 @@ class GAMEPLAYABILITIES_API UAbilitySystemComponent : public UGameplayTasksCompo
 	int32 GetAggregatedStackCount(const FGameplayEffectQuery& Query) const;
 
 	/** This only exists so it can be hooked up to a multicast delegate */
-	UE_DEPRECATED(5.5, "This shouldn't be public.  Use RemoveActiveGameplayEffect")
+	UE_DEPRECATED_FORGAME(5.5, "This shouldn't be public. Use RemoveActiveGameplayEffect (which won't allow non-authority to remove) or expose RemoveActiveGameplayEffect_AllowClientRemoval if your intent is to break that contract.")
 	void RemoveActiveGameplayEffect_NoReturn(FActiveGameplayEffectHandle Handle, int32 StacksToRemove = -1);
 
 	/** Called for predictively added gameplay cue. Needs to remove tag count and possible invoke OnRemove event if misprediction */
