PR #11697: Adding `meta=(DeterminesOutputType)` to `GetAttributeSet` of `AbilitySystemComponent`

diff --git a/Engine/Plugins/Runtime/GameplayAbilities/Source/GameplayAbilities/Public/AbilitySystemComponent.h b/Engine/Plugins/Runtime/GameplayAbilities/Source/GameplayAbilities/Public/AbilitySystemComponent.h
--- a/Engine/Plugins/Runtime/GameplayAbilities/Source/GameplayAbilities/Public/AbilitySystemComponent.h
+++ b/Engine/Plugins/Runtime/GameplayAbilities/Source/GameplayAbilities/Public/AbilitySystemComponent.h
@@ -184,7 +184,7 @@ class GAMEPLAYABILITIES_API UAbilitySystemComponent : public UGameplayTasksCompo
 	 * @param AttributeSetClass The type of attribute set to look for
 	 * @param bFound Set to true if an instance of the Attribute Set exists
 	 */
-	UFUNCTION(BlueprintCallable, BlueprintPure = false, Category = "Gameplay Attributes")
+	UFUNCTION(BlueprintCallable, BlueprintPure = false, Category = "Gameplay Attributes", meta=(DeterminesOutputType = AttributeSetClass))
 	const UAttributeSet* GetAttributeSet(TSubclassOf<UAttributeSet> AttributeSetClass) const;
 
 	/**
