Handling a case where InstancedPerExecution abilities could ensure if being Removed while in an EndAbility callstack.

diff --git a/Engine/Plugins/Runtime/GameplayAbilities/Source/GameplayAbilities/Private/AbilitySystemComponent_Abilities.cpp b/Engine/Plugins/Runtime/GameplayAbilities/Source/GameplayAbilities/Private/AbilitySystemComponent_Abilities.cpp
--- a/Engine/Plugins/Runtime/GameplayAbilities/Source/GameplayAbilities/Private/AbilitySystemComponent_Abilities.cpp
+++ b/Engine/Plugins/Runtime/GameplayAbilities/Source/GameplayAbilities/Private/AbilitySystemComponent_Abilities.cpp
@@ -664,10 +664,29 @@ void UAbilitySystemComponent::OnRemoveAbility(FGameplayAbilitySpec& Spec)
 	}
 	else
 	{
-		// If we're non-instanced and still active, we need to End
+		// This is error handling some edge cases
 		if (Spec.IsActive())
 		{
-			if (ensureMsgf(Spec.Ability->GetInstancingPolicy() == EGameplayAbilityInstancingPolicy::NonInstanced, TEXT("We should never have an instanced Gameplay Ability that is still active by this point. All instances should have EndAbility called just before here.")))
+			if (Spec.Ability->GetInstancingPolicy() == EGameplayAbilityInstancingPolicy::InstancedPerExecution)
+			{
+#if DO_ENSURE
+				// Let's make sure we aren't recursively executing this ability.
+				Instances = Spec.GetAbilityInstances();
+				for (UGameplayAbility* Instance : Instances)
+				{
+					ensureMsgf(Instance->bIsAbilityEnding, TEXT("All instances of %s on %s should have been ended by now. Maybe it was retriggered from OnEndAbility (bad)?"), *GetNameSafe(Spec.Ability), *GetName());
+				}
+#endif
+
+				// If we're a client, fall-through anyway as this is our last chance to execute OnRemoveAbility.
+				if (IsOwnerActorAuthoritative())
+				{
+					// We assume we are inside an EndAbility callstack (bIsActive && bIsAbilityEnding), so remove us when we get out of it (NotifyAbilityEnded).
+					Spec.RemoveAfterActivation = true;
+					return;
+				}
+			}
+			else if (ensureMsgf(Spec.Ability->GetInstancingPolicy() == EGameplayAbilityInstancingPolicy::NonInstanced, TEXT("We should never have an instanced Gameplay Ability that is still active by this point. All instances should have EndAbility called just before here.")))
 			{
 				// Seems like it should be cancelled, but we're just following the existing pattern (could be due to functionality from OnRep)
 				constexpr bool bReplicateEndAbility = false;
