Fix issue where PostCDOCompiledFixupSubobjects would try to recreate subobjects that already existed, which causes problems in some editor loading situations

diff --git a/Engine/Plugins/Runtime/GameplayAbilities/Source/GameplayAbilities/Private/GameplayEffect.cpp b/Engine/Plugins/Runtime/GameplayAbilities/Source/GameplayAbilities/Private/GameplayEffect.cpp
--- a/Engine/Plugins/Runtime/GameplayAbilities/Source/GameplayAbilities/Private/GameplayEffect.cpp
+++ b/Engine/Plugins/Runtime/GameplayAbilities/Source/GameplayAbilities/Private/GameplayEffect.cpp
@@ -372,6 +372,9 @@ void UGameplayEffect::PostCDOCompiledFixupSubobjects()
 		return;
 	}
 
+	// Clear out any bad inherited components
+	GEComponents.Remove(nullptr);
+
 	for (const UGameplayEffectComponent* ParentComponent : Archetype->GEComponents)
 	{
 		if (!ParentComponent)
@@ -398,9 +401,19 @@ void UGameplayEffect::PostCDOCompiledFixupSubobjects()
 			continue;
 		}
 
-		// We don't already have the Archetype's Component, so add it here using the exact same name so we link up.
-		UE_LOG(LogGameplayEffects, Verbose, TEXT("%s is manually duplicating Archetype %s because they were not inherited through automatic instancing"), *GetFullNameSafe(this), *GetFullNameSafe(ParentComponent));
-		UGameplayEffectComponent* ChildComponent = DuplicateObject(ParentComponent, this, ParentComponentName);
+		// This will find components inherited from the parent that are missing from the array
+		UGameplayEffectComponent* ChildComponent = FindObjectFast<UGameplayEffectComponent>(this, *ParentComponentName.ToString());
+		if (!ChildComponent)
+		{
+			// We don't already have the Archetype's Component, so add it here using the exact same name so we link up.
+			UE_LOG(LogGameplayEffects, Verbose, TEXT("%s is manually duplicating Archetype %s because they were not inherited through automatic instancing"), *GetFullNameSafe(this), *GetFullNameSafe(ParentComponent));
+			ChildComponent = DuplicateObject(ParentComponent, this, ParentComponentName);
+		}
+		else
+		{
+			UE_LOG(LogGameplayEffects, Verbose, TEXT("%s is re-adding inherited copy of Archetype %s to components array"), *GetFullNameSafe(this), *GetFullNameSafe(ParentComponent));
+		}
+
 		GEComponents.Add(ChildComponent);
 	}
 }
