PR #11498: Add more detailed warning log to UAbilitySystemBlueprintLibrary::MakeSpecHandle

diff --git a/Engine/Plugins/Runtime/GameplayAbilities/Source/GameplayAbilities/Private/AbilitySystemBlueprintLibrary.cpp b/Engine/Plugins/Runtime/GameplayAbilities/Source/GameplayAbilities/Private/AbilitySystemBlueprintLibrary.cpp
--- a/Engine/Plugins/Runtime/GameplayAbilities/Source/GameplayAbilities/Private/AbilitySystemBlueprintLibrary.cpp
+++ b/Engine/Plugins/Runtime/GameplayAbilities/Source/GameplayAbilities/Private/AbilitySystemBlueprintLibrary.cpp
@@ -273,7 +273,11 @@ FGameplayEffectSpecHandle UAbilitySystemBlueprintLibrary::MakeSpecHandle(UGamepl
 		return FGameplayEffectSpecHandle(new FGameplayEffectSpec(InGameplayEffect, FGameplayEffectContextHandle(EffectContext), InLevel));
 	}
 	
-	ABILITY_LOG(Warning, TEXT("%s was called with an invalid GameplayEffect object!"), *FString(__FUNCTION__));
+	const FString InstigatorName = InInstigator ? InInstigator->GetActorNameOrLabel() : TEXT("None");
+	const FString CauserName = InEffectCauser ? InEffectCauser->GetActorNameOrLabel() : TEXT("None");
+
+	ABILITY_LOG(Warning, TEXT("[%hs] called with null GameplayEffect. Instigator: %s, Causer: %s"), __func__, *InstigatorName, *CauserName);
+	
 	return FGameplayEffectSpecHandle();
 }
 
