PR #10976: GAS ActivationOwnedTags bug fix

diff --git a/Engine/Plugins/Runtime/GameplayAbilities/Source/GameplayAbilities/Private/Abilities/GameplayAbility.cpp b/Engine/Plugins/Runtime/GameplayAbilities/Source/GameplayAbilities/Private/Abilities/GameplayAbility.cpp
--- a/Engine/Plugins/Runtime/GameplayAbilities/Source/GameplayAbilities/Private/Abilities/GameplayAbility.cpp
+++ b/Engine/Plugins/Runtime/GameplayAbilities/Source/GameplayAbilities/Private/Abilities/GameplayAbility.cpp
@@ -753,7 +753,15 @@ void UGameplayAbility::EndAbility(const FGameplayAbilitySpecHandle Handle, const
 
 			if (UAbilitySystemGlobals::Get().ShouldReplicateActivationOwnedTags())
 			{
-				AbilitySystemComponent->RemoveReplicatedLooseGameplayTags(ActivationOwnedTags);
+				if (GetNetExecutionPolicy() == EGameplayAbilityNetExecutionPolicy::LocalPredicted || GetNetExecutionPolicy() == EGameplayAbilityNetExecutionPolicy::ServerInitiated)
+				{
+					// If this ability also executes on the client, then don't communicate the tags to the client (it already used RemoveLooseGameplayTags above)
+					AbilitySystemComponent->RemoveMinimalReplicationGameplayTags(ActivationOwnedTags);
+				}
+				else
+				{
+					AbilitySystemComponent->RemoveReplicatedLooseGameplayTags(ActivationOwnedTags);
+				}
 			}
 
 			// Remove tracked GameplayCues that we added
@@ -871,7 +879,15 @@ void UGameplayAbility::PreActivate(const FGameplayAbilitySpecHandle Handle, cons
 
 	if (UAbilitySystemGlobals::Get().ShouldReplicateActivationOwnedTags())
 	{
-		Comp->AddReplicatedLooseGameplayTags(ActivationOwnedTags);
+		if (GetNetExecutionPolicy() == EGameplayAbilityNetExecutionPolicy::LocalPredicted || GetNetExecutionPolicy() == EGameplayAbilityNetExecutionPolicy::ServerInitiated)
+		{
+			// If this ability also executes on the client, then don't communicate the tags to the client (it already used AddLooseGameplayTags above)
+			Comp->AddMinimalReplicationGameplayTags(ActivationOwnedTags);
+		}
+		else
+		{
+			Comp->AddReplicatedLooseGameplayTags(ActivationOwnedTags);
+		}
 	}
 
 	if (OnGameplayAbilityEndedDelegate)
