Refactor and Move GameplayDebuggerCategory_Abilities so that it can be accessed by other modules

diff --git a/Engine/Plugins/Runtime/GameplayAbilities/Source/GameplayAbilities/Private/GameplayDebuggerCategory_Abilities.h b/Engine/Plugins/Runtime/GameplayAbilities/Source/GameplayAbilities/Private/GameplayDebuggerCategory_Abilities.h
deleted file mode 100644
diff --git a/Engine/Plugins/Runtime/GameplayAbilities/Source/GameplayAbilities/Public/GameplayDebuggerCategory_Abilities.h b/Engine/Plugins/Runtime/GameplayAbilities/Source/GameplayAbilities/Public/GameplayDebuggerCategory_Abilities.h
new file mode 100644
--- /dev/null
+++ b/Engine/Plugins/Runtime/GameplayAbilities/Source/GameplayAbilities/Public/GameplayDebuggerCategory_Abilities.h
@@ -0,0 +1,110 @@
+// Copyright Epic Games, Inc. All Rights Reserved.
+
+#pragma once
+
+#include "CoreMinimal.h"
+
+#if WITH_GAMEPLAY_DEBUGGER_MENU
+#include "GameplayTagContainer.h"
+#include "GameplayDebuggerCategory.h"
+#include "GameplayPrediction.h"
+
+class AActor;
+class APlayerController;
+class UAbilitySystemComponent;
+class UPackageMap;
+
+class GAMEPLAYABILITIES_API FGameplayDebuggerCategory_Abilities : public FGameplayDebuggerCategory
+{
+public:
+	FGameplayDebuggerCategory_Abilities();
+
+	virtual void CollectData(APlayerController* OwnerPC, AActor* DebugActor) override;
+	virtual void DrawData(APlayerController* OwnerPC, FGameplayDebuggerCanvasContext& CanvasContext) override;
+
+	static TSharedRef<FGameplayDebuggerCategory> MakeInstance();
+
+	void OnShowGameplayTagsToggle();
+	void OnShowGameplayAbilitiesToggle();
+	void OnShowGameplayEffectsToggle();
+	void OnShowGameplayAttributesToggle();
+		
+
+	// Some GAS features such as Attributes can exist on the server, client, or both.  We can also get 'detached' if both sides have the same values (such as Attributes) that aren't networked.
+	// We are reusing this same concept to try and reconcile Gameplay Effects predicted locally, or triggered server-side.
+	enum class ENetworkStatus : uint8
+	{
+		ServerOnly, LocalOnly, Networked, Detached, MAX
+	};
+
+	// Unary operator + for quick conversion from enum class to int32
+	friend constexpr int32 operator+(const ENetworkStatus& value) { return static_cast<int32>(value); }
+
+protected:
+
+	void DrawGameplayTags(FGameplayDebuggerCanvasContext& CanvasContext, const APlayerController* OwnerPC) const;
+	void DrawGameplayAbilities(FGameplayDebuggerCanvasContext& CanvasContext, const APlayerController* OwnerPC) const;
+	void DrawGameplayEffects(FGameplayDebuggerCanvasContext& CanvasContext, const APlayerController* OwnerPC) const;
+	void DrawGameplayAttributes(FGameplayDebuggerCanvasContext& CanvasContext, const APlayerController* OwnerPC) const;
+
+	struct FRepData
+	{
+		// to aid in NetSerialize
+		TWeakObjectPtr<UPackageMap>	ClientPackageMap;
+
+		FGameplayTagContainer OwnedTags;
+		TArray<int32> TagCounts;
+
+		struct FGameplayAbilityDebug
+		{
+			FString Ability;
+			FString Source;
+			int32 Level = 0;
+			bool bIsActive = false;
+		};
+		TArray<FGameplayAbilityDebug> Abilities;
+
+		struct FGameplayEffectDebug
+		{
+			int32 ReplicationID = INDEX_NONE; // unique & shared between server/client (or INDEX_NONE if local)
+			FString Effect;
+			FString Context;
+			float Duration = 0.0f;
+			float Period = 0.0f;
+			int32 Stacks = 0;
+			float Level = 0.0f;
+
+			ENetworkStatus NetworkStatus = ENetworkStatus::ServerOnly;
+			bool bInhibited = false;
+		};
+		TArray<FGameplayEffectDebug> GameplayEffects;
+
+		struct FGameplayAttributeDebug
+		{
+			FString AttributeName;
+			float BaseValue = 0.0f;
+			float CurrentValue = 0.0f;
+			ENetworkStatus NetworkStatus = ENetworkStatus::ServerOnly;
+		};
+		TArray<FGameplayAttributeDebug> Attributes;
+
+		void Serialize(FArchive& Ar);
+	};
+	FRepData DataPack;
+
+	bool WrapStringAccordingToViewport(const FString& iStr, FString& oStr, FGameplayDebuggerCanvasContext& CanvasContext, float ViewportWitdh) const;
+
+private:
+	TArray<FRepData::FGameplayAttributeDebug> CollectAttributeData(const APlayerController* OwnerPC, const UAbilitySystemComponent* AbilityComp) const;
+	TArray<FRepData::FGameplayEffectDebug> CollectEffectsData(const APlayerController* OwnerPC, const UAbilitySystemComponent* AbilityComp) const;
+
+	// Save off the last expected draw size so that we can draw a border around it next frame (and hope we're the same size)
+	float LastDrawDataEndSize = 0.0f;
+
+	bool bShowGameplayTags = true;
+	bool bShowGameplayAbilities = true;
+	bool bShowGameplayEffects = true;
+	bool bShowGameplayAttributes = true;
+};
+
+#endif // WITH_GAMEPLAY_DEBUGGER_MENU
