PR #11717: Added context menu to Gameplay Attribute widget

diff --git a/Engine/Plugins/Runtime/GameplayAbilities/Source/GameplayAbilitiesEditor/GameplayAbilitiesEditor.Build.cs b/Engine/Plugins/Runtime/GameplayAbilities/Source/GameplayAbilitiesEditor/GameplayAbilitiesEditor.Build.cs
--- a/Engine/Plugins/Runtime/GameplayAbilities/Source/GameplayAbilitiesEditor/GameplayAbilitiesEditor.Build.cs
+++ b/Engine/Plugins/Runtime/GameplayAbilities/Source/GameplayAbilitiesEditor/GameplayAbilitiesEditor.Build.cs
@@ -43,7 +43,8 @@ namespace UnrealBuildTool.Rules
 					"MovieScene",
 					"DataRegistry",
 					"DataRegistryEditor",
-					"ToolMenus"
+					"ToolMenus",
+					"ApplicationCore"
 				}
 			);
 		}
diff --git a/Engine/Plugins/Runtime/GameplayAbilities/Source/GameplayAbilitiesEditor/Private/SGameplayAttributeWidget.cpp b/Engine/Plugins/Runtime/GameplayAbilities/Source/GameplayAbilitiesEditor/Private/SGameplayAttributeWidget.cpp
--- a/Engine/Plugins/Runtime/GameplayAbilities/Source/GameplayAbilitiesEditor/Private/SGameplayAttributeWidget.cpp
+++ b/Engine/Plugins/Runtime/GameplayAbilities/Source/GameplayAbilitiesEditor/Private/SGameplayAttributeWidget.cpp
@@ -1,21 +1,26 @@
 // Copyright Epic Games, Inc. All Rights Reserved.
 
 #include "SGameplayAttributeWidget.h"
-#include "UObject/UnrealType.h"
 #include "Widgets/Layout/SSeparator.h"
 #include "Widgets/Layout/SBox.h"
 #include "Widgets/Views/STableViewBase.h"
 #include "Widgets/Views/STableRow.h"
 #include "Widgets/Views/SListView.h"
 #include "Widgets/Input/SComboBox.h"
-
-#include "AttributeSet.h"
 #include "Widgets/Input/SSearchBox.h"
-#include "UObject/UObjectHash.h"
-#include "UObject/UObjectIterator.h"
+
 #include "AbilitySystemComponent.h"
+#include "AttributeSet.h"
+#include "Editor.h"
+#include "Framework/MultiBox/MultiBoxBuilder.h"
+#include "HAL/PlatformApplicationMisc.h"
 #include "Misc/TextFilter.h"
+#include "Misc/OutputDeviceNull.h"
+#include "ScopedTransaction.h"
 #include "SlateOptMacros.h"
+#include "UObject/UnrealType.h"
+#include "UObject/UObjectHash.h"
+#include "UObject/UObjectIterator.h"
 
 #define LOCTEXT_NAMESPACE "K2Node"
 
@@ -345,23 +350,147 @@ void SGameplayAttributeWidget::Construct(const FArguments& InArgs)
 	OnAttributeChanged = InArgs._OnAttributeChanged;
 	SelectedProperty = InArgs._DefaultProperty;
 
-	// set up the combo button
-	SAssignNew(ComboButton, SComboButton)
-		.OnGetMenuContent(this, &SGameplayAttributeWidget::GenerateAttributePicker)
-		.ContentPadding(FMargin(2.0f, 2.0f))
-		.ToolTipText(this, &SGameplayAttributeWidget::GetSelectedValueAsString)
-		.ButtonContent()
-		[
-			SNew(STextBlock)
-			.Text(this, &SGameplayAttributeWidget::GetSelectedValueAsString)
-		];
+	TWeakPtr<SGameplayAttributeWidget> WeakSelf = StaticCastWeakPtr<SGameplayAttributeWidget>(AsWeak());
 
 	ChildSlot
 	[
-		ComboButton.ToSharedRef()
+		SNew(SBorder)
+			.OnMouseButtonDown_Lambda([WeakSelf](const FGeometry& MyGeometry, const FPointerEvent& MouseEvent)
+				{
+					const TSharedPtr<SGameplayAttributeWidget> Self = WeakSelf.Pin();
+					if (Self.IsValid() && MouseEvent.IsMouseButtonDown(EKeys::RightMouseButton))
+					{
+						Self->OnMenu(MouseEvent);
+						return FReply::Handled();
+					}
+					return FReply::Unhandled();
+				})
+			.Padding(0.0f)
+			.BorderImage(FStyleDefaults::GetNoBrush())
+			[
+				// set up the combo button
+				SAssignNew(ComboButton, SComboButton)
+					.OnGetMenuContent(this, &SGameplayAttributeWidget::GenerateAttributePicker)
+					.ContentPadding(FMargin(2.0f, 2.0f))
+					.ToolTipText(this, &SGameplayAttributeWidget::GetSelectedValueAsString)
+					.ButtonContent()
+					[
+						SNew(STextBlock)
+							.Text(this, &SGameplayAttributeWidget::GetSelectedValueAsString)
+					]
+			]
 	];
 }
 
+void SGameplayAttributeWidget::OnCopyAttribute(FProperty* AttributeProperty)
+{
+	FGameplayAttribute Attribute(AttributeProperty);
+
+	if (Attribute.IsValid())
+	{
+		if (const TObjectPtr<UObject> OwnerObject = Attribute.GetAttributeSetClass()->ClassDefaultObject)
+		{
+			FString ExportedString;
+
+			FGameplayAttribute::StaticStruct()->ExportText(ExportedString, &Attribute, &Attribute, OwnerObject, 0, nullptr);
+			FPlatformApplicationMisc::ClipboardCopy(*ExportedString);
+		}
+	}
+}
+
+FGameplayAttribute AttributeTryImportTextFromClipboard()
+{
+	FString PastedText;
+	FPlatformApplicationMisc::ClipboardPaste(PastedText);
+
+	FOutputDeviceNull NullOut;
+	FGameplayAttribute Attribute;
+	FGameplayAttribute::StaticStruct()->ImportText(*PastedText, &Attribute, /*OwnerObject*/nullptr, 0, &NullOut, FGameplayAttribute::StaticStruct()->GetName(), /*bAllowNativeOverride*/true);
+
+	return Attribute;
+}
+
+bool SGameplayAttributeWidget::CanPaste() const
+{
+	const FGameplayAttribute Attribute = AttributeTryImportTextFromClipboard();
+
+	return Attribute.IsValid();
+}
+
+void SGameplayAttributeWidget::OnPasteAttribute()
+{
+	const FGameplayAttribute Attribute = AttributeTryImportTextFromClipboard();
+
+	if (Attribute.IsValid())
+	{
+		FScopedTransaction Transaction(LOCTEXT("GameplayAttributeWidget_PasteAttribute", "Paste Gameplay Attribute"));
+
+		SelectedProperty = Attribute.GetUProperty();
+		OnAttributeChanged.ExecuteIfBound(SelectedProperty);
+	}
+}
+
+void SGameplayAttributeWidget::OnMenu(const FPointerEvent& MouseEvent)
+{
+	FMenuBuilder MenuBuilder(/*bShouldCloseWindowAfterMenuSelection=*/ true, /*CommandList=*/ nullptr);
+
+	FGameplayAttribute SelectedAttribute(SelectedProperty);
+
+	auto IsValidAttribute = [](const FGameplayAttribute& Attribute)
+		{
+			return Attribute.IsValid();
+		};
+
+	MenuBuilder.AddMenuEntry(
+		LOCTEXT("GameplayAttribute_SearchForReferences", "Search For References"),
+		FText::Format(LOCTEXT("GameplayAttributeWidget_SearchForReferencesTooltip", "Find references to attribute {0}"), GetSelectedValueAsString()),
+		FSlateIcon(FAppStyle::GetAppStyleSetName(), "Icons.Search"),
+		FUIAction(FExecuteAction::CreateLambda([&]()
+			{
+				const FText AttributeName = GetSelectedValueAsString();
+				if (FEditorDelegates::OnOpenReferenceViewer.IsBound() && !AttributeName.IsEmpty())
+				{
+					TArray<FAssetIdentifier> AssetIdentifiers;
+					AssetIdentifiers.Emplace(FGameplayAttribute::StaticStruct(), *AttributeName.ToString());
+					FEditorDelegates::OnOpenReferenceViewer.Broadcast(AssetIdentifiers, FReferenceViewerParams());
+				}
+			}))
+	);
+
+	MenuBuilder.AddSeparator();
+
+	MenuBuilder.AddMenuEntry(
+		NSLOCTEXT("PropertyView", "CopyProperty", "Copy"),
+		FText::Format(LOCTEXT("GameplayAttributeWidget_CopyAttributeTooltip", "Copy attribute {0} to clipboard"), GetSelectedValueAsString()),
+		FSlateIcon(FAppStyle::GetAppStyleSetName(), "GenericCommands.Copy"),
+		FUIAction(FExecuteAction::CreateSP(this, &SGameplayAttributeWidget::OnCopyAttribute, SelectedProperty), FCanExecuteAction::CreateLambda(IsValidAttribute, SelectedAttribute)));
+
+	const FGameplayAttribute Attribute = AttributeTryImportTextFromClipboard();
+	FText FormattedPastedAttribute;
+
+	if (Attribute.GetUProperty())
+	{
+		UClass* Class = Attribute.GetUProperty()->GetOwnerClass();
+		FString PropertyString = FString::Printf(TEXT("%s.%s"), *Class->GetName(), *Attribute.GetUProperty()->GetName());
+		FormattedPastedAttribute = FText::FromString(PropertyString);
+	}
+	else
+	{
+		FormattedPastedAttribute = FText::FromString(TEXT("None"));
+	}
+
+	MenuBuilder.AddMenuEntry(
+		NSLOCTEXT("PropertyView", "PasteProperty", "Paste"),
+		FText::Format(LOCTEXT("GameplayAttributeWidget_PasteAttributeTooltip", "Paste attribute ({0}) from clipboard."), FormattedPastedAttribute),
+		FSlateIcon(FAppStyle::GetAppStyleSetName(), "GenericCommands.Paste"),
+		FUIAction(FExecuteAction::CreateSP(this, &SGameplayAttributeWidget::OnPasteAttribute), FCanExecuteAction::CreateSP(this, &SGameplayAttributeWidget::CanPaste)));
+
+	// Spawn context menu
+	FWidgetPath WidgetPath = MouseEvent.GetEventPath() != nullptr ? *MouseEvent.GetEventPath() : FWidgetPath();
+	FSlateApplication::Get().PushMenu(AsShared(), WidgetPath, MenuBuilder.MakeWidget(), MouseEvent.GetScreenSpacePosition(), FPopupTransitionEffect(FPopupTransitionEffect::ContextMenu));
+
+}
+
 void SGameplayAttributeWidget::OnAttributePicked(FProperty* InProperty)
 {
 	if (OnAttributeChanged.IsBound())
diff --git a/Engine/Plugins/Runtime/GameplayAbilities/Source/GameplayAbilitiesEditor/Private/SGameplayAttributeWidget.h b/Engine/Plugins/Runtime/GameplayAbilities/Source/GameplayAbilitiesEditor/Private/SGameplayAttributeWidget.h
--- a/Engine/Plugins/Runtime/GameplayAbilities/Source/GameplayAbilitiesEditor/Private/SGameplayAttributeWidget.h
+++ b/Engine/Plugins/Runtime/GameplayAbilities/Source/GameplayAbilitiesEditor/Private/SGameplayAttributeWidget.h
@@ -32,6 +32,13 @@ private:
 
 	FText GetSelectedValueAsString() const;
 
+	/**	Called when the widget is clicked using RMB */
+	void OnMenu(const FPointerEvent& MouseEvent);
+
+	void OnCopyAttribute(FProperty* AttributeProperty);
+	bool CanPaste() const;
+	void OnPasteAttribute();
+
 	/** Handles updates when the selected attribute changes */
 	void OnAttributePicked(FProperty* InProperty);
 
