Fixing a static analysis error in GAS.

diff --git a/Engine/Plugins/Runtime/GameplayAbilities/Source/GameplayAbilities/Private/Abilities/GameplayAbility.cpp b/Engine/Plugins/Runtime/GameplayAbilities/Source/GameplayAbilities/Private/Abilities/GameplayAbility.cpp
--- a/Engine/Plugins/Runtime/GameplayAbilities/Source/GameplayAbilities/Private/Abilities/GameplayAbility.cpp
+++ b/Engine/Plugins/Runtime/GameplayAbilities/Source/GameplayAbilities/Private/Abilities/GameplayAbility.cpp
@@ -1809,18 +1809,19 @@ UObject* UGameplayAbility::GetCurrentSourceObject() const
 
 FGameplayEffectContextHandle UGameplayAbility::MakeEffectContext(const FGameplayAbilitySpecHandle Handle, const FGameplayAbilityActorInfo *ActorInfo) const
 {
-	ensure(ActorInfo);
 	FGameplayEffectContextHandle Context = FGameplayEffectContextHandle(UAbilitySystemGlobals::Get().AllocGameplayEffectContext());
-	// By default use the owner and avatar as the instigator and causer
-	Context.AddInstigator(ActorInfo->OwnerActor.Get(), ActorInfo->AvatarActor.Get());
-
-	// add in the ability tracking here.
 	Context.SetAbility(this);
 
-	// Pass along the source object to the effect
-	if (FGameplayAbilitySpec* AbilitySpec = ActorInfo->AbilitySystemComponent->FindAbilitySpecFromHandle(Handle))
+	if (ensure(ActorInfo))
 	{
-		Context.AddSourceObject(AbilitySpec->SourceObject.Get());
+		// By default use the owner and avatar as the instigator and causer
+		Context.AddInstigator(ActorInfo->OwnerActor.Get(), ActorInfo->AvatarActor.Get());
+
+		// Pass along the source object to the effect
+		if (FGameplayAbilitySpec* AbilitySpec = ActorInfo->AbilitySystemComponent->FindAbilitySpecFromHandle(Handle))
+		{
+			Context.AddSourceObject(AbilitySpec->SourceObject.Get());
+		}
 	}
 
 	return Context;
