- Returned when trying to call a remote function from an invalid ability

diff --git a/Engine/Plugins/Runtime/GameplayAbilities/Source/GameplayAbilities/Private/Abilities/GameplayAbility.cpp b/Engine/Plugins/Runtime/GameplayAbilities/Source/GameplayAbilities/Private/Abilities/GameplayAbility.cpp
--- a/Engine/Plugins/Runtime/GameplayAbilities/Source/GameplayAbilities/Private/Abilities/GameplayAbility.cpp
+++ b/Engine/Plugins/Runtime/GameplayAbilities/Source/GameplayAbilities/Private/Abilities/GameplayAbility.cpp
@@ -127,6 +127,13 @@ int32 UGameplayAbility::GetFunctionCallspace(UFunction* Function, FFrame* Stack)
 
 bool UGameplayAbility::CallRemoteFunction(UFunction* Function, void* Parameters, FOutParmRec* OutParms, FFrame* Stack)
 {
+	// Make sure we are not invalid when being called
+	if (!IsValid(this))
+	{
+		ABILITY_LOG(Warning, TEXT("%s:CallRemoteFunction called while 'this' was invalid. Skipping the rest of the code."), *GetPathName());
+		return false;
+	}
+
 	check(!HasAnyFlags(RF_ClassDefaultObject));
 	check(GetOuter() != nullptr);
 
