GAS:  Cooldown code was incorrectly changed to use internal (instanced) logic, but was still being called on CDO externally.  Reverted those cases.

diff --git a/Engine/Plugins/Runtime/GameplayAbilities/Source/GameplayAbilities/Private/Abilities/GameplayAbility.cpp b/Engine/Plugins/Runtime/GameplayAbilities/Source/GameplayAbilities/Private/Abilities/GameplayAbility.cpp
--- a/Engine/Plugins/Runtime/GameplayAbilities/Source/GameplayAbilities/Private/Abilities/GameplayAbility.cpp
+++ b/Engine/Plugins/Runtime/GameplayAbilities/Source/GameplayAbilities/Private/Abilities/GameplayAbility.cpp
@@ -1032,7 +1032,8 @@ bool UGameplayAbility::CheckCost(const FGameplayAbilitySpecHandle Handle, const
 	UGameplayEffect* CostGE = GetCostGameplayEffect();
 	if (CostGE)
 	{
-		if (UAbilitySystemComponent* AbilitySystemComponent = GetAbilitySystemComponentFromActorInfo_Ensured())
+		UAbilitySystemComponent* AbilitySystemComponent = ActorInfo ? ActorInfo->AbilitySystemComponent.Get() : nullptr;
+		if (ensure(AbilitySystemComponent))
 		{
 			if (!AbilitySystemComponent->CanApplyAttributeModifiers(CostGE, GetAbilityLevel(Handle, ActorInfo), MakeEffectContext(Handle, ActorInfo)))
 			{
@@ -1106,7 +1107,8 @@ void UGameplayAbility::GetCooldownTimeRemainingAndDuration(FGameplayAbilitySpecH
 	const FGameplayTagContainer* CooldownTags = GetCooldownTags();
 	if (CooldownTags && CooldownTags->Num() > 0)
 	{
-		if (UAbilitySystemComponent* AbilitySystemComponent = GetAbilitySystemComponentFromActorInfo_Ensured())
+		UAbilitySystemComponent* AbilitySystemComponent = ActorInfo ? ActorInfo->AbilitySystemComponent.Get() : nullptr;
+		if (ensure(AbilitySystemComponent))
 		{
 			FGameplayEffectQuery const Query = FGameplayEffectQuery::MakeQuery_MatchAnyOwningTags(*CooldownTags);
 			TArray<TPair<float, float>> DurationAndTimeRemaining = AbilitySystemComponent->GetActiveEffectsTimeRemainingAndDuration(Query);
