GAS: Exposing Gameplay Ability IsActive to Blueprints.

diff --git a/Engine/Plugins/Runtime/GameplayAbilities/Source/GameplayAbilities/Private/AbilitySystemBlueprintLibrary.cpp b/Engine/Plugins/Runtime/GameplayAbilities/Source/GameplayAbilities/Private/AbilitySystemBlueprintLibrary.cpp
--- a/Engine/Plugins/Runtime/GameplayAbilities/Source/GameplayAbilities/Private/AbilitySystemBlueprintLibrary.cpp
+++ b/Engine/Plugins/Runtime/GameplayAbilities/Source/GameplayAbilities/Private/AbilitySystemBlueprintLibrary.cpp
@@ -1235,6 +1235,22 @@ const UGameplayAbility* UAbilitySystemBlueprintLibrary::GetGameplayAbilityFromSp
 	return AbilityInstance;
 }
 
+bool UAbilitySystemBlueprintLibrary::IsGameplayAbilityActive(const UGameplayAbility* GameplayAbility)
+{
+	if (!GameplayAbility)
+	{
+		UE_LOG(LogAbilitySystem, Error, TEXT("%hs passed in invalid (null) GameplayAbility"), __func__);
+		return false;
+	}
+	else if (!GameplayAbility->IsInstantiated())
+	{
+		UE_LOG(LogAbilitySystem, Error, TEXT("%hs passed a non-instantiated instance: %s"), __func__, *GetNameSafe(GameplayAbility));
+		return false;
+	}
+
+	return GameplayAbility->IsActive();
+}
+
 bool UAbilitySystemBlueprintLibrary::EqualEqual_GameplayAbilitySpecHandle(const FGameplayAbilitySpecHandle& A, const FGameplayAbilitySpecHandle& B)
 {
 	return A == B;
diff --git a/Engine/Plugins/Runtime/GameplayAbilities/Source/GameplayAbilities/Public/AbilitySystemBlueprintLibrary.h b/Engine/Plugins/Runtime/GameplayAbilities/Source/GameplayAbilities/Public/AbilitySystemBlueprintLibrary.h
--- a/Engine/Plugins/Runtime/GameplayAbilities/Source/GameplayAbilities/Public/AbilitySystemBlueprintLibrary.h
+++ b/Engine/Plugins/Runtime/GameplayAbilities/Source/GameplayAbilities/Public/AbilitySystemBlueprintLibrary.h
@@ -445,11 +445,15 @@ class GAMEPLAYABILITIES_API UAbilitySystemBlueprintLibrary : public UBlueprintFu
 	UFUNCTION(BlueprintPure, Category = "Ability|GameplayAbility")
 	static const UGameplayAbility* GetGameplayAbilityFromSpecHandle(UAbilitySystemComponent* AbilitySystem, const FGameplayAbilitySpecHandle& AbilitySpecHandle, bool& bIsInstance);
 
+	/** Returns true if the passed-in Gameplay Ability instance is active (activated and not yet ended). */
+	UFUNCTION(BlueprintPure, Category = "Ability|GameplayAbility", meta=(DisplayName="Is Active",DefaultToSelf=GameplayAbility))
+	static bool IsGameplayAbilityActive(const UGameplayAbility* GameplayAbility);
+
 	/** Equality operator for two Gameplay Ability Spec Handles */
-	UFUNCTION(BlueprintPure, Category = "Ability|GameplayEffect", meta = (DisplayName = "Equal (Gameplay Ability Spec Handle)", CompactNodeTitle = "==", ScriptOperator = "=="))
+	UFUNCTION(BlueprintPure, Category = "Ability|GameplayAbility", meta = (DisplayName = "Equal (Gameplay Ability Spec Handle)", CompactNodeTitle = "==", ScriptOperator = "=="))
 	static bool EqualEqual_GameplayAbilitySpecHandle(const FGameplayAbilitySpecHandle& A, const FGameplayAbilitySpecHandle& B);
 
 	/** Inequality operator for two Gameplay Ability Spec Handles */
-	UFUNCTION(BlueprintPure, Category = "Ability|GameplayEffect", meta = (DisplayName = "Not Equal (Gameplay Ability Spec Handle)", CompactNodeTitle = "!=", ScriptOperator = "!="))
+	UFUNCTION(BlueprintPure, Category = "Ability|GameplayAbility", meta = (DisplayName = "Not Equal (Gameplay Ability Spec Handle)", CompactNodeTitle = "!=", ScriptOperator = "!="))
 	static bool NotEqual_GameplayAbilitySpecHandle(const FGameplayAbilitySpecHandle& A, const FGameplayAbilitySpecHandle& B);
 };
