GAS Hotfix:  Server Initiated PredictionKeys were being sent down to the client, but it takes a spot in a circular buffer and thus it's possible it stomps an actual client FPredictionKey.

diff --git a/Engine/Plugins/Runtime/GameplayAbilities/Source/GameplayAbilities/Private/GameplayPrediction.cpp b/Engine/Plugins/Runtime/GameplayAbilities/Source/GameplayAbilities/Private/GameplayPrediction.cpp
--- a/Engine/Plugins/Runtime/GameplayAbilities/Source/GameplayAbilities/Private/GameplayPrediction.cpp
+++ b/Engine/Plugins/Runtime/GameplayAbilities/Source/GameplayAbilities/Private/GameplayPrediction.cpp
@@ -32,11 +32,11 @@ namespace UE::AbilitySystem::Private
 	FAutoConsoleVariableRef CVarDependentChainBehavior(TEXT("AbilitySystem.PredictionKey.DepChainBehavior"), CVarDependentChainBehaviorValue,
 		TEXT("How do we handle dependency key chains? Bitmask: 0 = Old Accept/Rejected Implies Newer Accepted/Rejected. 0x1 = Newer Accepted also implies Older Accepted. 0x2 = Old Accepted Does NOT imply Newer Accepted"));
 
-	// Prior to UE5.5, we used to allow Server Initiated Replication Keys to be sent to the client as 'acknowledged' but that hardly makes sense.
+	// Prior to UE5.4.2, we used to allow Server Initiated Replication Keys to be sent to the client as 'acknowledged' but that hardly makes sense.
 	// Unfortunately, it also causes key hash collisions since there is limited space based on the key value.
 	int32 CVarReplicateServerKeysAsAcknowledgedValue = 0;
 	FAutoConsoleVariableRef CVarReplicateServerKeysAsAcknowledged(TEXT("AbilitySystem.PredictionKey.RepServerKeysAsAcknowledged"), CVarReplicateServerKeysAsAcknowledgedValue,
-		TEXT("Do we send server initiated keys as acknowledged to the client? Default false. Was true prior to UE5.5."));
+		TEXT("Do we send server initiated keys as acknowledged to the client? Default false. Was true prior to UE5.4.2."));
 
 	/**
 	 * Given an FProperty Link (such as a UFunction's Properties, or a UStruct's Properties), find and return any linked FPredictionKeys.
@@ -587,7 +587,7 @@ void FReplicatedPredictionKeyItem::OnRep(const FReplicatedPredictionKeyMap& InAr
 	// initiates a key that matches a local value.  Then we catch-up to a local value, even though the server was not specifically acknowledging it.
 	if (PredictionKey.bIsServerInitiated)
 	{
-		UE_LOG(LogPredictionKey, Warning, TEXT("FReplicatedPredictionKeyItem::OnRep received Server Initiated Key %s (which likely stomped a local key due to keymap hash collisions)."), *PredictionKey.ToString());
+		UE_LOG(LogPredictionKey, Warning, TEXT("FReplicatedPredictionKeyItem::OnRep received Server Initiated Key %s (which likely stomped a local key due to keymap hash collisions). This shouldn't happen if CVarReplicateServerKeysAsAcknowledged is 0."), *PredictionKey.ToString());
 		return;
 	}
 
