GAS:  Adding "ShowDebug AbilitySystem" to the console auto-complete list

diff --git a/Engine/Plugins/Runtime/GameplayAbilities/Config/Input.ini b/Engine/Plugins/Runtime/GameplayAbilities/Config/Input.ini
new file mode 100644
--- /dev/null
+++ b/Engine/Plugins/Runtime/GameplayAbilities/Config/Input.ini
@@ -0,0 +1,2 @@
+[/Script/EngineSettings.ConsoleSettings]
++ManualAutoCompleteList=(Command="ShowDebug AbilitySystem",Desc="Toggles debug information about the Gameplay Ability System"
\ No newline at end of file
