Deprecating FGameplayAbilitySpec::ActivationInfo since it only makes sense on NonInstanced Gameplay Abilities (which have been deprecated).  If you want that functionality, grab the instance and look at GetCurrentActivationInfo().

diff --git a/Engine/Plugins/Runtime/GameplayAbilities/Source/GameplayAbilities/Private/AbilitySystemCheatManagerExtension.cpp b/Engine/Plugins/Runtime/GameplayAbilities/Source/GameplayAbilities/Private/AbilitySystemCheatManagerExtension.cpp
--- a/Engine/Plugins/Runtime/GameplayAbilities/Source/GameplayAbilities/Private/AbilitySystemCheatManagerExtension.cpp
+++ b/Engine/Plugins/Runtime/GameplayAbilities/Source/GameplayAbilities/Private/AbilitySystemCheatManagerExtension.cpp
@@ -378,7 +378,7 @@ void UAbilitySystemCheatManagerExtension::AbilityCancel(const FString& PartialNa
 					if (Instance)
 					{
 						UE_LOG(LogConsoleResponse, Log, TEXT("%s (%s): Cancelling (instanced) %s"), *PC->GetName(), *ASC->GetName(), *Instance->GetName());
-						Instance->CancelAbility(GASpec.Handle, ASC->AbilityActorInfo.Get(), GASpec.ActivationInfo, true);
+						Instance->CancelAbility(GASpec.Handle, ASC->AbilityActorInfo.Get(), Instance->GetCurrentActivationInfoRef(), true);
 						bFound = true;
 					}
 				}
@@ -386,7 +386,9 @@ void UAbilitySystemCheatManagerExtension::AbilityCancel(const FString& PartialNa
 			else
 			{
 				UE_LOG(LogConsoleResponse, Log, TEXT("%s (%s): Cancelling (non-instanced) %s"), *PC->GetName(), *ASC->GetName(), *GASpec.Ability->GetName());
+PRAGMA_DISABLE_DEPRECATION_WARNINGS
 				GASpec.Ability->CancelAbility(GASpec.Handle, ASC->AbilityActorInfo.Get(), GASpec.ActivationInfo, true);
+PRAGMA_ENABLE_DEPRECATION_WARNINGS
 				bFound = true;
 			}
 		}
diff --git a/Engine/Plugins/Runtime/GameplayAbilities/Source/GameplayAbilities/Private/AbilitySystemComponent.cpp b/Engine/Plugins/Runtime/GameplayAbilities/Source/GameplayAbilities/Private/AbilitySystemComponent.cpp
--- a/Engine/Plugins/Runtime/GameplayAbilities/Source/GameplayAbilities/Private/AbilitySystemComponent.cpp
+++ b/Engine/Plugins/Runtime/GameplayAbilities/Source/GameplayAbilities/Private/AbilitySystemComponent.cpp
@@ -2822,8 +2822,10 @@ void UAbilitySystemComponent::Debug_Internal(FAbilitySystemComponentDebugInfo& I
 				}
 			}
 
+PRAGMA_DISABLE_DEPRECATION_WARNINGS
 			FString InputPressedStr = AbilitySpec.InputPressed ? TEXT("(InputPressed)") : TEXT("");
 			FString ActivationModeStr = AbilitySpec.IsActive() ? UEnum::GetValueAsString(TEXT("GameplayAbilities.EGameplayAbilityActivationMode"), AbilitySpec.ActivationInfo.ActivationMode) : TEXT("");
+PRAGMA_ENABLE_DEPRECATION_WARNINGS
 
 			if (Info.Canvas)
 			{
diff --git a/Engine/Plugins/Runtime/GameplayAbilities/Source/GameplayAbilities/Private/AbilitySystemComponent_Abilities.cpp b/Engine/Plugins/Runtime/GameplayAbilities/Source/GameplayAbilities/Private/AbilitySystemComponent_Abilities.cpp
--- a/Engine/Plugins/Runtime/GameplayAbilities/Source/GameplayAbilities/Private/AbilitySystemComponent_Abilities.cpp
+++ b/Engine/Plugins/Runtime/GameplayAbilities/Source/GameplayAbilities/Private/AbilitySystemComponent_Abilities.cpp
@@ -687,6 +687,7 @@ void UAbilitySystemComponent::OnRemoveAbility(FGameplayAbilitySpec& Spec)
 					return;
 				}
 			}
+PRAGMA_DISABLE_DEPRECATION_WARNINGS
 			else if (ensureMsgf(Spec.Ability->GetInstancingPolicy() == EGameplayAbilityInstancingPolicy::NonInstanced, TEXT("We should never have an instanced Gameplay Ability that is still active by this point. All instances should have EndAbility called just before here.")))
 			{
 				// Seems like it should be cancelled, but we're just following the existing pattern (could be due to functionality from OnRep)
@@ -694,6 +695,7 @@ void UAbilitySystemComponent::OnRemoveAbility(FGameplayAbilitySpec& Spec)
 				constexpr bool bWasCancelled = false;
 				Spec.Ability->EndAbility(Spec.Handle, AbilityActorInfo.Get(), Spec.ActivationInfo, bReplicateEndAbility, bWasCancelled);
 			}
+PRAGMA_ENABLE_DEPRECATION_WARNINGS
 		}
 
 		Spec.Ability->OnRemoveAbility(AbilityActorInfo.Get(), Spec);
@@ -1324,7 +1326,7 @@ void UAbilitySystemComponent::CancelAbilitySpec(FGameplayAbilitySpec& Spec, UGam
 		{
 			if (InstanceAbility && Ignore != InstanceAbility)
 			{
-				InstanceAbility->CancelAbility(Spec.Handle, ActorInfo, InstanceAbility->GetCurrentActivationInfo(), true);
+				InstanceAbility->CancelAbility(Spec.Handle, ActorInfo, InstanceAbility->GetCurrentActivationInfoRef(), true);
 			}
 		}
 	}
@@ -1472,9 +1474,12 @@ void UAbilitySystemComponent::OnRep_ActivateAbilities()
 
 	CheckForClearedAbilities();
 
-	// Try to run any pending activations that couldn't run before. If they don't work now, kill them
+	// Make a copy in case a pending ability alters the array while iterating
+	TArray<FPendingAbilityInfo> PendingCopy = PendingServerActivatedAbilities;
+	PendingServerActivatedAbilities.Empty();
 
-	for (const FPendingAbilityInfo& PendingAbilityInfo : PendingServerActivatedAbilities)
+	// Try to run any pending activations that couldn't run before. If they don't work now, kill them
+	for (const FPendingAbilityInfo& PendingAbilityInfo : PendingCopy)
 	{
 		if (PendingAbilityInfo.bPartiallyActivated)
 		{
@@ -1485,8 +1490,8 @@ void UAbilitySystemComponent::OnRep_ActivateAbilities()
 			ClientTryActivateAbility(PendingAbilityInfo.Handle);
 		}
 	}
+	ensureMsgf(PendingServerActivatedAbilities.IsEmpty(), TEXT("Execution of Pending Abilities caused %d more Pending Abilities (ignoring them)"), PendingServerActivatedAbilities.Num());
 	PendingServerActivatedAbilities.Empty();
-
 }
 
 void UAbilitySystemComponent::GetActivatableGameplayAbilitySpecsByAllMatchingTags(const FGameplayTagContainer& GameplayTagContainer, TArray < struct FGameplayAbilitySpec* >& MatchingGameplayAbilities, bool bOnlyAbilitiesThatSatisfyTagRequirements) const
@@ -1784,9 +1789,10 @@ bool UAbilitySystemComponent::InternalTryActivateAbility(FGameplayAbilitySpecHan
 				UE_LOG(LogAbilitySystem, Verbose, TEXT("%s: Ending %s prematurely to retrigger."), *GetNameSafe(GetOwner()), *Ability->GetName());
 				UE_VLOG(GetOwner(), VLogAbilitySystem, Verbose, TEXT("Ending %s prematurely to retrigger."), *Ability->GetName());
 
-				bool bReplicateEndAbility = true;
-				bool bWasCancelled = false;
-				InstancedAbility->EndAbility(Handle, ActorInfo, Spec->ActivationInfo, bReplicateEndAbility, bWasCancelled);
+				constexpr bool bReplicateEndAbility = true;
+				constexpr bool bWasCancelled = false;
+				const FGameplayAbilityActivationInfo& ActivationInfo = InstancedAbility->GetCurrentActivationInfoRef();
+				InstancedAbility->EndAbility(Handle, ActorInfo, ActivationInfo, bReplicateEndAbility, bWasCancelled);
 			}
 			else
 			{
@@ -1803,9 +1809,15 @@ bool UAbilitySystemComponent::InternalTryActivateAbility(FGameplayAbilitySpecHan
 		return false;
 	}
 
-	// Setup a fresh ActivationInfo for this AbilitySpec.
-	Spec->ActivationInfo = FGameplayAbilityActivationInfo(ActorInfo->OwnerActor.Get());
-	FGameplayAbilityActivationInfo &ActivationInfo = Spec->ActivationInfo;
+PRAGMA_DISABLE_DEPRECATION_WARNINGS
+	// We have deprecated NonInstanced and Spec.ActivationInfo but keep backwards compatibility
+	const bool bNonInstanced = Ability->GetInstancingPolicy() == EGameplayAbilityInstancingPolicy::NonInstanced;
+	FGameplayAbilityActivationInfo NewActivationInfo;
+	FGameplayAbilityActivationInfo& ActivationInfo = bNonInstanced ? Spec->ActivationInfo : NewActivationInfo;
+PRAGMA_ENABLE_DEPRECATION_WARNINGS
+
+	// Setup a fresh ActivationInfo (possibly overwriting the Spec's ActivationInfo if non-instanced)
+	ActivationInfo = FGameplayAbilityActivationInfo(ActorInfo->OwnerActor.Get());
 
 	// If we are the server or this is local only
 	if (Ability->GetNetExecutionPolicy() == EGameplayAbilityNetExecutionPolicy::LocalOnly || (NetMode == ROLE_Authority))
@@ -2241,11 +2253,13 @@ void UAbilitySystemComponent::ClientActivateAbilityFailed_Implementation(FGamepl
 	}
 
 
+PRAGMA_DISABLE_DEPRECATION_WARNINGS
 	// The ability should be either confirmed or rejected by the time we get here
 	if (Spec->ActivationInfo.GetActivationPredictionKey().Current == PredictionKey)
 	{
 		Spec->ActivationInfo.SetActivationRejected();
 	}
+PRAGMA_ENABLE_DEPRECATION_WARNINGS
 
 	TArray<UGameplayAbility*> Instances = Spec->GetAbilityInstances();
 	for (UGameplayAbility* Ability : Instances)
@@ -2263,11 +2277,24 @@ void UAbilitySystemComponent::OnClientActivateAbilityCaughtUp(FGameplayAbilitySp
 	FGameplayAbilitySpec* Spec = FindAbilitySpecFromHandle(Handle);
 	if (Spec && Spec->IsActive())
 	{
+PRAGMA_DISABLE_DEPRECATION_WARNINGS
+		const FGameplayAbilityActivationInfo* ActivationInfo = (Spec->ActivationInfo.GetActivationPredictionKey().Current == PredictionKey) ? &Spec->ActivationInfo : nullptr;
+PRAGMA_ENABLE_DEPRECATION_WARNINGS
+
+		for (const UGameplayAbility* Ability : Spec->GetAbilityInstances())
+		{
+			if (Ability->CurrentActivationInfo.GetActivationPredictionKey().Current == PredictionKey)
+			{
+				ActivationInfo = &Ability->CurrentActivationInfo;
+				break;
+			}
+		}
+
 		// The ability should be either confirmed or rejected by the time we get here
-		if (Spec->ActivationInfo.ActivationMode == EGameplayAbilityActivationMode::Predicting && Spec->ActivationInfo.GetActivationPredictionKey().Current == PredictionKey)
+		if (ActivationInfo && ActivationInfo->ActivationMode == EGameplayAbilityActivationMode::Predicting)
 		{
 			// It is possible to have this happen under bad network conditions. (Reliable Confirm/Reject RPC is lost, but separate property bunch makes it through before the reliable resend happens)
-			ABILITY_LOG(Display, TEXT("UAbilitySystemComponent::OnClientActivateAbilityCaughtUp. Ability %s caught up to PredictionKey %d but instance is still active and in predicting state."), *GetNameSafe(Spec->Ability), PredictionKey);
+			ABILITY_LOG(Display, TEXT("%hs. Ability %s caught up to PredictionKey %d but instance is still active and in predicting state."), __func__, *GetNameSafe(Spec->Ability), PredictionKey);
 		}
 	}
 }
@@ -2299,18 +2326,24 @@ void UAbilitySystemComponent::ClientActivateAbilitySucceedWithEventData_Implemen
 	check(AbilityToActivate);
 	ensure(AbilityActorInfo.IsValid());
 
-	Spec->ActivationInfo.SetActivationConfirmed();
-
 	UE_LOG(LogAbilitySystem, Verbose, TEXT("%s: Server Confirmed [%s] %s. PredictionKey: %s"), *GetNameSafe(GetOwner()), *Handle.ToString(), *GetNameSafe(AbilityToActivate), *PredictionKey.ToString());
 	UE_VLOG(GetOwner(), VLogAbilitySystem, Verbose, TEXT("Server Confirmed [%s] %s. PredictionKey: %s"), *Handle.ToString(), *GetNameSafe(AbilityToActivate), *PredictionKey.ToString());
 
 	// Fixme: We need a better way to link up/reconcile predictive replicated abilities. It would be ideal if we could predictively spawn an
 	// ability and then replace/link it with the server spawned one once the server has confirmed it.
 
-	if (AbilityToActivate->NetExecutionPolicy == EGameplayAbilityNetExecutionPolicy::LocalPredicted)
+PRAGMA_DISABLE_DEPRECATION_WARNINGS
+	const bool bNonInstanced = AbilityToActivate->GetInstancingPolicy() == EGameplayAbilityInstancingPolicy::NonInstanced;
+	FGameplayAbilityActivationInfo NewActivationInfo{ AbilityActorInfo->OwnerActor.Get() };
+	FGameplayAbilityActivationInfo& ActivationInfo = bNonInstanced ? Spec->ActivationInfo : NewActivationInfo;
+PRAGMA_ENABLE_DEPRECATION_WARNINGS
+
+	const bool bLocallyPredicted = AbilityToActivate->NetExecutionPolicy == EGameplayAbilityNetExecutionPolicy::LocalPredicted;
+	if (bLocallyPredicted)
 	{
-		if (AbilityToActivate->GetInstancingPolicy() == EGameplayAbilityInstancingPolicy::NonInstanced)
+		if (bNonInstanced)
 		{
+			ActivationInfo.SetActivationConfirmed();
 			// AbilityToActivate->ConfirmActivateSucceed(); // This doesn't do anything for non instanced
 		}
 		else
@@ -2330,25 +2363,25 @@ void UAbilitySystemComponent::ClientActivateAbilitySucceedWithEventData_Implemen
 
 			if (!found)
 			{
-				ABILITY_LOG(Verbose, TEXT("Ability %s was confirmed by server but no longer exists on client (replication key: %d"), *AbilityToActivate->GetName(), PredictionKey.Current);
+				ABILITY_LOG(Verbose, TEXT("Ability %s was confirmed by server but no longer exists on client (replication key: %s)"), *AbilityToActivate->GetName(), *PredictionKey.ToString());
 			}
 		}
 	}
 	else
 	{
-		// We haven't already executed this ability at all, so kick it off.
 
+		// We haven't already executed this ability at all, so kick it off.
 		if (PredictionKey.bIsServerInitiated)
 		{
 			// We have an active server key, set our key equal to it
-			Spec->ActivationInfo.ServerSetActivationPredictionKey(PredictionKey);
+			ActivationInfo.ServerSetActivationPredictionKey(PredictionKey);
 		}
 		
 		if (AbilityToActivate->GetInstancingPolicy() == EGameplayAbilityInstancingPolicy::InstancedPerExecution)
 		{
 			// Need to instantiate this in order to execute
 			UGameplayAbility* InstancedAbility = CreateNewInstanceOfAbility(*Spec, AbilityToActivate);
-			InstancedAbility->CallActivateAbility(Handle, AbilityActorInfo.Get(), Spec->ActivationInfo, nullptr, TriggerEventData.EventTag.IsValid() ?  &TriggerEventData : nullptr);
+			InstancedAbility->CallActivateAbility(Handle, AbilityActorInfo.Get(), ActivationInfo, nullptr, TriggerEventData.EventTag.IsValid() ?  &TriggerEventData : nullptr);
 		}
 		else if (AbilityToActivate->GetInstancingPolicy() != EGameplayAbilityInstancingPolicy::NonInstanced)
 		{
@@ -2359,11 +2392,11 @@ void UAbilitySystemComponent::ClientActivateAbilitySucceedWithEventData_Implemen
 				ABILITY_LOG(Warning, TEXT("Ability %s cannot be activated on the client because it's missing a primary instance!"), *AbilityToActivate->GetName());
 				return;
 			}
-			InstancedAbility->CallActivateAbility(Handle, AbilityActorInfo.Get(), Spec->ActivationInfo, nullptr, TriggerEventData.EventTag.IsValid() ? &TriggerEventData : nullptr);
+			InstancedAbility->CallActivateAbility(Handle, AbilityActorInfo.Get(), ActivationInfo, nullptr, TriggerEventData.EventTag.IsValid() ? &TriggerEventData : nullptr);
 		}
 		else
 		{
-			AbilityToActivate->CallActivateAbility(Handle, AbilityActorInfo.Get(), Spec->ActivationInfo, nullptr, TriggerEventData.EventTag.IsValid() ? &TriggerEventData : nullptr);
+			AbilityToActivate->CallActivateAbility(Handle, AbilityActorInfo.Get(), ActivationInfo, nullptr, TriggerEventData.EventTag.IsValid() ? &TriggerEventData : nullptr);
 		}
 	}
 }
@@ -2721,8 +2754,12 @@ void UAbilitySystemComponent::AbilityLocalInputPressed(int32 InputID)
 
 					AbilitySpecInputPressed(Spec);
 
+PRAGMA_DISABLE_DEPRECATION_WARNINGS
+					// Fixing this up to use the instance activation, but this function should be deprecated as it cannot work with InstancedPerExecution
+					const FGameplayAbilityActivationInfo& ActivationInfo = Spec.GetPrimaryInstance() ? Spec.GetPrimaryInstance()->GetCurrentActivationInfoRef() : Spec.ActivationInfo;
+PRAGMA_ENABLE_DEPRECATION_WARNINGS
 					// Invoke the InputPressed event. This is not replicated here. If someone is listening, they may replicate the InputPressed event to the server.
-					InvokeReplicatedEvent(EAbilityGenericReplicatedEvent::InputPressed, Spec.Handle, Spec.ActivationInfo.GetActivationPredictionKey());					
+					InvokeReplicatedEvent(EAbilityGenericReplicatedEvent::InputPressed, Spec.Handle, ActivationInfo.GetActivationPredictionKey());					
 				}
 				else
 				{
@@ -2750,8 +2787,12 @@ void UAbilitySystemComponent::AbilityLocalInputReleased(int32 InputID)
 				}
 
 				AbilitySpecInputReleased(Spec);
-				
-				InvokeReplicatedEvent(EAbilityGenericReplicatedEvent::InputReleased, Spec.Handle, Spec.ActivationInfo.GetActivationPredictionKey());
+
+PRAGMA_DISABLE_DEPRECATION_WARNINGS
+				// Fixing this up to use the instance activation, but this function should be deprecated as it cannot work with InstancedPerExecution
+				const FGameplayAbilityActivationInfo& ActivationInfo = Spec.GetPrimaryInstance() ? Spec.GetPrimaryInstance()->GetCurrentActivationInfoRef() : Spec.ActivationInfo;
+PRAGMA_ENABLE_DEPRECATION_WARNINGS
+				InvokeReplicatedEvent(EAbilityGenericReplicatedEvent::InputReleased, Spec.Handle, ActivationInfo.GetActivationPredictionKey());
 			}
 		}
 	}
@@ -2804,14 +2845,16 @@ void UAbilitySystemComponent::AbilitySpecInputPressed(FGameplayAbilitySpec& Spec
 		// The ability is active, so just pipe the input event to it
 		if (Spec.Ability->GetInstancingPolicy() == EGameplayAbilityInstancingPolicy::NonInstanced)
 		{
+PRAGMA_DISABLE_DEPRECATION_WARNINGS
 			Spec.Ability->InputPressed(Spec.Handle, AbilityActorInfo.Get(), Spec.ActivationInfo);
+PRAGMA_ENABLE_DEPRECATION_WARNINGS
 		}
 		else
 		{
 			TArray<UGameplayAbility*> Instances = Spec.GetAbilityInstances();
 			for (UGameplayAbility* Instance : Instances)
 			{
-				Instance->InputPressed(Spec.Handle, AbilityActorInfo.Get(), Spec.ActivationInfo);
+				Instance->InputPressed(Spec.Handle, AbilityActorInfo.Get(), Instance->CurrentActivationInfo);
 			}
 		}
 	}
@@ -2825,14 +2868,16 @@ void UAbilitySystemComponent::AbilitySpecInputReleased(FGameplayAbilitySpec& Spe
 		// The ability is active, so just pipe the input event to it
 		if (Spec.Ability->GetInstancingPolicy() == EGameplayAbilityInstancingPolicy::NonInstanced)
 		{
+PRAGMA_DISABLE_DEPRECATION_WARNINGS
 			Spec.Ability->InputReleased(Spec.Handle, AbilityActorInfo.Get(), Spec.ActivationInfo);
+PRAGMA_ENABLE_DEPRECATION_WARNINGS
 		}
 		else
 		{
 			TArray<UGameplayAbility*> Instances = Spec.GetAbilityInstances();
 			for (UGameplayAbility* Instance : Instances)
 			{
-				Instance->InputReleased(Spec.Handle, AbilityActorInfo.Get(), Spec.ActivationInfo);
+				Instance->InputReleased(Spec.Handle, AbilityActorInfo.Get(), Instance->CurrentActivationInfo);
 			}
 		}
 	}
diff --git a/Engine/Plugins/Runtime/GameplayAbilities/Source/GameplayAbilities/Private/Tests/AbilitySystemComponentTests.cpp b/Engine/Plugins/Runtime/GameplayAbilities/Source/GameplayAbilities/Private/Tests/AbilitySystemComponentTests.cpp
--- a/Engine/Plugins/Runtime/GameplayAbilities/Source/GameplayAbilities/Private/Tests/AbilitySystemComponentTests.cpp
+++ b/Engine/Plugins/Runtime/GameplayAbilities/Source/GameplayAbilities/Private/Tests/AbilitySystemComponentTests.cpp
@@ -166,9 +166,6 @@ public: // the tests
 			SourceComponent->CancelAbilityHandle(GivenAbilitySpecHandle);
 			Test->TestTrue(TEXT(" AbilityEnded (after CancelAbilityHandle)"), TestCallbacks.bReceiveAbilityEnded);
 			Test->TestFalse(TEXT(" AbilitySpec.IsActive() (after CancelAbilityHandle)"), AbilitySpec->IsActive());
-
-			// We're the authority
-			Test->TestEqual(TEXT(" ActivationInfo.ActivationMode"), AbilitySpec->ActivationInfo.ActivationMode, EGameplayAbilityActivationMode::Authority);
 		}
 	}
 
@@ -207,9 +204,6 @@ public: // the tests
 			Test->TestFalse(TEXT(" AbilityCommitted after TryActivateAbility (using FGameplayAbilitySpecHandle)"), TestCallbacks.bReceivedAbilityCommitted);
 			Test->TestFalse(TEXT(" AbilityEnded (prematurely) after TryActivateAbility (using FGameplayAbilitySpecHandle)"), TestCallbacks.bReceiveAbilityEnded);
 			Test->TestTrue(TEXT(" AbilityFailed after TryActivateAbility (with an Ability that should fail)"), TestCallbacks.bReceiveAbilityFailed);
-
-			// Even though we've been rejected, the ActivationMode should still say we're the authority
-			Test->TestEqual(TEXT(" ActivationInfo.ActivationMode"), AbilitySpec->ActivationInfo.ActivationMode, EGameplayAbilityActivationMode::Authority);
 		}
 	}
 
diff --git a/Engine/Plugins/Runtime/GameplayAbilities/Source/GameplayAbilities/Public/GameplayAbilitySpec.h b/Engine/Plugins/Runtime/GameplayAbilities/Source/GameplayAbilities/Public/GameplayAbilitySpec.h
--- a/Engine/Plugins/Runtime/GameplayAbilities/Source/GameplayAbilities/Public/GameplayAbilitySpec.h
+++ b/Engine/Plugins/Runtime/GameplayAbilities/Source/GameplayAbilities/Public/GameplayAbilitySpec.h
@@ -231,6 +231,7 @@ PRAGMA_ENABLE_DEPRECATION_WARNINGS
 	TSharedPtr<FGameplayEventData> GameplayEventData = nullptr;
 
 	/** Activation state of this ability. This is not replicated since it needs to be overwritten locally on clients during prediction. */
+	UE_DEPRECATED(5.5, "ActivationInfo on the Spec only applies to NonInstanced abilities (which are now deprecated; instanced abilities have their own per-instance CurrentActivationInfo)")
 	UPROPERTY(NotReplicated)
 	FGameplayAbilityActivationInfo	ActivationInfo;
 
