GAS:  PR #11694: Use CurrentValue instead of BaseValue when dispatching AttributeChangeDelegates in replication case

This solves a case where you can listen for the AttributeValueChangeDelegate on a Simulated Proxy and receive the Base value rather than properly aggregated value.

diff --git a/Engine/Plugins/Runtime/GameplayAbilities/Source/GameplayAbilities/Private/GameplayEffect.cpp b/Engine/Plugins/Runtime/GameplayAbilities/Source/GameplayAbilities/Private/GameplayEffect.cpp
--- a/Engine/Plugins/Runtime/GameplayAbilities/Source/GameplayAbilities/Private/GameplayEffect.cpp
+++ b/Engine/Plugins/Runtime/GameplayAbilities/Source/GameplayAbilities/Private/GameplayEffect.cpp
@@ -3505,7 +3505,12 @@ bool FActiveGameplayEffectsContainer::ShouldUseMinimalReplication()
 	return IsNetAuthority() && (Owner->ReplicationMode == EGameplayEffectReplicationMode::Minimal || Owner->ReplicationMode == EGameplayEffectReplicationMode::Mixed);
 }
 
-void FActiveGameplayEffectsContainer::SetBaseAttributeValueFromReplication(const FGameplayAttribute& Attribute, float ServerBaseValue, float OldBaseValue)
+void FActiveGameplayEffectsContainer::SetBaseAttributeValueFromReplication(const FGameplayAttribute& Attribute, float NewBaseValue, float OldBaseValue)
+{
+	SetBaseAttributeValueFromReplication(Attribute, FGameplayAttributeData(NewBaseValue), FGameplayAttributeData(OldBaseValue));
+}
+
+void FActiveGameplayEffectsContainer::SetBaseAttributeValueFromReplication(const FGameplayAttribute& Attribute, const FGameplayAttributeData& NewValue, const FGameplayAttributeData& OldValue)
 {
 	FAggregatorRef* RefPtr = AttributeAggregatorMap.Find(Attribute);
 	if (RefPtr && RefPtr->Get())
@@ -3513,6 +3518,9 @@ void FActiveGameplayEffectsContainer::SetBaseAttributeValueFromReplication(const
 		FAggregator* Aggregator = RefPtr->Get();
 		if (FGameplayAttribute::IsGameplayAttributeDataProperty(Attribute.GetUProperty()))
 		{
+			const float ServerBaseValue = NewValue.GetBaseValue();
+			const float OldBaseValue = OldValue.GetBaseValue();
+			
 			// Reset to the server's old value
 			constexpr bool bDoNotExecuteCallbacksValue = false;
 			Aggregator->SetBaseValue(OldBaseValue, bDoNotExecuteCallbacksValue);
@@ -3538,7 +3546,7 @@ void FActiveGameplayEffectsContainer::SetBaseAttributeValueFromReplication(const
 		PRAGMA_DISABLE_DEPRECATION_WARNINGS
 		if (FOnGameplayAttributeChange* LegacyDelegate = AttributeChangeDelegates.Find(Attribute))
 		{
-			LegacyDelegate->Broadcast(ServerBaseValue, nullptr);
+			LegacyDelegate->Broadcast(NewValue.GetCurrentValue(), nullptr);
 		}
 		PRAGMA_ENABLE_DEPRECATION_WARNINGS
 
@@ -3546,8 +3554,8 @@ void FActiveGameplayEffectsContainer::SetBaseAttributeValueFromReplication(const
 		{
 			FOnAttributeChangeData CallbackData;
 			CallbackData.Attribute = Attribute;
-			CallbackData.NewValue = ServerBaseValue;
-			CallbackData.OldValue = OldBaseValue;
+			CallbackData.NewValue = NewValue.GetCurrentValue();
+			CallbackData.OldValue = OldValue.GetCurrentValue();
 			CallbackData.GEModData = nullptr;
 
 			Delegate->Broadcast(CallbackData);
diff --git a/Engine/Plugins/Runtime/GameplayAbilities/Source/GameplayAbilities/Public/AbilitySystemComponent.h b/Engine/Plugins/Runtime/GameplayAbilities/Source/GameplayAbilities/Public/AbilitySystemComponent.h
--- a/Engine/Plugins/Runtime/GameplayAbilities/Source/GameplayAbilities/Public/AbilitySystemComponent.h
+++ b/Engine/Plugins/Runtime/GameplayAbilities/Source/GameplayAbilities/Public/AbilitySystemComponent.h
@@ -809,13 +809,13 @@ class GAMEPLAYABILITIES_API UAbilitySystemComponent : public UGameplayTasksCompo
 	/** Call from OnRep functions to set the attribute base value on the client */
 	void SetBaseAttributeValueFromReplication(const FGameplayAttribute& Attribute, float NewValue, float OldValue)
 	{
-		ActiveGameplayEffects.SetBaseAttributeValueFromReplication(Attribute, NewValue, OldValue);
+		ActiveGameplayEffects.SetBaseAttributeValueFromReplication(Attribute, FGameplayAttributeData(NewValue), FGameplayAttributeData(OldValue) );
 	}
 
 	/** Call from OnRep functions to set the attribute base value on the client */
 	void SetBaseAttributeValueFromReplication(const FGameplayAttribute& Attribute, const FGameplayAttributeData& NewValue, const FGameplayAttributeData& OldValue)
 	{
-		ActiveGameplayEffects.SetBaseAttributeValueFromReplication(Attribute, NewValue.GetBaseValue(), OldValue.GetBaseValue());
+		ActiveGameplayEffects.SetBaseAttributeValueFromReplication(Attribute, NewValue, OldValue);
 	}
 
 	/** Tests if all modifiers in this GameplayEffect will leave the attribute > 0.f */
@@ -1997,4 +1997,4 @@ protected:
 
 	TArray<FAbilityListLockActiveChange*> AbilityListLockActiveChanges;
 	
-};
\ No newline at end of file
+};
diff --git a/Engine/Plugins/Runtime/GameplayAbilities/Source/GameplayAbilities/Public/GameplayEffect.h b/Engine/Plugins/Runtime/GameplayAbilities/Source/GameplayAbilities/Public/GameplayEffect.h
--- a/Engine/Plugins/Runtime/GameplayAbilities/Source/GameplayAbilities/Public/GameplayEffect.h
+++ b/Engine/Plugins/Runtime/GameplayAbilities/Source/GameplayAbilities/Public/GameplayEffect.h
@@ -1785,7 +1785,8 @@ struct GAMEPLAYABILITIES_API FActiveGameplayEffectsContainer : public FFastArray
 	bool HasReceivedEffectWithPredictedKey(FPredictionKey PredictionKey) const;
 
 	bool HasPredictedEffectWithPredictedKey(FPredictionKey PredictionKey) const;
-		
+	
+	UE_DEPRECATED(5.5, "Replaced by private SetBaseAttributeValueFromReplication that uses FGameplayAttributeData")
 	void SetBaseAttributeValueFromReplication(const FGameplayAttribute& Attribute, float NewBaseValue, float OldBaseValue);
 
 	void GetAllActiveGameplayEffectSpecs(TArray<FGameplayEffectSpec>& OutSpecCopies) const;
@@ -1818,6 +1819,7 @@ struct GAMEPLAYABILITIES_API FActiveGameplayEffectsContainer : public FFastArray
 	void PostReplicatedReceive(const FFastArraySerializer::FPostReplicatedReceiveParameters& Parameters);
 
 private:
+	void SetBaseAttributeValueFromReplication(const FGameplayAttribute& Attribute, const FGameplayAttributeData& NewValue, const FGameplayAttributeData& OldValue);
 
 	/**
 	 *	Accessors for internal functions to get GameplayEffects directly by index.
@@ -2441,4 +2443,4 @@ GEComponentClass& UGameplayEffect::FindOrAddComponent()
 	}
 
 	return AddComponent<GEComponentClass>();
-}
\ No newline at end of file
+}
