GAS: More explicit warnings on invalid use of RemoveActiveGameplayEffects

#rnx

diff --git a/Engine/Plugins/Runtime/GameplayAbilities/Source/GameplayAbilities/Private/AbilitySystemComponent.cpp b/Engine/Plugins/Runtime/GameplayAbilities/Source/GameplayAbilities/Private/AbilitySystemComponent.cpp
--- a/Engine/Plugins/Runtime/GameplayAbilities/Source/GameplayAbilities/Private/AbilitySystemComponent.cpp
+++ b/Engine/Plugins/Runtime/GameplayAbilities/Source/GameplayAbilities/Private/AbilitySystemComponent.cpp
@@ -1045,7 +1045,8 @@ bool UAbilitySystemComponent::RemoveActiveGameplayEffect(FActiveGameplayEffectHa
 	{
 		if (!IsOwnerActorAuthoritative())
 		{
-			UE_LOG(LogAbilitySystem, Warning, TEXT("%hs called without Authority. Fix-up code, or temporarily patch using AbilitySystem.Fix.AllowPredictiveGEFlags"), __func__);
+			const UGameplayEffect* GEDef = GetGameplayEffectDefForHandle(Handle);
+			UE_LOG(LogAbilitySystem, Warning, TEXT("%hs called without Authority when attempting to remove %s. Fix-up code, or temporarily patch using AbilitySystem.Fix.AllowPredictiveGEFlags"), __func__, *GetNameSafe(GEDef));
 			return false;
 		}
 	}
