GAS:  When Clearing an FScalableFloat Property field (using the UI), the RowName doesn't get reset.  This can cause warnings/errors later when the scalable float is being evaluated because it has a RowName but no CurveTable.

diff --git a/Engine/Plugins/Runtime/GameplayAbilities/Source/GameplayAbilitiesEditor/Private/AttributeDetails.cpp b/Engine/Plugins/Runtime/GameplayAbilities/Source/GameplayAbilitiesEditor/Private/AttributeDetails.cpp
--- a/Engine/Plugins/Runtime/GameplayAbilities/Source/GameplayAbilitiesEditor/Private/AttributeDetails.cpp
+++ b/Engine/Plugins/Runtime/GameplayAbilities/Source/GameplayAbilitiesEditor/Private/AttributeDetails.cpp
@@ -553,6 +553,10 @@ void FScalableFloatDetails::OnBrowseTo()
 void FScalableFloatDetails::OnClear()
 {
 	OnSelectCurveTable(FAssetData());
+	if (RowNameProperty.IsValid())
+	{
+		RowNameProperty->ResetToDefault();
+	}
 }
 
 void FScalableFloatDetails::OnUseSelected()
