Fix GAS ensure being wrongly triggered
#rnx

diff --git a/Engine/Plugins/Runtime/GameplayAbilities/Source/GameplayAbilities/Private/AbilitySystemComponent_Abilities.cpp b/Engine/Plugins/Runtime/GameplayAbilities/Source/GameplayAbilities/Private/AbilitySystemComponent_Abilities.cpp
--- a/Engine/Plugins/Runtime/GameplayAbilities/Source/GameplayAbilities/Private/AbilitySystemComponent_Abilities.cpp
+++ b/Engine/Plugins/Runtime/GameplayAbilities/Source/GameplayAbilities/Private/AbilitySystemComponent_Abilities.cpp
@@ -1476,7 +1476,6 @@ void UAbilitySystemComponent::OnRep_ActivateAbilities()
 
 	// Make a copy in case a pending ability alters the array while iterating
 	TArray<FPendingAbilityInfo> PendingCopy = PendingServerActivatedAbilities;
-	PendingServerActivatedAbilities.Empty();
 
 	// Try to run any pending activations that couldn't run before. If they don't work now, kill them
 	for (const FPendingAbilityInfo& PendingAbilityInfo : PendingCopy)
@@ -1490,7 +1489,7 @@ void UAbilitySystemComponent::OnRep_ActivateAbilities()
 			ClientTryActivateAbility(PendingAbilityInfo.Handle);
 		}
 	}
-	ensureMsgf(PendingServerActivatedAbilities.IsEmpty(), TEXT("Execution of Pending Abilities caused %d more Pending Abilities (ignoring them)"), PendingServerActivatedAbilities.Num());
+	ensureMsgf(PendingServerActivatedAbilities.Num() == PendingCopy.Num(), TEXT("Execution of Pending Abilities caused %d more Pending Abilities (ignoring them)"), PendingServerActivatedAbilities.Num());
 	PendingServerActivatedAbilities.Empty();
 }
 
