GameplayCueInterface - Don't create FNames when searching for class function names. Since we're looking for something that should alread exist and is stored as an FName we can use FNAME_Find to see if this name already exists.

diff --git a/Engine/Plugins/Runtime/GameplayAbilities/Source/GameplayAbilities/Private/GameplayCueInterface.cpp b/Engine/Plugins/Runtime/GameplayAbilities/Source/GameplayAbilities/Private/GameplayCueInterface.cpp
--- a/Engine/Plugins/Runtime/GameplayAbilities/Source/GameplayAbilities/Private/GameplayCueInterface.cpp
+++ b/Engine/Plugins/Runtime/GameplayAbilities/Source/GameplayAbilities/Private/GameplayCueInterface.cpp
@@ -115,17 +115,24 @@ void IGameplayCueInterface::HandleGameplayCue(UObject* Self, FGameplayTag Gamepl
 			}
 
 			// Native functions cant be named with ".", so look for them with _. 
-			FName NativeCueFuncName = *CueName.ToString().Replace(TEXT("."), TEXT("_"));
-			Func = Class->FindFunctionByName(NativeCueFuncName, EIncludeSuperFlag::IncludeSuper);
-
-			while (Func)
+			FString NativeCueFuncName = CueName.ToString().Replace(TEXT("."), TEXT("_"));
+			FName NativeCueFuncFName(NativeCueFuncName, FNAME_Find);
+			
+			// The UClass stores all the function names as FNames internally.
+			// If the FName isn't found using FNAME_Find then it is not part of the UClass Function set.
+			if (!NativeCueFuncFName.IsNone())
 			{
-				GameplayCueInterfacePrivate::FCueNameAndUFunction NewCueFunctionPair;
-				NewCueFunctionPair.Tag = *InnerTagIt;
-				NewCueFunctionPair.Func = Func;
-				FunctionList->Add(NewCueFunctionPair);
+				Func = Class->FindFunctionByName(NativeCueFuncFName, EIncludeSuperFlag::IncludeSuper);
 
-				Func = Func->GetSuperFunction();
+				while (Func)
+				{
+					GameplayCueInterfacePrivate::FCueNameAndUFunction NewCueFunctionPair;
+					NewCueFunctionPair.Tag = *InnerTagIt;
+					NewCueFunctionPair.Func = Func;
+					FunctionList->Add(NewCueFunctionPair);
+
+					Func = Func->GetSuperFunction();
+				}
 			}
 		}
 	}
