[ Gameplay Cues] Added missing LastSourceArrayReplicationKey increment to mark Minimal Gameplay Cue struct dirty after a removal. This would cause an issue in replay where looping gameplay cues are never removed

#rnx

diff --git a/Engine/Plugins/Runtime/GameplayAbilities/Source/GameplayAbilities/Private/Serialization/MinimalGameplayCueReplicationProxyReplicationFragment.cpp b/Engine/Plugins/Runtime/GameplayAbilities/Source/GameplayAbilities/Private/Serialization/MinimalGameplayCueReplicationProxyReplicationFragment.cpp
--- a/Engine/Plugins/Runtime/GameplayAbilities/Source/GameplayAbilities/Private/Serialization/MinimalGameplayCueReplicationProxyReplicationFragment.cpp
+++ b/Engine/Plugins/Runtime/GameplayAbilities/Source/GameplayAbilities/Private/Serialization/MinimalGameplayCueReplicationProxyReplicationFragment.cpp
@@ -16,6 +16,11 @@
 namespace UE::Net
 {
 
+bool bShouldMarkStructDirtyAfterRemoval = true;
+static FAutoConsoleVariableRef CVarShouldMarkStructDirtyAfterRemoval(
+	TEXT("Net.GameplayCues.ShouldMarkStructDirtyAfterRemoval"), bShouldMarkStructDirtyAfterRemoval,
+	TEXT("Whether to increment LastSourceArrayReplicationKey after Gameplay Cue Removals"));
+
 FReplicationFragment* CreateAndRegisterMinimalGameplayCueReplicationProxyReplicationFragment(UObject* Owner, const FReplicationStateDescriptor* Descriptor, FFragmentRegistrationContext& Context)
 {
 	if (FMinimalGameplayCueReplicationProxyReplicationFragment* Fragment = new FMinimalGameplayCueReplicationProxyReplicationFragment(Context.GetFragmentTraits() | EReplicationFragmentTraits::DeleteWithInstanceProtocol, Owner, Descriptor))
@@ -188,7 +193,14 @@ void FMinimalGameplayCueReplicationProxyReplicationFragment::MimicMinimalGamepla
 
 		FGameplayTag RemovedTag = *It;
 		StateOwner->SetTagMapCount(RemovedTag, 0);
+
 		StateOwner->InvokeGameplayCueEvent(RemovedTag, EGameplayCueEvent::Removed, Parameters);
+
+		if(bShouldMarkStructDirtyAfterRemoval)
+		{
+			// The demo recorder needs to believe that this structure is dirty so it will get saved into the demo stream
+			ExternalSourceState->LastSourceArrayReplicationKey++;
+		}
 	}
 }
 
