Sequencer: remove IMovieScenePlayer usage in eval-hook sections

diff --git a/Engine/Plugins/Runtime/GameplayAbilities/Source/GameplayAbilities/Private/Sequencer/MovieSceneGameplayCueSections.cpp b/Engine/Plugins/Runtime/GameplayAbilities/Source/GameplayAbilities/Private/Sequencer/MovieSceneGameplayCueSections.cpp
--- a/Engine/Plugins/Runtime/GameplayAbilities/Source/GameplayAbilities/Private/Sequencer/MovieSceneGameplayCueSections.cpp
+++ b/Engine/Plugins/Runtime/GameplayAbilities/Source/GameplayAbilities/Private/Sequencer/MovieSceneGameplayCueSections.cpp
@@ -1,13 +1,14 @@
 // Copyright Epic Games, Inc. All Rights Reserved.
 
 #include "Sequencer/MovieSceneGameplayCueSections.h"
-#include "Sequencer/MovieSceneGameplayCueTrack.h"
-#include "EntitySystem/MovieSceneEvaluationHookSystem.h"
+
+#include "AbilitySystemGlobals.h"
 #include "Channels/MovieSceneChannelProxy.h"
 #include "EntitySystem/BuiltInComponentTypes.h"
-#include "IMovieScenePlayer.h"
-#include "AbilitySystemGlobals.h"
+#include "EntitySystem/MovieSceneEvaluationHookSystem.h"
+#include "Evaluation/MovieSceneEvaluationState.h"
 #include "GameplayCueManager.h"
+#include "Sequencer/MovieSceneGameplayCueTrack.h"
 
 #include UE_INLINE_GENERATED_CPP_BY_NAME(MovieSceneGameplayCueSections)
 
@@ -33,7 +34,7 @@ AActor* ActorFromResolvedObject(UObject* BoundObject)
 }
 
 
-void ExecuteGameplayCueEvent(IMovieScenePlayer* Player, const UE::MovieScene::FEvaluationHookParams& Params, const FMovieSceneGameplayCueKey& GameplayCueKey, EGameplayCueEvent::Type Event)
+void ExecuteGameplayCueEvent(TSharedRef<FSharedPlaybackState> SharedPlaybackState, const UE::MovieScene::FEvaluationHookParams& Params, const FMovieSceneGameplayCueKey& GameplayCueKey, EGameplayCueEvent::Type Event)
 {
 	auto OverrideHandler = [](AActor* InActor, FGameplayTag InTag, const FGameplayCueParameters& InParams, EGameplayCueEvent::Type InEvent)
 	{
@@ -67,7 +68,7 @@ void ExecuteGameplayCueEvent(IMovieScenePlayer* Player, const UE::MovieScene::FE
 
 	if (GameplayCueKey.Instigator.IsValid())
 	{
-		for (TWeakObjectPtr<> WeakObject : GameplayCueKey.Instigator.ResolveBoundObjects(Params.SequenceID, *Player))
+		for (TWeakObjectPtr<> WeakObject : GameplayCueKey.Instigator.ResolveBoundObjects(Params.SequenceID, SharedPlaybackState))
 		{
 			if (AActor* Actor = ActorFromResolvedObject(WeakObject.Get()))
 			{
@@ -79,7 +80,7 @@ void ExecuteGameplayCueEvent(IMovieScenePlayer* Player, const UE::MovieScene::FE
 
 	if (GameplayCueKey.EffectCauser.IsValid())
 	{
-		for (TWeakObjectPtr<> WeakObject : GameplayCueKey.EffectCauser.ResolveBoundObjects(Params.SequenceID, *Player))
+		for (TWeakObjectPtr<> WeakObject : GameplayCueKey.EffectCauser.ResolveBoundObjects(Params.SequenceID, SharedPlaybackState))
 		{
 			if (AActor* Actor = ActorFromResolvedObject(WeakObject.Get()))
 			{
@@ -92,7 +93,8 @@ void ExecuteGameplayCueEvent(IMovieScenePlayer* Player, const UE::MovieScene::FE
 
 	if (Params.ObjectBindingID.IsValid())
 	{
-		for (TWeakObjectPtr<> WeakObject : Player->FindBoundObjects(Params.ObjectBindingID, Params.SequenceID))
+		TArrayView<TWeakObjectPtr<>> BoundObjects = SharedPlaybackState->FindBoundObjects(Params.ObjectBindingID, Params.SequenceID);
+		for (TWeakObjectPtr<> WeakObject : BoundObjects)
 		{
 			UObject* Object = WeakObject.Get();
 			if (!Object)
@@ -255,7 +257,7 @@ EMovieSceneChannelProxyType UMovieSceneGameplayCueTriggerSection::CacheChannelPr
 	return EMovieSceneChannelProxyType::Static;
 }
 
-void UMovieSceneGameplayCueTriggerSection::Trigger(IMovieScenePlayer* Player, const UE::MovieScene::FEvaluationHookParams& Params) const
+void UMovieSceneGameplayCueTriggerSection::Trigger(TSharedRef<FSharedPlaybackState> SharedPlaybackState, const UE::MovieScene::FEvaluationHookParams& Params) const
 {
 	TMovieSceneChannelData<const FMovieSceneGameplayCueKey> ChannelData = Channel.GetData();
 	if (!ensureMsgf(ChannelData.GetValues().IsValidIndex(Params.TriggerIndex), TEXT("Invalid trigger index specified: %d (Num triggers in channel = %d)"), Params.TriggerIndex, ChannelData.GetValues().Num()))
@@ -265,7 +267,7 @@ void UMovieSceneGameplayCueTriggerSection::Trigger(IMovieScenePlayer* Player, co
 
 	if (Params.Context.GetStatus() == EMovieScenePlayerStatus::Playing && Params.Context.GetDirection() == EPlayDirection::Forwards && !Params.Context.IsSilent())
 	{
-		UE::MovieScene::ExecuteGameplayCueEvent(Player, Params, ChannelData.GetValues()[Params.TriggerIndex], EGameplayCueEvent::Executed);
+		UE::MovieScene::ExecuteGameplayCueEvent(SharedPlaybackState, Params, ChannelData.GetValues()[Params.TriggerIndex], EGameplayCueEvent::Executed);
 	}
 }
 
@@ -276,26 +278,26 @@ UMovieSceneGameplayCueSection::UMovieSceneGameplayCueSection(const FObjectInitia
 	bRequiresRangedHook = true;
 }
 
-void UMovieSceneGameplayCueSection::Begin(IMovieScenePlayer* Player, const UE::MovieScene::FEvaluationHookParams& Params) const
+void UMovieSceneGameplayCueSection::Begin(TSharedRef<FSharedPlaybackState> SharedPlaybackState, const UE::MovieScene::FEvaluationHookParams& Params) const
 {
 	if (Params.Context.GetStatus() == EMovieScenePlayerStatus::Playing && Params.Context.GetDirection() == EPlayDirection::Forwards && !Params.Context.IsSilent())
 	{
-		UE::MovieScene::ExecuteGameplayCueEvent(Player, Params, Cue, EGameplayCueEvent::OnActive);
+		UE::MovieScene::ExecuteGameplayCueEvent(SharedPlaybackState, Params, Cue, EGameplayCueEvent::OnActive);
 	}
 }
 
-void UMovieSceneGameplayCueSection::Update(IMovieScenePlayer* Player, const UE::MovieScene::FEvaluationHookParams& Params) const
+void UMovieSceneGameplayCueSection::Update(TSharedRef<FSharedPlaybackState> SharedPlaybackState, const UE::MovieScene::FEvaluationHookParams& Params) const
 {
 	if (Params.Context.GetStatus() == EMovieScenePlayerStatus::Playing && Params.Context.GetDirection() == EPlayDirection::Forwards && !Params.Context.IsSilent())
 	{
-		UE::MovieScene::ExecuteGameplayCueEvent(Player, Params, Cue, EGameplayCueEvent::WhileActive);
+		UE::MovieScene::ExecuteGameplayCueEvent(SharedPlaybackState, Params, Cue, EGameplayCueEvent::WhileActive);
 	}
 }
 
-void UMovieSceneGameplayCueSection::End(IMovieScenePlayer* Player, const UE::MovieScene::FEvaluationHookParams& Params) const
+void UMovieSceneGameplayCueSection::End(TSharedRef<FSharedPlaybackState> SharedPlaybackState, const UE::MovieScene::FEvaluationHookParams& Params) const
 {
 	if (Params.Context.GetStatus() == EMovieScenePlayerStatus::Playing && Params.Context.GetDirection() == EPlayDirection::Forwards && !Params.Context.IsSilent())
 	{
-		UE::MovieScene::ExecuteGameplayCueEvent(Player, Params, Cue, EGameplayCueEvent::Removed);
+		UE::MovieScene::ExecuteGameplayCueEvent(SharedPlaybackState, Params, Cue, EGameplayCueEvent::Removed);
 	}
 }
diff --git a/Engine/Plugins/Runtime/GameplayAbilities/Source/GameplayAbilities/Public/Sequencer/MovieSceneGameplayCueSections.h b/Engine/Plugins/Runtime/GameplayAbilities/Source/GameplayAbilities/Public/Sequencer/MovieSceneGameplayCueSections.h
--- a/Engine/Plugins/Runtime/GameplayAbilities/Source/GameplayAbilities/Public/Sequencer/MovieSceneGameplayCueSections.h
+++ b/Engine/Plugins/Runtime/GameplayAbilities/Source/GameplayAbilities/Public/Sequencer/MovieSceneGameplayCueSections.h
@@ -165,7 +165,7 @@ public:
 private:
 
 	virtual EMovieSceneChannelProxyType CacheChannelProxy() override;
-	virtual void Trigger(IMovieScenePlayer* Player, const UE::MovieScene::FEvaluationHookParams& Params) const override;
+	virtual void Trigger(TSharedRef<FSharedPlaybackState> SharedPlaybackState, const UE::MovieScene::FEvaluationHookParams& Params) const override;
 	virtual TArrayView<const FFrameNumber> GetTriggerTimes() const { return Channel.GetData().GetTimes(); }
 
 	UPROPERTY(AssetRegistrySearchable)
@@ -188,9 +188,9 @@ public:
 
 private:
 
-	virtual void Begin(IMovieScenePlayer* Player, const UE::MovieScene::FEvaluationHookParams& Params) const override;
-	virtual void Update(IMovieScenePlayer* Player, const UE::MovieScene::FEvaluationHookParams& Params) const override;
-	virtual void End(IMovieScenePlayer* Player, const UE::MovieScene::FEvaluationHookParams& Params) const override;
+	virtual void Begin(TSharedRef<FSharedPlaybackState> SharedPlaybackState, const UE::MovieScene::FEvaluationHookParams& Params) const override;
+	virtual void Update(TSharedRef<FSharedPlaybackState> SharedPlaybackState, const UE::MovieScene::FEvaluationHookParams& Params) const override;
+	virtual void End(TSharedRef<FSharedPlaybackState> SharedPlaybackState, const UE::MovieScene::FEvaluationHookParams& Params) const override;
 
 	UPROPERTY(EditAnywhere, AssetRegistrySearchable, Category="Cue", meta=(ShowOnlyInnerProperties))
 	FMovieSceneGameplayCueKey Cue;
diff --git a/Engine/Plugins/VirtualProduction/LensComponent/Source/LensComponentEditor/Private/MovieScene/MovieSceneLensComponentSection.cpp b/Engine/Plugins/VirtualProduction/LensComponent/Source/LensComponentEditor/Private/MovieScene/MovieSceneLensComponentSection.cpp
--- a/Engine/Plugins/VirtualProduction/LensComponent/Source/LensComponentEditor/Private/MovieScene/MovieSceneLensComponentSection.cpp
+++ b/Engine/Plugins/VirtualProduction/LensComponent/Source/LensComponentEditor/Private/MovieScene/MovieSceneLensComponentSection.cpp
@@ -3,9 +3,9 @@
 #include "MovieScene/MovieSceneLensComponentSection.h"
 
 #include "Channels/MovieSceneChannelProxy.h"
+#include "Evaluation/MovieSceneEvaluationState.h"
 #include "Evaluation/MovieScenePreAnimatedState.h"
 #include "Evaluation/PreAnimatedState/MovieScenePreAnimatedCaptureSource.h"
-#include "IMovieScenePlayer.h"
 #include "MovieSceneExecutionToken.h"
 #include "MovieSceneTimeHelpers.h"
 
@@ -215,15 +215,17 @@ void UMovieSceneLensComponentSection::Finalize()
 	UE::MovieScene::Optimize(&ImageCenterChannels[1], Params);
 }
 
-void UMovieSceneLensComponentSection::Begin(IMovieScenePlayer* Player, const UE::MovieScene::FEvaluationHookParams& Params) const
+void UMovieSceneLensComponentSection::Begin(TSharedRef<FSharedPlaybackState> SharedPlaybackState, const UE::MovieScene::FEvaluationHookParams& Params) const
 {
-	FScopedPreAnimatedCaptureSource CaptureSource(&Player->PreAnimatedState, this, Params.SequenceID, EvalOptions.CompletionMode == EMovieSceneCompletionMode::RestoreState);
+	FScopedPreAnimatedCaptureSource CaptureSource(SharedPlaybackState, this, Params.SequenceID, EvalOptions.CompletionMode == EMovieSceneCompletionMode::RestoreState);
 
-	for (TWeakObjectPtr<> BoundObject : Player->FindBoundObjects(Params.ObjectBindingID, Params.SequenceID))
+	FMovieSceneInstancePreAnimatedState& PreAnimatedState = SharedPlaybackState->GetPreAnimatedState();
+	TArrayView<TWeakObjectPtr<>> BoundObjects = SharedPlaybackState->FindBoundObjects(Params.ObjectBindingID, Params.SequenceID);
+	for (TWeakObjectPtr<> BoundObject : BoundObjects)
 	{
 		if (ULensComponent* LensComponent = Cast<ULensComponent>(BoundObject.Get()))
 		{
-			Player->SavePreAnimatedState(*LensComponent, FPreAnimatedLensFileTokenProducer::GetAnimTypeID(), FPreAnimatedLensFileTokenProducer());
+			PreAnimatedState.SavePreAnimatedState(*LensComponent, FPreAnimatedLensFileTokenProducer::GetAnimTypeID(), FPreAnimatedLensFileTokenProducer());
 
 			// Override the existing lens file settings of the lens component to use the cached, duplicate lens file
 			LensComponent->SetFIZEvaluationMode(EFIZEvaluationMode::UseRecordedValues);
@@ -233,9 +235,10 @@ void UMovieSceneLensComponentSection::Begin(IMovieScenePlayer* Player, const UE:
 	}
 }
 
-void UMovieSceneLensComponentSection::Update(IMovieScenePlayer* Player, const UE::MovieScene::FEvaluationHookParams& Params) const
+void UMovieSceneLensComponentSection::Update(TSharedRef<FSharedPlaybackState> SharedPlaybackState, const UE::MovieScene::FEvaluationHookParams& Params) const
 {
-	for (TWeakObjectPtr<> BoundObject : Player->FindBoundObjects(Params.ObjectBindingID, Params.SequenceID))
+	TArrayView<TWeakObjectPtr<>> BoundObjects = SharedPlaybackState->FindBoundObjects(Params.ObjectBindingID, Params.SequenceID);
+	for (TWeakObjectPtr<> BoundObject : BoundObjects)
 	{
 		if (ULensComponent* LensComponent = Cast<ULensComponent>(BoundObject.Get()))
 		{
@@ -307,9 +310,10 @@ void UMovieSceneLensComponentSection::Update(IMovieScenePlayer* Player, const UE
 	}
 }
 
-void UMovieSceneLensComponentSection::End(IMovieScenePlayer* Player, const UE::MovieScene::FEvaluationHookParams& Params) const
+void UMovieSceneLensComponentSection::End(TSharedRef<FSharedPlaybackState> SharedPlaybackState, const UE::MovieScene::FEvaluationHookParams& Params) const
 {
-	Player->RestorePreAnimatedState();
+	SharedPlaybackState->GetPreAnimatedState().RestorePreAnimatedState();
+	SharedPlaybackState->ClearObjectCaches();
 }
 
 #undef LOCTEXT_NAMESPACE
diff --git a/Engine/Plugins/VirtualProduction/LensComponent/Source/LensComponentEditor/Public/MovieScene/MovieSceneLensComponentSection.h b/Engine/Plugins/VirtualProduction/LensComponent/Source/LensComponentEditor/Public/MovieScene/MovieSceneLensComponentSection.h
--- a/Engine/Plugins/VirtualProduction/LensComponent/Source/LensComponentEditor/Public/MovieScene/MovieSceneLensComponentSection.h
+++ b/Engine/Plugins/VirtualProduction/LensComponent/Source/LensComponentEditor/Public/MovieScene/MovieSceneLensComponentSection.h
@@ -27,9 +27,9 @@ public:
 	}
 
 	//~ Begin UMovieSceneHookSection interface
-	virtual void Begin(IMovieScenePlayer* Player, const UE::MovieScene::FEvaluationHookParams& Params) const override;
-	virtual void Update(IMovieScenePlayer* Player, const UE::MovieScene::FEvaluationHookParams& Params) const override;
-	virtual void End(IMovieScenePlayer* Player, const UE::MovieScene::FEvaluationHookParams& Params) const override;
+	virtual void Begin(TSharedRef<FSharedPlaybackState> SharedPlaybackState, const UE::MovieScene::FEvaluationHookParams& Params) const override;
+	virtual void Update(TSharedRef<FSharedPlaybackState> SharedPlaybackState, const UE::MovieScene::FEvaluationHookParams& Params) const override;
+	virtual void End(TSharedRef<FSharedPlaybackState> SharedPlaybackState, const UE::MovieScene::FEvaluationHookParams& Params) const override;
 	//~ End UMovieSceneHookSection interface
 
 	/** Initialize the section with the Lens Component whose properties it will record */
diff --git a/Engine/Plugins/VirtualProduction/LiveLinkCamera/Source/LiveLinkCameraRecording/Private/MovieSceneLiveLinkCameraControllerSection.cpp b/Engine/Plugins/VirtualProduction/LiveLinkCamera/Source/LiveLinkCameraRecording/Private/MovieSceneLiveLinkCameraControllerSection.cpp
--- a/Engine/Plugins/VirtualProduction/LiveLinkCamera/Source/LiveLinkCameraRecording/Private/MovieSceneLiveLinkCameraControllerSection.cpp
+++ b/Engine/Plugins/VirtualProduction/LiveLinkCamera/Source/LiveLinkCameraRecording/Private/MovieSceneLiveLinkCameraControllerSection.cpp
@@ -3,7 +3,8 @@
 #include "MovieSceneLiveLinkCameraControllerSection.h"
 
 #include "CineCameraComponent.h"
-#include "IMovieScenePlayer.h"
+#include "EntitySystem/MovieSceneSharedPlaybackState.h"
+#include "Evaluation/MovieSceneEvaluationState.h"
 #include "LiveLinkCameraController.h"
 #include "LiveLinkComponentController.h"
 #include "Roles/LiveLinkCameraRole.h"
@@ -12,14 +13,15 @@ void UMovieSceneLiveLinkCameraControllerSection::Initialize(ULiveLinkControllerB
 {
 }
 
-void UMovieSceneLiveLinkCameraControllerSection::Update(IMovieScenePlayer* Player, const UE::MovieScene::FEvaluationHookParams& Params) const
+void UMovieSceneLiveLinkCameraControllerSection::Update(TSharedRef<FSharedPlaybackState> SharedPlaybackState, const UE::MovieScene::FEvaluationHookParams& Params) const
 {
 	if (!CachedLensFile || !bApplyNodalOffsetFromCachedLensFile)
 	{
 		return;
 	}
 
-	for (TWeakObjectPtr<>& BoundObject : Player->FindBoundObjects(Params.ObjectBindingID, Params.SequenceID))
+	TArrayView<TWeakObjectPtr<>> BoundObjects = SharedPlaybackState->FindBoundObjects(Params.ObjectBindingID, Params.SequenceID);
+	for (TWeakObjectPtr<>& BoundObject : BoundObjects)
 	{
 		if (ULiveLinkComponentController* LiveLinkComponent = Cast<ULiveLinkComponentController>(BoundObject.Get()))
 		{
diff --git a/Engine/Plugins/VirtualProduction/LiveLinkCamera/Source/LiveLinkCameraRecording/Private/MovieSceneLiveLinkCameraControllerSection.h b/Engine/Plugins/VirtualProduction/LiveLinkCamera/Source/LiveLinkCameraRecording/Private/MovieSceneLiveLinkCameraControllerSection.h
--- a/Engine/Plugins/VirtualProduction/LiveLinkCamera/Source/LiveLinkCameraRecording/Private/MovieSceneLiveLinkCameraControllerSection.h
+++ b/Engine/Plugins/VirtualProduction/LiveLinkCamera/Source/LiveLinkCameraRecording/Private/MovieSceneLiveLinkCameraControllerSection.h
@@ -30,7 +30,7 @@ public:
 	void Initialize(ULiveLinkControllerBase* InLiveLinkController);
 
 	//~ Begin UMovieSceneHookSection interface
-	virtual void Update(IMovieScenePlayer* Player, const UE::MovieScene::FEvaluationHookParams& Params) const override;
+	virtual void Update(TSharedRef<FSharedPlaybackState> SharedPlaybackState, const UE::MovieScene::FEvaluationHookParams& Params) const override;
 	//~ End UMovieSceneHookSection interface
 
 public:
diff --git a/Engine/Source/Runtime/MovieScene/Private/EntitySystem/MovieSceneEvaluationHookSystem.cpp b/Engine/Source/Runtime/MovieScene/Private/EntitySystem/MovieSceneEvaluationHookSystem.cpp
--- a/Engine/Source/Runtime/MovieScene/Private/EntitySystem/MovieSceneEvaluationHookSystem.cpp
+++ b/Engine/Source/Runtime/MovieScene/Private/EntitySystem/MovieSceneEvaluationHookSystem.cpp
@@ -1,10 +1,12 @@
 // Copyright Epic Games, Inc. All Rights Reserved.
 
 #include "EntitySystem/MovieSceneEvaluationHookSystem.h"
+
 #include "EntitySystem/MovieSceneEntitySystemLinker.h"
 #include "EntitySystem/MovieSceneEntitySystemRunner.h"
-#include "EntitySystem/MovieSceneSpawnablesSystem.h"
 #include "EntitySystem/MovieSceneEntitySystemTask.h"
+#include "EntitySystem/MovieSceneSharedPlaybackState.h"
+#include "EntitySystem/MovieSceneSpawnablesSystem.h"
 #include "Evaluation/PreAnimatedState/MovieScenePreAnimatedCaptureSource.h"
 #include "Evaluation/PreAnimatedState/MovieScenePreAnimatedCaptureSources.h"
 #include "Evaluation/MovieSceneEvaluationTemplateInstance.h"
@@ -235,9 +237,8 @@ void UMovieSceneEvaluationHookSystem::TriggerAllEvents()
 	{
 		const FSequenceInstance& SequenceInstance = InstanceRegistry->GetInstance(Pair.Key.InstanceHandle);
 
-		IMovieScenePlayer* Player      = SequenceInstance.GetPlayer();
 		FMovieSceneContext RootContext = SequenceInstance.GetContext();
-		TSharedRef<const FSharedPlaybackState> SharedPlaybackState = SequenceInstance.GetSharedPlaybackState();
+		TSharedRef<FSharedPlaybackState> SharedPlaybackState = SequenceInstance.GetSharedPlaybackState();
 		FPreAnimatedEvaluationHookCaptureSources* EvaluationHookMetaData = Linker->PreAnimatedState.GetEvaluationHookMetaData();
 
 		for (const FMovieSceneEvaluationHookEvent& Event : Pair.Value.Events)
@@ -260,13 +261,13 @@ void UMovieSceneEvaluationHookSystem::TriggerAllEvents()
 			switch (Event.Type)
 			{
 				case EEvaluationHookEvent::Begin:
-					Event.Hook.Interface->Begin(Player, Params);
+					Event.Hook.Interface->Begin(SharedPlaybackState, Params);
 					break;
 				case EEvaluationHookEvent::Update:
-					Event.Hook.Interface->Update(Player, Params);
+					Event.Hook.Interface->Update(SharedPlaybackState, Params);
 					break;
 				case EEvaluationHookEvent::End:
-					Event.Hook.Interface->End(Player, Params);
+					Event.Hook.Interface->End(SharedPlaybackState, Params);
 					if (EvaluationHookMetaData)
 					{
 						EvaluationHookMetaData->StopTrackingCaptureSource(Event.Hook.Interface.GetObject(), Event.RootInstanceHandle, Event.SequenceID);
@@ -274,7 +275,7 @@ void UMovieSceneEvaluationHookSystem::TriggerAllEvents()
 					break;
 
 				case EEvaluationHookEvent::Trigger:
-					Event.Hook.Interface->Trigger(Player, Params);
+					Event.Hook.Interface->Trigger(SharedPlaybackState, Params);
 					break;
 			}
 		}
diff --git a/Engine/Source/Runtime/MovieScene/Private/Evaluation/IMovieSceneEvaluationHook.cpp b/Engine/Source/Runtime/MovieScene/Private/Evaluation/IMovieSceneEvaluationHook.cpp
new file mode 100644
--- /dev/null
+++ b/Engine/Source/Runtime/MovieScene/Private/Evaluation/IMovieSceneEvaluationHook.cpp
@@ -0,0 +1,36 @@
+// Copyright Epic Games, Inc. All Rights Reserved.
+
+#include "Evaluation/IMovieSceneEvaluationHook.h"
+
+#include "IMovieScenePlayer.h"
+
+#include UE_INLINE_GENERATED_CPP_BY_NAME(IMovieSceneEvaluationHook)
+
+void IMovieSceneEvaluationHook::Begin(TSharedRef<FSharedPlaybackState> SharedPlaybackState, const UE::MovieScene::FEvaluationHookParams& Params) const
+{
+	PRAGMA_DISABLE_DEPRECATION_WARNINGS
+	Begin(UE::MovieScene::FPlayerIndexPlaybackCapability::GetPlayer(SharedPlaybackState), Params);
+	PRAGMA_ENABLE_DEPRECATION_WARNINGS
+}
+
+void IMovieSceneEvaluationHook::Update(TSharedRef<FSharedPlaybackState> SharedPlaybackState, const UE::MovieScene::FEvaluationHookParams& Params) const
+{
+	PRAGMA_DISABLE_DEPRECATION_WARNINGS
+	Update(UE::MovieScene::FPlayerIndexPlaybackCapability::GetPlayer(SharedPlaybackState), Params);
+	PRAGMA_ENABLE_DEPRECATION_WARNINGS
+}
+
+void IMovieSceneEvaluationHook::End(TSharedRef<FSharedPlaybackState> SharedPlaybackState, const UE::MovieScene::FEvaluationHookParams& Params) const
+{
+	PRAGMA_DISABLE_DEPRECATION_WARNINGS
+	End(UE::MovieScene::FPlayerIndexPlaybackCapability::GetPlayer(SharedPlaybackState), Params);
+	PRAGMA_ENABLE_DEPRECATION_WARNINGS
+}
+
+void IMovieSceneEvaluationHook::Trigger(TSharedRef<FSharedPlaybackState> SharedPlaybackState, const UE::MovieScene::FEvaluationHookParams& Params) const
+{
+	PRAGMA_DISABLE_DEPRECATION_WARNINGS
+	Trigger(UE::MovieScene::FPlayerIndexPlaybackCapability::GetPlayer(SharedPlaybackState), Params);
+	PRAGMA_ENABLE_DEPRECATION_WARNINGS
+}
+
diff --git a/Engine/Source/Runtime/MovieScene/Private/Tests/MovieScenePreAnimatedStateTests.cpp b/Engine/Source/Runtime/MovieScene/Private/Tests/MovieScenePreAnimatedStateTests.cpp
--- a/Engine/Source/Runtime/MovieScene/Private/Tests/MovieScenePreAnimatedStateTests.cpp
+++ b/Engine/Source/Runtime/MovieScene/Private/Tests/MovieScenePreAnimatedStateTests.cpp
@@ -112,39 +112,43 @@ namespace Impl
 	}
 }
 
-void UTestMovieSceneEvalHookSection::Begin(IMovieScenePlayer* Player, const UE::MovieScene::FEvaluationHookParams& Params) const
+void UTestMovieSceneEvalHookSection::Begin(TSharedRef<FSharedPlaybackState> SharedPlaybackState, const UE::MovieScene::FEvaluationHookParams& Params) const
 {
 	using namespace Impl;
 
-	FScopedPreAnimatedCaptureSource CaptureSource(&Player->PreAnimatedState, this, Params.SequenceID, EvalOptions.CompletionMode == EMovieSceneCompletionMode::RestoreState);
-	Player->SavePreAnimatedState(AnimType1, FPreAnimatedTokenProducer(&TestValue1));
+	FScopedPreAnimatedCaptureSource CaptureSource(SharedPlaybackState, this, Params.SequenceID, EvalOptions.CompletionMode == EMovieSceneCompletionMode::RestoreState);
+	UMovieSceneEntitySystemLinker* Linker = SharedPlaybackState->GetLinker();
+	Linker->PreAnimatedState.SavePreAnimatedState(AnimType1, FPreAnimatedTokenProducer(&TestValue1));
 
 	TestValue1 = StartValue;
 }
-void UTestMovieSceneEvalHookSection::Update(IMovieScenePlayer* Player, const UE::MovieScene::FEvaluationHookParams& Params) const
+void UTestMovieSceneEvalHookSection::Update(TSharedRef<FSharedPlaybackState> SharedPlaybackState, const UE::MovieScene::FEvaluationHookParams& Params) const
 {
 	using namespace Impl;
 
-	FScopedPreAnimatedCaptureSource CaptureSource(&Player->PreAnimatedState, this, Params.SequenceID, EvalOptions.CompletionMode == EMovieSceneCompletionMode::RestoreState);
-	Player->SavePreAnimatedState(AnimType1, FPreAnimatedTokenProducer(&TestValue1));
+	FScopedPreAnimatedCaptureSource CaptureSource(SharedPlaybackState, this, Params.SequenceID, EvalOptions.CompletionMode == EMovieSceneCompletionMode::RestoreState);
+	UMovieSceneEntitySystemLinker* Linker = SharedPlaybackState->GetLinker();
+	Linker->PreAnimatedState.SavePreAnimatedState(AnimType1, FPreAnimatedTokenProducer(&TestValue1));
 
 	++TestValue1;
 }
-void UTestMovieSceneEvalHookSection::End(IMovieScenePlayer* Player, const UE::MovieScene::FEvaluationHookParams& Params) const
+void UTestMovieSceneEvalHookSection::End(TSharedRef<FSharedPlaybackState> SharedPlaybackState, const UE::MovieScene::FEvaluationHookParams& Params) const
 {
 	using namespace Impl;
 
-	FScopedPreAnimatedCaptureSource CaptureSource(&Player->PreAnimatedState, this, Params.SequenceID, EvalOptions.CompletionMode == EMovieSceneCompletionMode::RestoreState);
-	Player->SavePreAnimatedState(AnimType1, FPreAnimatedTokenProducer(&TestValue1));
+	FScopedPreAnimatedCaptureSource CaptureSource(SharedPlaybackState, this, Params.SequenceID, EvalOptions.CompletionMode == EMovieSceneCompletionMode::RestoreState);
+	UMovieSceneEntitySystemLinker* Linker = SharedPlaybackState->GetLinker();
+	Linker->PreAnimatedState.SavePreAnimatedState(AnimType1, FPreAnimatedTokenProducer(&TestValue1));
 
 	TestValue1 = EndValue;
 }
-void UTestMovieSceneEvalHookSection::Trigger(IMovieScenePlayer* Player, const UE::MovieScene::FEvaluationHookParams& Params) const
+void UTestMovieSceneEvalHookSection::Trigger(TSharedRef<FSharedPlaybackState> SharedPlaybackState, const UE::MovieScene::FEvaluationHookParams& Params) const
 {
 	using namespace Impl;
 
-	FScopedPreAnimatedCaptureSource CaptureSource(&Player->PreAnimatedState, this, Params.SequenceID, EvalOptions.CompletionMode == EMovieSceneCompletionMode::RestoreState);
-	Player->SavePreAnimatedState(AnimType1, FPreAnimatedTokenProducer(&TestValue1));
+	FScopedPreAnimatedCaptureSource CaptureSource(SharedPlaybackState, this, Params.SequenceID, EvalOptions.CompletionMode == EMovieSceneCompletionMode::RestoreState);
+	UMovieSceneEntitySystemLinker* Linker = SharedPlaybackState->GetLinker();
+	Linker->PreAnimatedState.SavePreAnimatedState(AnimType1, FPreAnimatedTokenProducer(&TestValue1));
 
 	--TestValue1;
 }
diff --git a/Engine/Source/Runtime/MovieScene/Private/Tests/MovieSceneTestObjects.h b/Engine/Source/Runtime/MovieScene/Private/Tests/MovieSceneTestObjects.h
--- a/Engine/Source/Runtime/MovieScene/Private/Tests/MovieSceneTestObjects.h
+++ b/Engine/Source/Runtime/MovieScene/Private/Tests/MovieSceneTestObjects.h
@@ -146,13 +146,13 @@ public:
 
 	TArray<FFrameNumber> TriggerTimes;
 
-	virtual void Begin(IMovieScenePlayer* Player, const UE::MovieScene::FEvaluationHookParams& Params) const override;
-	virtual void Update(IMovieScenePlayer* Player, const UE::MovieScene::FEvaluationHookParams& Params) const override;
-	virtual void End(IMovieScenePlayer* Player, const UE::MovieScene::FEvaluationHookParams& Params) const override;
-	virtual void Trigger(IMovieScenePlayer* Player, const UE::MovieScene::FEvaluationHookParams& Params) const override;
+	virtual void Begin(TSharedRef<FSharedPlaybackState> SharedPlaybackState, const UE::MovieScene::FEvaluationHookParams& Params) const override;
+	virtual void Update(TSharedRef<FSharedPlaybackState> SharedPlaybackState, const UE::MovieScene::FEvaluationHookParams& Params) const override;
+	virtual void End(TSharedRef<FSharedPlaybackState> SharedPlaybackState, const UE::MovieScene::FEvaluationHookParams& Params) const override;
+	virtual void Trigger(TSharedRef<FSharedPlaybackState> SharedPlaybackState, const UE::MovieScene::FEvaluationHookParams& Params) const override;
 
 	virtual TArrayView<const FFrameNumber> GetTriggerTimes() const
 	{
 		return TriggerTimes;
 	}
-};
\ No newline at end of file
+};
diff --git a/Engine/Source/Runtime/MovieScene/Public/EntitySystem/MovieSceneEvaluationHookSystem.h b/Engine/Source/Runtime/MovieScene/Public/EntitySystem/MovieSceneEvaluationHookSystem.h
--- a/Engine/Source/Runtime/MovieScene/Public/EntitySystem/MovieSceneEvaluationHookSystem.h
+++ b/Engine/Source/Runtime/MovieScene/Public/EntitySystem/MovieSceneEvaluationHookSystem.h
@@ -17,7 +17,6 @@
 
 #include "MovieSceneEvaluationHookSystem.generated.h"
 
-class IMovieScenePlayer;
 class UMovieSceneEntitySystemLinker;
 class UObject;
 
diff --git a/Engine/Source/Runtime/MovieScene/Public/Evaluation/IMovieSceneEvaluationHook.h b/Engine/Source/Runtime/MovieScene/Public/Evaluation/IMovieSceneEvaluationHook.h
--- a/Engine/Source/Runtime/MovieScene/Public/Evaluation/IMovieSceneEvaluationHook.h
+++ b/Engine/Source/Runtime/MovieScene/Public/Evaluation/IMovieSceneEvaluationHook.h
@@ -16,6 +16,7 @@ namespace UE
 namespace MovieScene
 {
 
+struct FSharedPlaybackState;
 
 enum class EEvaluationHookEvent
 {
@@ -60,11 +61,26 @@ class IMovieSceneEvaluationHook
 {
 public:
 
+	using FSharedPlaybackState = UE::MovieScene::FSharedPlaybackState;
+
 	GENERATED_BODY()
 
+	MOVIESCENE_API virtual void Begin(TSharedRef<FSharedPlaybackState> SharedPlaybackState, const UE::MovieScene::FEvaluationHookParams& Params) const;
+	MOVIESCENE_API virtual void Update(TSharedRef<FSharedPlaybackState> SharedPlaybackState, const UE::MovieScene::FEvaluationHookParams& Params) const;
+	MOVIESCENE_API virtual void End(TSharedRef<FSharedPlaybackState> SharedPlaybackState, const UE::MovieScene::FEvaluationHookParams& Params) const;
+
+	MOVIESCENE_API virtual void Trigger(TSharedRef<FSharedPlaybackState> SharedPlaybackState, const UE::MovieScene::FEvaluationHookParams& Params) const;
+
+protected:
+
+	UE_DEPRECATED(5.5, "Please implement the version that takes a SharedPlaybackState")
 	virtual void Begin(IMovieScenePlayer* Player, const UE::MovieScene::FEvaluationHookParams& Params) const {}
+	UE_DEPRECATED(5.5, "Please implement the version that takes a SharedPlaybackState")
 	virtual void Update(IMovieScenePlayer* Player, const UE::MovieScene::FEvaluationHookParams& Params) const {}
+	UE_DEPRECATED(5.5, "Please implement the version that takes a SharedPlaybackState")
 	virtual void End(IMovieScenePlayer* Player, const UE::MovieScene::FEvaluationHookParams& Params) const {}
 
+	UE_DEPRECATED(5.5, "Please implement the version that takes a SharedPlaybackState")
 	virtual void Trigger(IMovieScenePlayer* Player, const UE::MovieScene::FEvaluationHookParams& Params) const {}
 };
+
diff --git a/Engine/Source/Runtime/MovieScene/Public/Sections/MovieSceneHookSection.h b/Engine/Source/Runtime/MovieScene/Public/Sections/MovieSceneHookSection.h
--- a/Engine/Source/Runtime/MovieScene/Public/Sections/MovieSceneHookSection.h
+++ b/Engine/Source/Runtime/MovieScene/Public/Sections/MovieSceneHookSection.h
@@ -41,11 +41,11 @@ protected:
 
 	/*~ Implemented in derived classes
 
-	virtual void Begin(IMovieScenePlayer* Player, const UE::MovieScene::FEvaluationHookParams& Params) const override;
-	virtual void Update(IMovieScenePlayer* Player, const UE::MovieScene::FEvaluationHookParams& Params) const override;
-	virtual void End(IMovieScenePlayer* Player, const UE::MovieScene::FEvaluationHookParams& Params) const override;
+	virtual void Begin(TSharedRef<FSharedPlaybackState> SharedPlaybackState, const UE::MovieScene::FEvaluationHookParams& Params) const override;
+	virtual void Update(TSharedRef<FSharedPlaybackState> SharedPlaybackState, const UE::MovieScene::FEvaluationHookParams& Params) const override;
+	virtual void End(TSharedRef<FSharedPlaybackState> SharedPlaybackState, const UE::MovieScene::FEvaluationHookParams& Params) const override;
 
-	virtual void Trigger(IMovieScenePlayer* Player, const UE::MovieScene::FEvaluationHookParams& Params) const override;
+	virtual void Trigger(TSharedRef<FSharedPlaybackState> SharedPlaybackState, const UE::MovieScene::FEvaluationHookParams& Params) const override;
 
 	*/
 
