Fix a live crash on GetWorld(), changing to ensure to try and track it down (this shouldn't happen in normal flow).

diff --git a/Engine/Plugins/Runtime/GameplayAbilities/Source/GameplayAbilities/Private/AbilitySystemComponent_Abilities.cpp b/Engine/Plugins/Runtime/GameplayAbilities/Source/GameplayAbilities/Private/AbilitySystemComponent_Abilities.cpp
--- a/Engine/Plugins/Runtime/GameplayAbilities/Source/GameplayAbilities/Private/AbilitySystemComponent_Abilities.cpp
+++ b/Engine/Plugins/Runtime/GameplayAbilities/Source/GameplayAbilities/Private/AbilitySystemComponent_Abilities.cpp
@@ -1945,7 +1945,11 @@ PRAGMA_ENABLE_DEPRECATION_WARNINGS
 
 	MarkAbilitySpecDirty(*Spec);
 
-	AbilityLastActivatedTime = GetWorld()->GetTimeSeconds();
+	const UWorld* LocalWorld = GetWorld();
+	if (ensureMsgf(LocalWorld, TEXT("%hs: Could not GetWorld during activation of %s"), __func__, *GetNameSafe(Ability)))
+	{
+		AbilityLastActivatedTime = LocalWorld->GetTimeSeconds();
+	}
 
 	UE_LOG(LogAbilitySystem, Log, TEXT("%s: Activated [%s] %s. Level: %d. PredictionKey: %s."), *GetNameSafe(GetOwner()), *Spec->Handle.ToString(), *GetNameSafe(AbilitySource), Spec->Level, *ActivationInfo.GetActivationPredictionKey().ToString());
 	UE_VLOG(GetOwner(), VLogAbilitySystem, Log, TEXT("Activated [%s] %s. Level: %d. PredictionKey: %s."), *Spec->Handle.ToString(), *GetNameSafe(AbilitySource), Spec->Level, *ActivationInfo.GetActivationPredictionKey().ToString());
