GAS: Reintroducing some error logging on recursive ability activation.

diff --git a/Engine/Plugins/Runtime/GameplayAbilities/Source/GameplayAbilities/Private/AbilitySystemComponent_Abilities.cpp b/Engine/Plugins/Runtime/GameplayAbilities/Source/GameplayAbilities/Private/AbilitySystemComponent_Abilities.cpp
--- a/Engine/Plugins/Runtime/GameplayAbilities/Source/GameplayAbilities/Private/AbilitySystemComponent_Abilities.cpp
+++ b/Engine/Plugins/Runtime/GameplayAbilities/Source/GameplayAbilities/Private/AbilitySystemComponent_Abilities.cpp
@@ -1476,6 +1476,7 @@ void UAbilitySystemComponent::OnRep_ActivateAbilities()
 
 	// Make a copy in case a pending ability alters the array while iterating
 	TArray<FPendingAbilityInfo> PendingCopy = PendingServerActivatedAbilities;
+	PendingServerActivatedAbilities.Empty();
 
 	// Try to run any pending activations that couldn't run before. If they don't work now, kill them
 	for (const FPendingAbilityInfo& PendingAbilityInfo : PendingCopy)
@@ -1488,8 +1489,21 @@ void UAbilitySystemComponent::OnRep_ActivateAbilities()
 		{
 			ClientTryActivateAbility(PendingAbilityInfo.Handle);
 		}
+
+		// Do some warning if we're about to drop this ability activation
+		if (!PendingServerActivatedAbilities.IsEmpty())
+		{
+			const bool bIsSame = (PendingServerActivatedAbilities[0] == PendingAbilityInfo);
+			const bool bHasNewItem = !bIsSame || (PendingServerActivatedAbilities.Num() > 1);
+			UE_CLOG(bIsSame, LogAbilitySystem, Warning, TEXT("Failed to execute Pending Ability %s (Handle %s) because it was not replicated in time for second-chance activation (it will be ignored)"), *PendingAbilityInfo.PredictionKey.ToString(), *PendingAbilityInfo.Handle.ToString());
+			UE_CLOG(bHasNewItem, LogAbilitySystem, Warning, TEXT("New Pending Ability added during existing execution of %s (Handle %s). New ability will be ignored."), *PendingAbilityInfo.PredictionKey.ToString(), *PendingAbilityInfo.Handle.ToString());
+
+			// Empty again, so we can test again on the next item.
+			PendingServerActivatedAbilities.Empty();
+		}
 	}
-	ensureMsgf(PendingServerActivatedAbilities.Num() == PendingCopy.Num(), TEXT("Execution of Pending Abilities caused %d more Pending Abilities (ignoring them)"), PendingServerActivatedAbilities.Num());
+
+	// This is redundant but helps signal we're leaving this function with an empty pending ability list
 	PendingServerActivatedAbilities.Empty();
 }
 
