GAS:  Giving a warning to users during validation that an FScalableFloat RowName could be invisibly set to a value (which would cause warnings at cook/runtime).

diff --git a/Engine/Plugins/Runtime/GameplayAbilities/Source/GameplayAbilities/Private/GameplayEffect.cpp b/Engine/Plugins/Runtime/GameplayAbilities/Source/GameplayAbilities/Private/GameplayEffect.cpp
--- a/Engine/Plugins/Runtime/GameplayAbilities/Source/GameplayAbilities/Private/GameplayEffect.cpp
+++ b/Engine/Plugins/Runtime/GameplayAbilities/Source/GameplayAbilities/Private/GameplayEffect.cpp
@@ -250,6 +250,14 @@ void UGameplayEffect::GetBlockedAbilityTags(FGameplayTagContainer& OutTagContain
 #if WITH_EDITOR
 EDataValidationResult UGameplayEffect::IsDataValid(FDataValidationContext& Context) const
 {
+	// Report these errors here & now, even though they're not using the DataValidationContext because
+	// these will show-up when saving/presubmit rather than during more extensive testing done at cook time.
+	DurationMagnitude.ReportErrors(GetPathName());
+	for (const FGameplayModifierInfo& CurModInfo : Modifiers)
+	{
+		CurModInfo.ModifierMagnitude.ReportErrors(GetPathName());
+	}
+
 	EDataValidationResult ValidationResult = Super::IsDataValid(Context);
 
 	if (ValidationResult != EDataValidationResult::Invalid)
@@ -304,13 +312,6 @@ void UGameplayEffect::PostLoad()
 	HasRemoveGameplayEffectsQuery = !RemoveGameplayEffectQuery.IsEmpty();
 	PRAGMA_ENABLE_DEPRECATION_WARNINGS
 
-	DurationMagnitude.ReportErrors(GetPathName());
-
-	for (FGameplayModifierInfo& CurModInfo : Modifiers)
-	{
-		CurModInfo.ModifierMagnitude.ReportErrors(GetPathName());
-	}
-
 	// Update the EditorStatusText
 	FDataValidationContext DataValidationContext;
 	IsDataValid(DataValidationContext);
