Adding some debug information to track down warnings for potentially invalid use of GetOwnedGameplayTags() and related functions on Gameplay Effects.

diff --git a/Engine/Plugins/Runtime/GameplayAbilities/Source/GameplayAbilities/Private/GameplayEffect.cpp b/Engine/Plugins/Runtime/GameplayAbilities/Source/GameplayAbilities/Private/GameplayEffect.cpp
--- a/Engine/Plugins/Runtime/GameplayAbilities/Source/GameplayAbilities/Private/GameplayEffect.cpp
+++ b/Engine/Plugins/Runtime/GameplayAbilities/Source/GameplayAbilities/Private/GameplayEffect.cpp
@@ -220,25 +220,25 @@ void UGameplayEffect::PostInitProperties()
 
 void UGameplayEffect::GetOwnedGameplayTags(FGameplayTagContainer& TagContainer) const
 {
-	UE_LOG(LogGameplayEffects, Warning, TEXT("GetOwnedGameplayTags: The implementation and method name did not match.  Use GetGrantedTags() to get the tags Granted to the Actor this GameplayEffect is applied to."));
-	TagContainer.AppendTags(GetGrantedTags());
+	UE_LOG(LogGameplayEffects, Warning, TEXT("%hs on %s: The implementation and method name did not match.  Use GetGrantedTags() to get the tags Granted to the Actor this GameplayEffect is applied to."), __func__, *GetName());
+ 	TagContainer.AppendTags(GetGrantedTags());
 }
 
 bool UGameplayEffect::HasMatchingGameplayTag(FGameplayTag TagToCheck) const
 {
-	UE_LOG(LogGameplayEffects, Warning, TEXT("HasMatchingGameplayTag: The implementation and method name did not match.  Use GetGrantedTags().HasTag() to check against the tags this GameplayEffect will Grant to the Actor."));
+	UE_LOG(LogGameplayEffects, Warning, TEXT("%hs on %s: The implementation and method name did not match.  Use GetGrantedTags().HasTag() to check against the tags this GameplayEffect will Grant to the Actor."), __func__, *GetName());
 	return IGameplayTagAssetInterface::HasMatchingGameplayTag(TagToCheck);
 }
 
 bool UGameplayEffect::HasAllMatchingGameplayTags(const FGameplayTagContainer& TagContainer) const
 {
-	UE_LOG(LogGameplayEffects, Warning, TEXT("HasAllMatchingGameplayTags: The implementation and method name did not match.  Use GetGrantedTags().HasAll() to check against the tags this GameplayEffect will Grant to the Actor."));
+	UE_LOG(LogGameplayEffects, Warning, TEXT("%hs on %s: The implementation and method name did not match.  Use GetGrantedTags().HasAll() to check against the tags this GameplayEffect will Grant to the Actor."), __func__, *GetName());
 	return IGameplayTagAssetInterface::HasAllMatchingGameplayTags(TagContainer);
 }
 
 bool UGameplayEffect::HasAnyMatchingGameplayTags(const FGameplayTagContainer& TagContainer) const
 {
-	UE_LOG(LogGameplayEffects, Warning, TEXT("HasAnyMatchingGameplayTags: The implementation and method name did not match.  Use GetGrantedTags().HasAny() to check against the tags this GameplayEffect will Grant to the Actor."));
+	UE_LOG(LogGameplayEffects, Warning, TEXT("%hs on %s: The implementation and method name did not match.  Use GetGrantedTags().HasAny() to check against the tags this GameplayEffect will Grant to the Actor."), __func__, *GetName());
 	return IGameplayTagAssetInterface::HasAnyMatchingGameplayTags(TagContainer);
 }
 
