GAS: FActiveGameplayEffectsContainer::CheckDuration will no longer remove Infinite Gameplay Effects.

Inspired by PR #11850: Calling CheckDuration(...) with infinite duration gameplay effect removes it

#tests Changed the AbilitySystem.Effect.Remove cheat command to use CheckDuration and verified it removes an Infinite GE before this change (and not after). General Lyra testing.

diff --git a/Engine/Plugins/Runtime/GameplayAbilities/Source/GameplayAbilities/Private/GameplayEffect.cpp b/Engine/Plugins/Runtime/GameplayAbilities/Source/GameplayAbilities/Private/GameplayEffect.cpp
--- a/Engine/Plugins/Runtime/GameplayAbilities/Source/GameplayAbilities/Private/GameplayEffect.cpp
+++ b/Engine/Plugins/Runtime/GameplayAbilities/Source/GameplayAbilities/Private/GameplayEffect.cpp
@@ -5006,110 +5006,119 @@ void FActiveGameplayEffectsContainer::CheckDuration(FActiveGameplayEffectHandle
 	for (int32 ActiveGEIdx = 0; ActiveGEIdx < GameplayEffects_Internal.Num(); ++ActiveGEIdx)
 	{
 		FActiveGameplayEffect& Effect = GameplayEffects_Internal[ActiveGEIdx];
-		if (Effect.Handle == Handle)
+		if (Effect.Handle != Handle)
 		{
-			if (Effect.IsPendingRemove)
-			{
-				// break is this effect is pending remove. 
-				// (Note: don't combine this with the above if statement that is looking for the effect via handle, since we want to stop iteration if we find a matching handle but are pending remove).
-				break;
-			}
+			continue;
+		}
 
-			FTimerManager& TimerManager = Owner->GetWorld()->GetTimerManager();
+		if (Effect.IsPendingRemove)
+		{
+			// break is this effect is pending remove. 
+			// (Note: don't combine this with the above if statement that is looking for the effect via handle, since we want to stop iteration if we find a matching handle but are pending remove).
+			break;
+		}
 
-			// The duration may have changed since we registered this callback with the timer manager.
-			// Make sure that this effect should really be destroyed now
-			float Duration = Effect.GetDuration();
-			float CurrentTime = GetWorldTime();
+		// The duration may have changed since we registered this callback with the timer manager.
+		// Make sure that this effect should really be destroyed now
+		float Duration = Effect.GetDuration();
 
-			int32 StacksToRemove = -2;
-			bool RefreshStartTime = false;
-			bool RefreshDurationTimer = false;
-			bool CheckForFinalPeriodicExec = false;
+		// If it's an infinite Gameplay Effect, we don't want to potentially remove it (or set any timers)
+		// If it was an instant Gameplay Effect, it would never make it into this container.
+		if (Duration <= 0.0f)
+		{
+			break;
+		}
 
-			if (Duration > 0.f && (((Effect.StartWorldTime + Duration) < CurrentTime) || FMath::IsNearlyZero(CurrentTime - Duration - Effect.StartWorldTime, KINDA_SMALL_NUMBER)))
-			{
-				// Figure out what to do based on the expiration policy
-				switch(Effect.Spec.Def->GetStackExpirationPolicy())
-				{
-				case EGameplayEffectStackingExpirationPolicy::ClearEntireStack:
-					StacksToRemove = -1; // Remove all stacks
-					CheckForFinalPeriodicExec = true;					
-					break;
+		float CurrentTime = GetWorldTime();
 
-				case EGameplayEffectStackingExpirationPolicy::RemoveSingleStackAndRefreshDuration:
-					StacksToRemove = 1;
-					CheckForFinalPeriodicExec = (Effect.Spec.GetStackCount() == 1);
-					RefreshStartTime = true;
-					RefreshDurationTimer = true;
-					break;
-				case EGameplayEffectStackingExpirationPolicy::RefreshDuration:
-					RefreshStartTime = true;
-					RefreshDurationTimer = true;
-					break;
-				};					
-			}
-			else
+		int32 StacksToRemove = -2;
+		bool RefreshStartTime = false;
+		bool RefreshDurationTimer = false;
+		bool CheckForFinalPeriodicExec = false;
+
+		if (((Effect.StartWorldTime + Duration) < CurrentTime) || FMath::IsNearlyZero(CurrentTime - Duration - Effect.StartWorldTime, KINDA_SMALL_NUMBER))
+		{
+			// Figure out what to do based on the expiration policy
+			switch(Effect.Spec.Def->GetStackExpirationPolicy())
 			{
-				// Effect isn't finished, just refresh its duration timer
+			case EGameplayEffectStackingExpirationPolicy::ClearEntireStack:
+				StacksToRemove = -1; // Remove all stacks
+				CheckForFinalPeriodicExec = true;					
+				break;
+
+			case EGameplayEffectStackingExpirationPolicy::RemoveSingleStackAndRefreshDuration:
+				StacksToRemove = 1;
+				CheckForFinalPeriodicExec = (Effect.Spec.GetStackCount() == 1);
+				RefreshStartTime = true;
 				RefreshDurationTimer = true;
-			}
+				break;
+			case EGameplayEffectStackingExpirationPolicy::RefreshDuration:
+				RefreshStartTime = true;
+				RefreshDurationTimer = true;
+				break;
+			};					
+		}
+		else
+		{
+			// Effect isn't finished, just refresh its duration timer
+			RefreshDurationTimer = true;
+		}
 
-			if (CheckForFinalPeriodicExec)
+		FTimerManager& TimerManager = Owner->GetWorld()->GetTimerManager();
+		if (CheckForFinalPeriodicExec)
+		{
+			// This gameplay effect has hit its duration. Check if it needs to execute one last time before removing it.
+			if (Effect.PeriodHandle.IsValid() && TimerManager.TimerExists(Effect.PeriodHandle))
 			{
-				// This gameplay effect has hit its duration. Check if it needs to execute one last time before removing it.
-				if (Effect.PeriodHandle.IsValid() && TimerManager.TimerExists(Effect.PeriodHandle))
+				float PeriodTimeRemaining = TimerManager.GetTimerRemaining(Effect.PeriodHandle);
+				if (PeriodTimeRemaining <= KINDA_SMALL_NUMBER && !Effect.bIsInhibited)
 				{
-					float PeriodTimeRemaining = TimerManager.GetTimerRemaining(Effect.PeriodHandle);
-					if (PeriodTimeRemaining <= KINDA_SMALL_NUMBER && !Effect.bIsInhibited)
-					{
-						InternalExecutePeriodicGameplayEffect(Effect);
+					InternalExecutePeriodicGameplayEffect(Effect);
 
-						// The call to ExecuteActiveEffectsFrom in InternalExecutePeriodicGameplayEffect could cause this effect to be explicitly removed
-						// (for example it could kill the owner and cause the effect to be wiped via death).
-						// In that case, we need to early out instead of possibly continuing to the below calls to InternalRemoveActiveGameplayEffect
-						if ( Effect.IsPendingRemove )
-						{
-							break;
-						}
+					// The call to ExecuteActiveEffectsFrom in InternalExecutePeriodicGameplayEffect could cause this effect to be explicitly removed
+					// (for example it could kill the owner and cause the effect to be wiped via death).
+					// In that case, we need to early out instead of possibly continuing to the below calls to InternalRemoveActiveGameplayEffect
+					if ( Effect.IsPendingRemove )
+					{
+						break;
 					}
-
-					// Forcibly clear the periodic ticks because this effect is going to be removed
-					TimerManager.ClearTimer(Effect.PeriodHandle);
 				}
-			}
 
-			if (StacksToRemove >= -1)
-			{
-				InternalRemoveActiveGameplayEffect(ActiveGEIdx, StacksToRemove, false);
+				// Forcibly clear the periodic ticks because this effect is going to be removed
+				TimerManager.ClearTimer(Effect.PeriodHandle);
 			}
+		}
 
-			if (RefreshStartTime)
-			{
-				RestartActiveGameplayEffectDuration(Effect);
-			}
+		if (StacksToRemove >= -1)
+		{
+			InternalRemoveActiveGameplayEffect(ActiveGEIdx, StacksToRemove, false);
+		}
 
-			if (RefreshDurationTimer)
-			{
-				// Always reset the timer, since the duration might have been modified
-				FTimerDelegate Delegate = FTimerDelegate::CreateUObject(Owner, &UAbilitySystemComponent::CheckDurationExpired, Effect.Handle);
+		if (RefreshStartTime)
+		{
+			RestartActiveGameplayEffectDuration(Effect);
+		}
+
+		if (RefreshDurationTimer)
+		{
+			// Always reset the timer, since the duration might have been modified
+			FTimerDelegate Delegate = FTimerDelegate::CreateUObject(Owner, &UAbilitySystemComponent::CheckDurationExpired, Effect.Handle);
 
-				float NewTimerDuration = (Effect.StartWorldTime + Duration) - CurrentTime;
-				TimerManager.SetTimer(Effect.DurationHandle, Delegate, NewTimerDuration, false);
+			float NewTimerDuration = (Effect.StartWorldTime + Duration) - CurrentTime;
+			TimerManager.SetTimer(Effect.DurationHandle, Delegate, NewTimerDuration, false);
 
-				if (Effect.DurationHandle.IsValid() == false)
+			if (Effect.DurationHandle.IsValid() == false)
+			{
+				UE_LOG(LogGameplayEffects, Warning, TEXT("Failed to set new timer in ::CheckDuration. Timer trying to be set for: %.2f. Removing GE instead"), NewTimerDuration);
+				if (!Effect.IsPendingRemove)
 				{
-					UE_LOG(LogGameplayEffects, Warning, TEXT("Failed to set new timer in ::CheckDuration. Timer trying to be set for: %.2f. Removing GE instead"), NewTimerDuration);
-					if (!Effect.IsPendingRemove)
-					{
-						InternalRemoveActiveGameplayEffect(ActiveGEIdx, -1, false);
-					}
-					check(Effect.IsPendingRemove);
+					InternalRemoveActiveGameplayEffect(ActiveGEIdx, -1, false);
 				}
+				check(Effect.IsPendingRemove);
 			}
-
-			break;
 		}
+
+		break;
 	}
 }
 
