PR #11857: Fix unbinding OnMontageBlendedInEnded delegate.

diff --git a/Engine/Plugins/Runtime/GameplayAbilities/Source/GameplayAbilities/Private/Abilities/Tasks/AbilityTask_PlayAnimAndWait.cpp b/Engine/Plugins/Runtime/GameplayAbilities/Source/GameplayAbilities/Private/Abilities/Tasks/AbilityTask_PlayAnimAndWait.cpp
--- a/Engine/Plugins/Runtime/GameplayAbilities/Source/GameplayAbilities/Private/Abilities/Tasks/AbilityTask_PlayAnimAndWait.cpp
+++ b/Engine/Plugins/Runtime/GameplayAbilities/Source/GameplayAbilities/Private/Abilities/Tasks/AbilityTask_PlayAnimAndWait.cpp
@@ -232,6 +232,7 @@ bool UAbilityTask_PlayAnimAndWait::StopPlayingMontage()
 			if (MontageInstance)
 			{
 				MontageInstance->OnMontageBlendingOutStarted.Unbind();
+				MontageInstance->OnMontageBlendedInEnded.Unbind();
 				MontageInstance->OnMontageEnded.Unbind();
 			}
 
@@ -267,4 +268,4 @@ FString UAbilityTask_PlayAnimAndWait::GetDebugString() const
 		return FString::Printf(TEXT("PlayAnimAndWait. AnimSequenceToPlay: %s  (Currently Playing): None"), *GetNameSafe(AnimSequenceToPlay));
 	}
 	
-}
\ No newline at end of file
+}
