PR #11263: Add OnMontageBlendedIn support for PlayMontageAndWait ability task

diff --git a/Engine/Plugins/Runtime/GameplayAbilities/Source/GameplayAbilities/Private/Abilities/Tasks/AbilityTask_PlayMontageAndWait.cpp b/Engine/Plugins/Runtime/GameplayAbilities/Source/GameplayAbilities/Private/Abilities/Tasks/AbilityTask_PlayMontageAndWait.cpp
--- a/Engine/Plugins/Runtime/GameplayAbilities/Source/GameplayAbilities/Private/Abilities/Tasks/AbilityTask_PlayMontageAndWait.cpp
+++ b/Engine/Plugins/Runtime/GameplayAbilities/Source/GameplayAbilities/Private/Abilities/Tasks/AbilityTask_PlayMontageAndWait.cpp
@@ -52,6 +52,14 @@ void UAbilityTask_PlayMontageAndWait::OnMontageBlendingOut(UAnimMontage* Montage
 	}
 }
 
+void UAbilityTask_PlayMontageAndWait::OnMontageBlendedIn(UAnimMontage* Montage)
+{
+	if (ShouldBroadcastAbilityTaskDelegates())
+	{
+		OnBlendedIn.Broadcast();
+	}
+}
+
 void UAbilityTask_PlayMontageAndWait::OnMontageInterrupted()
 {
 	// Call the new function
@@ -131,6 +139,9 @@ void UAbilityTask_PlayMontageAndWait::Activate()
 
 				InterruptedHandle = Ability->OnGameplayAbilityCancelled.AddUObject(this, &UAbilityTask_PlayMontageAndWait::OnGameplayAbilityCancelled);
 
+				BlendedInDelegate.BindUObject(this, &UAbilityTask_PlayMontageAndWait::OnMontageBlendedIn);
+				AnimInstance->Montage_SetBlendedInDelegate(BlendedInDelegate, MontageToPlay);
+
 				BlendingOutDelegate.BindUObject(this, &UAbilityTask_PlayMontageAndWait::OnMontageBlendingOut);
 				AnimInstance->Montage_SetBlendingOutDelegate(BlendingOutDelegate, MontageToPlay);
 
@@ -228,6 +239,7 @@ bool UAbilityTask_PlayMontageAndWait::StopPlayingMontage()
 			FAnimMontageInstance* MontageInstance = AnimInstance->GetActiveInstanceForMontage(MontageToPlay);
 			if (MontageInstance)
 			{
+				MontageInstance->OnMontageBlendedInEnded.Unbind();
 				MontageInstance->OnMontageBlendingOutStarted.Unbind();
 				MontageInstance->OnMontageEnded.Unbind();
 			}
diff --git a/Engine/Plugins/Runtime/GameplayAbilities/Source/GameplayAbilities/Public/Abilities/Tasks/AbilityTask_PlayMontageAndWait.h b/Engine/Plugins/Runtime/GameplayAbilities/Source/GameplayAbilities/Public/Abilities/Tasks/AbilityTask_PlayMontageAndWait.h
--- a/Engine/Plugins/Runtime/GameplayAbilities/Source/GameplayAbilities/Public/Abilities/Tasks/AbilityTask_PlayMontageAndWait.h
+++ b/Engine/Plugins/Runtime/GameplayAbilities/Source/GameplayAbilities/Public/Abilities/Tasks/AbilityTask_PlayMontageAndWait.h
@@ -19,6 +19,9 @@ public:
 	UPROPERTY(BlueprintAssignable)
 	FMontageWaitSimpleDelegate	OnCompleted;
 
+	UPROPERTY(BlueprintAssignable)
+	FMontageWaitSimpleDelegate	OnBlendedIn;
+
 	UPROPERTY(BlueprintAssignable)
 	FMontageWaitSimpleDelegate	OnBlendOut;
 
@@ -28,6 +31,9 @@ public:
 	UPROPERTY(BlueprintAssignable)
 	FMontageWaitSimpleDelegate	OnCancelled;
 
+	UFUNCTION()
+	void OnMontageBlendedIn(UAnimMontage* Montage);
+
 	UFUNCTION()
 	void OnMontageBlendingOut(UAnimMontage* Montage, bool bInterrupted);
 
@@ -76,6 +82,7 @@ protected:
 	/** Checks if the ability is playing a montage and stops that montage, returns true if a montage was stopped, false if not. */
 	bool StopPlayingMontage();
 
+	FOnMontageBlendedInEnded BlendedInDelegate;
 	FOnMontageBlendingOutStarted BlendingOutDelegate;
 	FOnMontageEnded MontageEndedDelegate;
 	FDelegateHandle InterruptedHandle;
