Fixes in the Gameplay Ability System handling of deprecated NonInstanced flow causing knock-ons with the input system -- discovered in AncientGame.

Added some warnings about these special cases.  AncientGame should use InstancedPerActor after these changes (InstancedPerExecution is/was hiding some of these bugs).

diff --git a/Engine/Plugins/Runtime/GameplayAbilities/Source/GameplayAbilities/Private/Abilities/Tasks/AbilityTask_StartAbilityState.cpp b/Engine/Plugins/Runtime/GameplayAbilities/Source/GameplayAbilities/Private/Abilities/Tasks/AbilityTask_StartAbilityState.cpp
--- a/Engine/Plugins/Runtime/GameplayAbilities/Source/GameplayAbilities/Private/Abilities/Tasks/AbilityTask_StartAbilityState.cpp
+++ b/Engine/Plugins/Runtime/GameplayAbilities/Source/GameplayAbilities/Private/Abilities/Tasks/AbilityTask_StartAbilityState.cpp
@@ -4,6 +4,16 @@
 
 #include UE_INLINE_GENERATED_CPP_BY_NAME(AbilityTask_StartAbilityState)
 
+namespace UE::AbilitySystem::AbilityTask_StartAbilityState
+{
+	// The default implementation of ShouldBroadcastAbilityTaskDelegates returns false if the Ability is not 'active'.
+	// However, we want to execute anyway because we expect to be able to execute during EndAbility.
+	bool CustomShouldBroadcastAbilityTaskDelegates(const UGameplayAbility* Ability)
+	{
+		return IsValid(Ability);
+	}
+}
+
 UAbilityTask_StartAbilityState::UAbilityTask_StartAbilityState(const FObjectInitializer& ObjectInitializer)
 : Super(ObjectInitializer)
 {
@@ -35,6 +45,8 @@ void UAbilityTask_StartAbilityState::Activate()
 
 void UAbilityTask_StartAbilityState::OnDestroy(bool AbilityEnded)
 {
+	using namespace UE::AbilitySystem::AbilityTask_StartAbilityState;
+
 	// Unbind delegates so this doesn't get recursively called
 	if (Ability)
 	{
@@ -44,14 +56,14 @@ void UAbilityTask_StartAbilityState::OnDestroy(bool AbilityEnded)
 
 	if (bWasInterrupted && OnStateInterrupted.IsBound())
 	{
-		if (ShouldBroadcastAbilityTaskDelegates())
+		if (CustomShouldBroadcastAbilityTaskDelegates(Ability))
 		{
 			OnStateInterrupted.Broadcast();
 		}
 	}
 	else if ((bWasEnded || AbilityEnded) && OnStateEnded.IsBound())
 	{
-		if (ShouldBroadcastAbilityTaskDelegates())
+		if (CustomShouldBroadcastAbilityTaskDelegates(Ability))
 		{
 			OnStateEnded.Broadcast();
 		}
diff --git a/Engine/Plugins/Runtime/GameplayAbilities/Source/GameplayAbilities/Private/AbilitySystemComponent_Abilities.cpp b/Engine/Plugins/Runtime/GameplayAbilities/Source/GameplayAbilities/Private/AbilitySystemComponent_Abilities.cpp
--- a/Engine/Plugins/Runtime/GameplayAbilities/Source/GameplayAbilities/Private/AbilitySystemComponent_Abilities.cpp
+++ b/Engine/Plugins/Runtime/GameplayAbilities/Source/GameplayAbilities/Private/AbilitySystemComponent_Abilities.cpp
@@ -2775,7 +2775,9 @@ void UAbilitySystemComponent::AbilityLocalInputPressed(int32 InputID)
 
 PRAGMA_DISABLE_DEPRECATION_WARNINGS
 					// Fixing this up to use the instance activation, but this function should be deprecated as it cannot work with InstancedPerExecution
-					const FGameplayAbilityActivationInfo& ActivationInfo = Spec.GetPrimaryInstance() ? Spec.GetPrimaryInstance()->GetCurrentActivationInfoRef() : Spec.ActivationInfo;
+					UE_CLOG(Spec.Ability->GetInstancingPolicy() == EGameplayAbilityInstancingPolicy::InstancedPerExecution, LogAbilitySystem, Warning, TEXT("%hs: %s is InstancedPerExecution. This is unreliable for Input as you may only interact with the latest spawned Instance"), __func__, *GetNameSafe(Spec.Ability));
+					TArray<UGameplayAbility*> Instances = Spec.GetAbilityInstances();
+					const FGameplayAbilityActivationInfo& ActivationInfo = Instances.IsEmpty() ? Spec.ActivationInfo : Instances.Last()->GetCurrentActivationInfoRef();
 PRAGMA_ENABLE_DEPRECATION_WARNINGS
 					// Invoke the InputPressed event. This is not replicated here. If someone is listening, they may replicate the InputPressed event to the server.
 					InvokeReplicatedEvent(EAbilityGenericReplicatedEvent::InputPressed, Spec.Handle, ActivationInfo.GetActivationPredictionKey());					
@@ -2809,7 +2811,9 @@ void UAbilitySystemComponent::AbilityLocalInputReleased(int32 InputID)
 
 PRAGMA_DISABLE_DEPRECATION_WARNINGS
 				// Fixing this up to use the instance activation, but this function should be deprecated as it cannot work with InstancedPerExecution
-				const FGameplayAbilityActivationInfo& ActivationInfo = Spec.GetPrimaryInstance() ? Spec.GetPrimaryInstance()->GetCurrentActivationInfoRef() : Spec.ActivationInfo;
+				UE_CLOG(Spec.Ability->GetInstancingPolicy() == EGameplayAbilityInstancingPolicy::InstancedPerExecution, LogAbilitySystem, Warning, TEXT("%hs: %s is InstancedPerExecution. This is unreliable for Input as you may only interact with the latest spawned Instance"), __func__, *GetNameSafe(Spec.Ability));
+				TArray<UGameplayAbility*> Instances = Spec.GetAbilityInstances();
+				const FGameplayAbilityActivationInfo& ActivationInfo = Instances.IsEmpty() ? Spec.ActivationInfo : Instances.Last()->GetCurrentActivationInfoRef();
 PRAGMA_ENABLE_DEPRECATION_WARNINGS
 				InvokeReplicatedEvent(EAbilityGenericReplicatedEvent::InputReleased, Spec.Handle, ActivationInfo.GetActivationPredictionKey());
 			}
