Add check for valid ability system component when creating an effect context.

diff --git a/Engine/Plugins/Runtime/GameplayAbilities/Source/GameplayAbilities/Private/Abilities/GameplayAbility.cpp b/Engine/Plugins/Runtime/GameplayAbilities/Source/GameplayAbilities/Private/Abilities/GameplayAbility.cpp
--- a/Engine/Plugins/Runtime/GameplayAbilities/Source/GameplayAbilities/Private/Abilities/GameplayAbility.cpp
+++ b/Engine/Plugins/Runtime/GameplayAbilities/Source/GameplayAbilities/Private/Abilities/GameplayAbility.cpp
@@ -1866,9 +1866,12 @@ FGameplayEffectContextHandle UGameplayAbility::MakeEffectContext(const FGameplay
 		Context.AddInstigator(ActorInfo->OwnerActor.Get(), ActorInfo->AvatarActor.Get());
 
 		// Pass along the source object to the effect
-		if (FGameplayAbilitySpec* AbilitySpec = ActorInfo->AbilitySystemComponent->FindAbilitySpecFromHandle(Handle))
+		if (UAbilitySystemComponent* AbilitySystemComponent = ActorInfo->AbilitySystemComponent.Get())
 		{
-			Context.AddSourceObject(AbilitySpec->SourceObject.Get());
+			if (FGameplayAbilitySpec* AbilitySpec = AbilitySystemComponent->FindAbilitySpecFromHandle(Handle))
+			{
+				Context.AddSourceObject(AbilitySpec->SourceObject.Get());
+			}
 		}
 	}
 
