GAS: Adding new Gameplay Effect attribute aggregation functionality (to EGameplayModOp): AddFinal and MultiplyCompound. Multiple users have requested this and the belief is that these can often be used instead of Gameplay Effect Modifier Evaluation Channels.

diff --git a/Engine/Plugins/Runtime/GameplayAbilities/Source/GameplayAbilities/Private/GameplayEffectAggregator.cpp b/Engine/Plugins/Runtime/GameplayAbilities/Source/GameplayAbilities/Private/GameplayEffectAggregator.cpp
--- a/Engine/Plugins/Runtime/GameplayAbilities/Source/GameplayAbilities/Private/GameplayEffectAggregator.cpp
+++ b/Engine/Plugins/Runtime/GameplayAbilities/Source/GameplayAbilities/Private/GameplayEffectAggregator.cpp
@@ -7,6 +7,24 @@
 #include "AbilitySystemLog.h"
 #include "AbilitySystemStats.h"
 
+namespace UE::AbilitySystem::Private
+{
+	float MultiplyMods(const TArray<FAggregatorMod>& InMods)
+	{
+		float Result = 1.0f;
+
+		for (const FAggregatorMod& Mod : InMods)
+		{
+			if (Mod.Qualifies())
+			{
+				Result *= Mod.EvaluatedMagnitude;
+			}
+		}
+
+		return Result;
+	}
+}
+
 void FAggregatorMod::UpdateQualifies(const FAggregatorEvaluateParameters& Parameters) const
 {
 	static const FGameplayTagContainer EmptyTagContainer;
@@ -68,6 +86,8 @@ float FAggregatorModChannel::EvaluateWithBase(float InlineBaseValue, const FAggr
 	float Additive = SumMods(Mods[EGameplayModOp::Additive], GameplayEffectUtilities::GetModifierBiasByModifierOp(EGameplayModOp::Additive), Parameters);
 	float Multiplicitive = SumMods(Mods[EGameplayModOp::Multiplicitive], GameplayEffectUtilities::GetModifierBiasByModifierOp(EGameplayModOp::Multiplicitive), Parameters);
 	float Division = SumMods(Mods[EGameplayModOp::Division], GameplayEffectUtilities::GetModifierBiasByModifierOp(EGameplayModOp::Division), Parameters);
+	float FinalAdd = SumMods(Mods[EGameplayModOp::AddFinal], GameplayEffectUtilities::GetModifierBiasByModifierOp(EGameplayModOp::AddFinal), Parameters);
+	float CompoundMultiply = UE::AbilitySystem::Private::MultiplyMods(Mods[EGameplayModOp::MultiplyCompound]);
 
 	if (FMath::IsNearlyZero(Division))
 	{
@@ -75,7 +95,7 @@ float FAggregatorModChannel::EvaluateWithBase(float InlineBaseValue, const FAggr
 		Division = 1.f;
 	}
 
-	return ((InlineBaseValue + Additive) * Multiplicitive) / Division;
+	return ((InlineBaseValue + Additive) * Multiplicitive / Division * CompoundMultiply) + FinalAdd;
 }
 
 bool FAggregatorModChannel::ReverseEvaluate(float FinalValue, const FAggregatorEvaluateParameters& Parameters, OUT float& ComputedValue) const
@@ -93,6 +113,8 @@ bool FAggregatorModChannel::ReverseEvaluate(float FinalValue, const FAggregatorE
 	float Additive = SumMods(Mods[EGameplayModOp::Additive], GameplayEffectUtilities::GetModifierBiasByModifierOp(EGameplayModOp::Additive), Parameters);
 	float Multiplicitive = SumMods(Mods[EGameplayModOp::Multiplicitive], GameplayEffectUtilities::GetModifierBiasByModifierOp(EGameplayModOp::Multiplicitive), Parameters);
 	float Division = SumMods(Mods[EGameplayModOp::Division], GameplayEffectUtilities::GetModifierBiasByModifierOp(EGameplayModOp::Division), Parameters);
+	float FinalAdd = SumMods(Mods[EGameplayModOp::AddFinal], GameplayEffectUtilities::GetModifierBiasByModifierOp(EGameplayModOp::AddFinal), Parameters);
+	float CompoundMultiply = UE::AbilitySystem::Private::MultiplyMods(Mods[EGameplayModOp::MultiplyCompound]);
 
 	if (FMath::IsNearlyZero(Division))
 	{
@@ -106,7 +128,7 @@ bool FAggregatorModChannel::ReverseEvaluate(float FinalValue, const FAggregatorE
 		return false;
 	}
 
-	ComputedValue = (FinalValue * Division / Multiplicitive) - Additive;
+	ComputedValue = ((FinalValue - FinalAdd) / CompoundMultiply * Division / Multiplicitive) - Additive;
 	return true;
 }
 
@@ -431,17 +453,19 @@ float FAggregator::StaticExecModOnBaseValue(float BaseValue, TEnumAsByte<EGamepl
 			BaseValue = EvaluatedMagnitude;
 			break;
 		}
-		case EGameplayModOp::Additive:
+		case EGameplayModOp::AddBase:
+		case EGameplayModOp::AddFinal:
 		{
 			BaseValue += EvaluatedMagnitude;
 			break;
 		}
-		case EGameplayModOp::Multiplicitive:
+		case EGameplayModOp::MultiplyAdditive:
+		case EGameplayModOp::MultiplyCompound:
 		{
 			BaseValue *= EvaluatedMagnitude;
 			break;
 		}
-		case EGameplayModOp::Division:
+		case EGameplayModOp::DivideAdditive:
 		{
 			if (FMath::IsNearlyZero(EvaluatedMagnitude) == false)
 			{
diff --git a/Engine/Plugins/Runtime/GameplayAbilities/Source/GameplayAbilities/Private/GameplayEffectAggregatorLibrary.cpp b/Engine/Plugins/Runtime/GameplayAbilities/Source/GameplayAbilities/Private/GameplayEffectAggregatorLibrary.cpp
--- a/Engine/Plugins/Runtime/GameplayAbilities/Source/GameplayAbilities/Private/GameplayEffectAggregatorLibrary.cpp
+++ b/Engine/Plugins/Runtime/GameplayAbilities/Source/GameplayAbilities/Private/GameplayEffectAggregatorLibrary.cpp
@@ -24,13 +24,15 @@ void QualifierFunc_MostNegativeMod_AllPositiveMods(const FAggregatorEvaluatePara
 		float ExpectedDelta = 0.f;
 		switch( ModInfo.Op )
 		{
-		case EGameplayModOp::Additive:
+		case EGameplayModOp::AddBase:
+		case EGameplayModOp::AddFinal:
 			ExpectedDelta = ModInfo.Mod->EvaluatedMagnitude;
 			break;
-		case EGameplayModOp::Multiplicitive:
+		case EGameplayModOp::MultiplyAdditive:
+		case EGameplayModOp::MultiplyCompound:
 			ExpectedDelta = (BaseValue * ModInfo.Mod->EvaluatedMagnitude) - BaseValue;
 			break;
-		case EGameplayModOp::Division:
+		case EGameplayModOp::DivideAdditive:
 			ExpectedDelta = ModInfo.Mod->EvaluatedMagnitude > 0.f ? ((BaseValue / ModInfo.Mod->EvaluatedMagnitude) - BaseValue) : 0.f;
 			break;
 		case EGameplayModOp::Override:
diff --git a/Engine/Plugins/Runtime/GameplayAbilities/Source/GameplayAbilities/Private/GameplayEffectTypes.cpp b/Engine/Plugins/Runtime/GameplayAbilities/Source/GameplayAbilities/Private/GameplayEffectTypes.cpp
--- a/Engine/Plugins/Runtime/GameplayAbilities/Source/GameplayAbilities/Private/GameplayEffectTypes.cpp
+++ b/Engine/Plugins/Runtime/GameplayAbilities/Source/GameplayAbilities/Private/GameplayEffectTypes.cpp
@@ -103,7 +103,7 @@ bool FGameplayModEvaluationChannelSettings::operator!=(const FGameplayModEvaluat
 
 float GameplayEffectUtilities::GetModifierBiasByModifierOp(EGameplayModOp::Type ModOp)
 {
-	static const float ModifierOpBiases[EGameplayModOp::Max] = {0.f, 1.f, 1.f, 0.f};
+	static const float ModifierOpBiases[EGameplayModOp::Max] = {0.f, 1.f, 1.f, 0.f, 0.f, 0.f};
 	check(ModOp >= 0 && ModOp < EGameplayModOp::Max);
 
 	return ModifierOpBiases[ModOp];
diff --git a/Engine/Plugins/Runtime/GameplayAbilities/Source/GameplayAbilities/Private/Tests/GameplayEffectTests.cpp b/Engine/Plugins/Runtime/GameplayAbilities/Source/GameplayAbilities/Private/Tests/GameplayEffectTests.cpp
--- a/Engine/Plugins/Runtime/GameplayAbilities/Source/GameplayAbilities/Private/Tests/GameplayEffectTests.cpp
+++ b/Engine/Plugins/Runtime/GameplayAbilities/Source/GameplayAbilities/Private/Tests/GameplayEffectTests.cpp
@@ -94,7 +94,7 @@ public:
 		SourceComponent = SourceActor->GetAbilitySystemComponent();
 		SourceComponent->GetSet<UAbilitySystemTestAttributeSet>()->Health = StartingHealth;
 		SourceComponent->GetSet<UAbilitySystemTestAttributeSet>()->MaxHealth = StartingHealth;
-		SourceComponent->GetSet<UAbilitySystemTestAttributeSet>()->Mana = StartingMana;
+		SourceComponent->GetSet<UAbilitySystemTestAttributeSet>()->Mana = FGameplayAttributeData(StartingMana);
 		SourceComponent->GetSet<UAbilitySystemTestAttributeSet>()->MaxMana = StartingMana;
 
 		// set up the destination actor
@@ -102,7 +102,7 @@ public:
 		DestComponent = DestActor->GetAbilitySystemComponent();
 		DestComponent->GetSet<UAbilitySystemTestAttributeSet>()->Health = StartingHealth;
 		DestComponent->GetSet<UAbilitySystemTestAttributeSet>()->MaxHealth = StartingHealth;
-		DestComponent->GetSet<UAbilitySystemTestAttributeSet>()->Mana = StartingMana;
+		DestComponent->GetSet<UAbilitySystemTestAttributeSet>()->Mana = FGameplayAttributeData(StartingMana);
 		DestComponent->GetSet<UAbilitySystemTestAttributeSet>()->MaxMana = StartingMana;
 	}
 
@@ -166,7 +166,7 @@ public: // the tests
 	void Test_ManaBuff()
 	{
 		const float BuffValue = 30.f;
-		const float StartingMana = DestComponent->GetSet<UAbilitySystemTestAttributeSet>()->Mana;
+		const float StartingMana = DestComponent->GetSet<UAbilitySystemTestAttributeSet>()->Mana.GetCurrentValue();
 
 		FActiveGameplayEffectHandle BuffHandle;
 
@@ -180,7 +180,7 @@ public: // the tests
 		}
 
 		// check that the value changed
-		Test->TestEqual(TEXT("Mana Buffed"), DestComponent->GetSet<UAbilitySystemTestAttributeSet>()->Mana, StartingMana + BuffValue);
+		Test->TestEqual(TEXT("Mana Buffed"), DestComponent->GetSet<UAbilitySystemTestAttributeSet>()->Mana.GetCurrentValue(), StartingMana + BuffValue);
 
 		// remove the effect
 		{
@@ -188,7 +188,85 @@ public: // the tests
 		}
 
 		// check that the value changed back
-		Test->TestEqual(TEXT("Mana Restored"), DestComponent->GetSet<UAbilitySystemTestAttributeSet>()->Mana, StartingMana);
+		Test->TestEqual(TEXT("Mana Restored"), DestComponent->GetSet<UAbilitySystemTestAttributeSet>()->Mana.GetCurrentValue(), StartingMana);
+	}
+
+	void Test_AttributeAggregators()
+	{
+		constexpr float BuffValue = 2.0f;
+		const float ManaBaseValue = DestComponent->GetSet<UAbilitySystemTestAttributeSet>()->Mana.GetBaseValue();
+
+		// Define a common lambda for applying a GE with a specific GameplayModOp
+		auto ApplyGameplayModOp = [this, BuffValue](EGameplayModOp::Type GameplayModOp)
+		{
+			CONSTRUCT_CLASS(UGameplayEffect, DamageBuffEffect);
+			AddModifier(DamageBuffEffect, GET_FIELD_CHECKED(UAbilitySystemTestAttributeSet, Mana), GameplayModOp, FScalableFloat(BuffValue));
+			DamageBuffEffect->DurationPolicy = EGameplayEffectDurationType::Infinite;
+
+			return SourceComponent->ApplyGameplayEffectToTarget(DamageBuffEffect, DestComponent, 1.f);
+		};
+
+		// Define a common lambda for testing if an op applied correctly
+		auto TestGameplayModOp = [this, &ApplyGameplayModOp, BuffValue](EGameplayModOp::Type GameplayModOp)
+		{
+			const float PrevValue = DestComponent->GetSet<UAbilitySystemTestAttributeSet>()->Mana.GetCurrentValue();
+			const float ExpectedValue = FAggregator::StaticExecModOnBaseValue(PrevValue, GameplayModOp, BuffValue);
+
+			FActiveGameplayEffectHandle AGEHandle = ApplyGameplayModOp(GameplayModOp);
+			const float CurrentValue = DestComponent->GetSet<UAbilitySystemTestAttributeSet>()->Mana.GetCurrentValue();
+			Test->TestEqual(FString::Printf(TEXT("Attribute GameplayModOp %s"), *UEnum::GetValueAsString(GameplayModOp)), CurrentValue, ExpectedValue);
+			
+			return AGEHandle;
+		};
+
+		// Test all of the ops (order matters here due to the implementation of the above StaticExecModOnBaseValue)
+		FActiveGameplayEffectHandle AdditiveHandle = TestGameplayModOp(EGameplayModOp::Additive);
+		FActiveGameplayEffectHandle MultiplicativeHandle = TestGameplayModOp(EGameplayModOp::Multiplicitive);
+		FActiveGameplayEffectHandle DivisionHandle = TestGameplayModOp(EGameplayModOp::Division);
+		FActiveGameplayEffectHandle CompoundHandle = TestGameplayModOp(EGameplayModOp::MultiplyCompound);
+		FActiveGameplayEffectHandle FinalAddHandle = TestGameplayModOp(EGameplayModOp::AddFinal);
+
+		// Test the override quickly here before testing the aggregation is as expected
+		{
+			FActiveGameplayEffectHandle OverrideHandle = TestGameplayModOp(EGameplayModOp::Override);
+			DestComponent->RemoveActiveGameplayEffect(OverrideHandle);
+		}
+
+		Test->TestEqual(TEXT("Mana BaseValue is Unchanged"), DestComponent->GetSet<UAbilitySystemTestAttributeSet>()->Mana.GetBaseValue(), ManaBaseValue);
+
+		// Add some compounding GameplayEffects and manually test our equation to see if it's all compounded correctly.
+		{
+			FActiveGameplayEffectHandle BaseAdd2 = ApplyGameplayModOp(EGameplayModOp::Additive);
+			FActiveGameplayEffectHandle BaseMultiply2 = ApplyGameplayModOp(EGameplayModOp::Multiplicitive);
+			FActiveGameplayEffectHandle Compound2 = ApplyGameplayModOp(EGameplayModOp::MultiplyCompound);
+			FActiveGameplayEffectHandle FinalAdd2 = ApplyGameplayModOp(EGameplayModOp::AddFinal);
+			
+			// ExpectedResult = ((BaseValue + Additive) * Multiplicative / Division * CompoundMultiply) + FinalAdd;
+			// Multiplicative and Division Compound as: 1.0f + ForEachValue(Value - 1.0f). E.g. two applications of 1.5 = 2.
+			// CompoundMultiply Compounds as: 1.0f *= ForEachValue(Value) e.g. two applications of 1.5 = 2.25.
+			constexpr float X = BuffValue;
+			const float ExpectedResult = ((ManaBaseValue + X + X) * (1.0f + (X - 1.0f) + (X - 1.0f)) / (1.0f + (X - 1.0f)) * (X * X)) + X + X;
+			const float CurrentValue = DestComponent->GetSet<UAbilitySystemTestAttributeSet>()->Mana.GetCurrentValue();
+			Test->TestEqual(TEXT("Aggregation Equation Works as Expected"), CurrentValue, ExpectedResult);
+
+			// Remove the compounding GE's
+			DestComponent->RemoveActiveGameplayEffect(FinalAdd2);
+			DestComponent->RemoveActiveGameplayEffect(Compound2);
+			DestComponent->RemoveActiveGameplayEffect(BaseMultiply2);
+			DestComponent->RemoveActiveGameplayEffect(BaseAdd2);
+		}
+
+		// Remove the base GE's
+		{
+			DestComponent->RemoveActiveGameplayEffect(FinalAddHandle);
+			DestComponent->RemoveActiveGameplayEffect(CompoundHandle);
+			DestComponent->RemoveActiveGameplayEffect(DivisionHandle);
+			DestComponent->RemoveActiveGameplayEffect(MultiplicativeHandle);
+			DestComponent->RemoveActiveGameplayEffect(AdditiveHandle);
+		}
+
+		// check that the value changed back
+		Test->TestEqual(TEXT("Mana Restored to BaseValue"), DestComponent->GetSet<UAbilitySystemTestAttributeSet>()->Mana.GetCurrentValue(), ManaBaseValue);
 	}
 
 	void Test_PeriodicDamage()
@@ -712,6 +790,7 @@ public:
 		ADD_TEST(Test_InstantDamage);
 		ADD_TEST(Test_InstantDamageRemap);
 		ADD_TEST(Test_ManaBuff);
+		ADD_TEST(Test_AttributeAggregators);
 		ADD_TEST(Test_PeriodicDamage);
 		ADD_TEST(Test_StackLimit);
 		ADD_TEST(Test_SetByCallerStackingDuration);
diff --git a/Engine/Plugins/Runtime/GameplayAbilities/Source/GameplayAbilities/Public/AbilitySystemTestAttributeSet.h b/Engine/Plugins/Runtime/GameplayAbilities/Source/GameplayAbilities/Public/AbilitySystemTestAttributeSet.h
--- a/Engine/Plugins/Runtime/GameplayAbilities/Source/GameplayAbilities/Public/AbilitySystemTestAttributeSet.h
+++ b/Engine/Plugins/Runtime/GameplayAbilities/Source/GameplayAbilities/Public/AbilitySystemTestAttributeSet.h
@@ -12,8 +12,6 @@ class GAMEPLAYABILITIES_API UAbilitySystemTestAttributeSet : public UAttributeSe
 {
 	GENERATED_UCLASS_BODY()
 
-
-
 	/**
 	 *	NOTE ON MUTABLE:
 	 *	This is only done so that UAbilitySystemTestAttributeSet can be initialized directly in GameplayEffectsTest.cpp without doing a const_cast in 100+ places.
@@ -28,7 +26,7 @@ class GAMEPLAYABILITIES_API UAbilitySystemTestAttributeSet : public UAttributeSe
 	mutable float	Health;
 
 	UPROPERTY(EditAnywhere, BlueprintReadWrite, Replicated, Category = "AttributeTest")
-	mutable float	Mana;
+	mutable FGameplayAttributeData	Mana;
 
 	UPROPERTY(EditAnywhere, BlueprintReadWrite, Replicated, Category = "AttributeTest")
 	mutable float	MaxMana;
diff --git a/Engine/Plugins/Runtime/GameplayAbilities/Source/GameplayAbilities/Public/GameplayEffect.h b/Engine/Plugins/Runtime/GameplayAbilities/Source/GameplayAbilities/Public/GameplayEffect.h
--- a/Engine/Plugins/Runtime/GameplayAbilities/Source/GameplayAbilities/Public/GameplayEffect.h
+++ b/Engine/Plugins/Runtime/GameplayAbilities/Source/GameplayAbilities/Public/GameplayEffect.h
@@ -534,7 +534,11 @@ struct GAMEPLAYABILITIES_API FGameplayModifierInfo
 	UPROPERTY(EditDefaultsOnly, Category=GameplayModifier, meta=(FilterMetaTag="HideFromModifiers"))
 	FGameplayAttribute Attribute;
 
-	/** The numeric operation of this modifier: Override, Add, Multiply, etc  */
+	/**
+	 * The numeric operation of this modifier: Override, Add, Multiply, etc
+	 * When multiple modifiers aggregate together, the equation is:
+	 * ((BaseValue + AddBase) * MultiplyAdditive / DivideAdditive * MultiplyCompound) + AddFinal
+	 */
 	UPROPERTY(EditDefaultsOnly, Category=GameplayModifier)
 	TEnumAsByte<EGameplayModOp::Type> ModifierOp = EGameplayModOp::Additive;
 
diff --git a/Engine/Plugins/Runtime/GameplayAbilities/Source/GameplayAbilities/Public/GameplayEffectTypes.h b/Engine/Plugins/Runtime/GameplayAbilities/Source/GameplayAbilities/Public/GameplayEffectTypes.h
--- a/Engine/Plugins/Runtime/GameplayAbilities/Source/GameplayAbilities/Public/GameplayEffectTypes.h
+++ b/Engine/Plugins/Runtime/GameplayAbilities/Source/GameplayAbilities/Public/GameplayEffectTypes.h
@@ -108,21 +108,41 @@ protected:
 UENUM(BlueprintType)
 namespace EGameplayModOp
 {
-	/** Defines the ways that mods will modify attributes. Numeric ones operate on the existing value, override ignores it */
+	/**
+	 * Defines the ways that mods will modify attributes. Values of the same type are aggregated, and then applied in the following equation:
+	 * ((BaseValue + AddBase) * MultiplyAdditive / DivideAdditive * MultiplyCompound) + AddFinal
+	 */
 	enum Type : int
 	{		
-		/** Numeric. */
-		Additive = 0		UMETA(DisplayName="Add"),
-		/** Numeric. */
-		Multiplicitive		UMETA(DisplayName = "Multiply"),
-		/** Numeric. */
-		Division			UMETA(DisplayName = "Divide"),
-
-		/** Other. */
-		Override 			UMETA(DisplayName="Override"),	// This should always be the first non numeric ModOp
-
-		// This must always be at the end.
-		Max					UMETA(DisplayName="Invalid")
+		/** Adds to the Base value. This happens first, before all other mods are considered. */
+		AddBase				UMETA(DisplayName="Add (Base)"),
+
+		/** Multipliers are added together first, then multiplied against prev result. E.g. 50% + 50% = 100% in values is 1.5 + 1.5 = 2.0. */
+		MultiplyAdditive	UMETA(DisplayName = "Multiply (Additive)"),
+
+		/** Divisors are added together, then divided against the prev result. E.g. 1/2 + 1/2 = 1/3 in values is 2 + 2 = 3. */
+		DivideAdditive		UMETA(DisplayName = "Divide (Additive)"),
+
+		// !!! Don't get tripped up here:  Override is 3 (due to existing serialized assets) !!!
+
+		/** Multiply the prev result by this value. E.g. two values of 1.5 compounded: 1.5 * 1.5 = 2.25. */
+		MultiplyCompound = 4 UMETA(DisplayName = "Multiply (Compound)"),
+
+		/** Add this value to the final computed result. */
+		AddFinal			UMETA(DisplayName = "Add (Final)"),
+		
+		/** This must always be the last value (used in iteration code). */
+		Max					UMETA(Hidden, DisplayName="Invalid"),
+
+		// Backwards compatible names
+		Additive = 0		UMETA(Hidden), // Backwards compat name
+		Multiplicitive = 1	 UMETA(Hidden), // Backwards compat name
+		Division = 2		UMETA(Hidden), // Backwards compat name
+
+		// Due to UI being displayed in the same order as the operations, we want Override at the bottom.
+
+		/** Override the value, regardless of what the computation provides. */
+		Override = 3		UMETA(DisplayName="Override"),
 	};
 }
 
