GAS: Introduce some GameplayCue debugging around replication.  This introduces a LogGameplayCues category which initially shows nothing, unless you set it to Verbose.

diff --git a/Engine/Plugins/Runtime/GameplayAbilities/Source/GameplayAbilities/Private/GameplayCueInterface.cpp b/Engine/Plugins/Runtime/GameplayAbilities/Source/GameplayAbilities/Private/GameplayCueInterface.cpp
--- a/Engine/Plugins/Runtime/GameplayAbilities/Source/GameplayAbilities/Private/GameplayCueInterface.cpp
+++ b/Engine/Plugins/Runtime/GameplayAbilities/Source/GameplayAbilities/Private/GameplayCueInterface.cpp
@@ -10,6 +10,7 @@
 
 #include UE_INLINE_GENERATED_CPP_BY_NAME(GameplayCueInterface)
 
+DEFINE_LOG_CATEGORY_STATIC(LogGameplayCues, Log, All);
 
 namespace GameplayCueInterfacePrivate
 {
@@ -23,6 +24,58 @@ namespace GameplayCueInterfacePrivate
 
 	bool bUseEqualTagCountAndRemovalCallbacks = true;
 	static FAutoConsoleVariableRef CVarUseEqualTagCountAndRemovalCallbacks(TEXT("GameplayCue.Fix.UseEqualTagCountAndRemovalCallbacks"), bUseEqualTagCountAndRemovalCallbacks, TEXT("Default: True. When calling RemoveCue, get an equal number of TagCountUpdated callbacks and Removal callbacks rather than 1:many"));
+
+	/**
+	 * Given a pointer to a UProperty and the owning UObject, attempt to find the reflected FName of the property.
+	 * @Return the property's FName as a string (or a debug string if the parameters are invalid)
+	 * Example:
+	 *
+	 * UCLASS() class UMyClass { UPROPERTY() float MyName; } MyClass;
+	 * FindNameOfProperty(&MyClass->MyName, MyClass) returns "MyName".
+	 * 
+	 * This is different than GET_MEMBER_NAME_CHECKED because it is done at runtime
+	 * (e.g. you have a float* but do not know which field it belongs to).
+	 * Note: This function isn't exactly fast and should only be used in debug situations.
+	 */
+	FString FindNameOfProperty(const void* PointerToProperty, const UObject* PropertyOwner)
+	{
+		if (!PropertyOwner)
+		{
+			return TEXT("(Null Owner)");
+		}
+
+		// Compute the offset and double-check that offset falls within the memory range of PropertyOwner.
+		ptrdiff_t ThisOffsetFromOwner = (intptr_t(PointerToProperty) - intptr_t(PropertyOwner));
+
+		const UClass* OwnerClass = PropertyOwner->GetClass();
+		if (ThisOffsetFromOwner < 0 || ThisOffsetFromOwner >= OwnerClass->GetStructureSize())
+		{
+			// This offset is bad
+			return TEXT("(Bad Offset)");
+		}
+
+		// Go through all of the properties and find one with the correct offset
+		for (FProperty* Property = PropertyOwner->GetClass()->PropertyLink; Property; Property = Property->PropertyLinkNext)
+		{
+			int32 PropertyOffset = Property->GetOffset_ForDebug();
+			if (PropertyOffset == ThisOffsetFromOwner)
+			{
+				return Property->GetName();
+			}
+		}
+
+		return TEXT("(No UProperty)");
+	}
+
+	/** Get the owner of the FActiveGameplayCueContainer (which is an FStruct so does not have an Outer) */
+	const UObject* GetPropertyOwner(const FActiveGameplayCueContainer& InArray)
+	{
+#if WITH_PUSH_MODEL
+		return (InArray.OwningObject ? InArray.OwningObject : InArray.GetOwner());
+#else
+		return InArray.GetOwner();
+#endif
+	}
 }
 
 
@@ -82,6 +135,12 @@ void IGameplayCueInterface::HandleGameplayCue(UObject* Self, FGameplayTag Gamepl
 		return;
 	}
 
+	if (UE_LOG_ACTIVE(LogGameplayCues, Verbose))
+	{
+		UE_LOG(LogGameplayCues, Verbose, TEXT("%s::%hs %s %s"), *GetNameSafe(Self), __func__, *GameplayCueTag.ToString(), *UEnum::GetValueAsString(EventType));
+		UE_VLOG(Self, LogGameplayCues, Verbose, TEXT("%hs %s %s"), __func__, *GameplayCueTag.ToString(), *UEnum::GetValueAsString(EventType));
+	}
+
 	// Look up a custom function for this gameplay tag. 
 	UClass* Class = Self->GetClass();
 	FGameplayTagContainer TagAndParentsContainer = GameplayCueTag.GetGameplayTagParents();
@@ -196,6 +255,14 @@ void FActiveGameplayCue::PreReplicatedRemove(const struct FActiveGameplayCueCont
 		return;
 	}
 
+	if (UE_LOG_ACTIVE(LogGameplayCues, VeryVerbose))
+	{
+		const UObject* PropertyOwner = GameplayCueInterfacePrivate::GetPropertyOwner(InArray);
+		const FString PropertyName = GameplayCueInterfacePrivate::FindNameOfProperty(&InArray, PropertyOwner);
+		UE_LOG(LogGameplayCues, VeryVerbose, TEXT("%s.%s::%hs %s %s"), *GetNameSafe(PropertyOwner), *PropertyName, __func__, *GameplayCueTag.ToString(), *PredictionKey.ToString());
+		UE_VLOG(PropertyOwner, LogGameplayCues, VeryVerbose, TEXT("%s::%hs %s %s"), *PropertyName, __func__, *GameplayCueTag.ToString(), *PredictionKey.ToString());
+	}
+
 	// We don't check the PredictionKey here like we do in PostReplicatedAdd. PredictionKey tells us
 	// if we were predictely created, but this doesn't mean we will predictively remove ourselves.
 	if (bPredictivelyRemoved == false)
@@ -213,6 +280,14 @@ void FActiveGameplayCue::PostReplicatedAdd(const struct FActiveGameplayCueContai
 		return;
 	}
 
+	if (UE_LOG_ACTIVE(LogGameplayCues, VeryVerbose))
+	{
+		const UObject* PropertyOwner = GameplayCueInterfacePrivate::GetPropertyOwner(InArray);
+		const FString PropertyName = GameplayCueInterfacePrivate::FindNameOfProperty(&InArray, PropertyOwner);
+		UE_LOG(LogGameplayCues, VeryVerbose, TEXT("%s.%s::%hs %s %s"), *GetNameSafe(PropertyOwner), *PropertyName, __func__, *GameplayCueTag.ToString(), *PredictionKey.ToString());
+		UE_VLOG(PropertyOwner, LogGameplayCues, VeryVerbose, TEXT("%s::%hs %s %s"), *PropertyName, __func__, *GameplayCueTag.ToString(), *PredictionKey.ToString());
+	}
+
 	InArray.Owner->UpdateTagMap(GameplayCueTag, 1);
 
 	if (PredictionKey.IsLocalClientKey() == false)
@@ -224,7 +299,7 @@ void FActiveGameplayCue::PostReplicatedAdd(const struct FActiveGameplayCueContai
 
 FString FActiveGameplayCue::GetDebugString()
 {
-	return FString::Printf(TEXT("(%s / %s"), *GameplayCueTag.ToString(), *PredictionKey.ToString());
+	return FString::Printf(TEXT("(%s / %s)"), *GameplayCueTag.ToString(), *PredictionKey.ToString());
 }
 
 void FActiveGameplayCueContainer::AddCue(const FGameplayTag& Tag, const FPredictionKey& PredictionKey, const FGameplayCueParameters& Parameters)
@@ -234,7 +309,13 @@ void FActiveGameplayCueContainer::AddCue(const FGameplayTag& Tag, const FPredict
 		return;
 	}
 
-	UWorld* World = Owner->GetWorld();
+	if (UE_LOG_ACTIVE(LogGameplayCues, VeryVerbose))
+	{
+		const UObject* PropertyOwner = GameplayCueInterfacePrivate::GetPropertyOwner(*this);
+		const FString PropertyName = GameplayCueInterfacePrivate::FindNameOfProperty(this, PropertyOwner);
+		UE_LOG(LogGameplayCues, VeryVerbose, TEXT("%s.%s::%hs %s %s"), *GetNameSafe(PropertyOwner), *PropertyName, __func__, *Tag.ToString(), *PredictionKey.ToString());
+		UE_VLOG(PropertyOwner, LogGameplayCues, VeryVerbose, TEXT("%s::%hs %s %s"), *PropertyName, __func__, *Tag.ToString(), *PredictionKey.ToString());
+	}
 
 	// Store the prediction key so the client can investigate it
 	FActiveGameplayCue&	NewCue = GameplayCues[GameplayCues.AddDefaulted()];
@@ -253,6 +334,15 @@ void FActiveGameplayCueContainer::RemoveCue(const FGameplayTag& Tag)
 		return;
 	}
 
+	if (UE_LOG_ACTIVE(LogGameplayCues, VeryVerbose))
+	{
+		const UObject* PropertyOwner = GameplayCueInterfacePrivate::GetPropertyOwner(*this);
+		const FString PropertyName = GameplayCueInterfacePrivate::FindNameOfProperty(this, PropertyOwner);
+		UE_LOG(LogGameplayCues, VeryVerbose, TEXT("%s.%s::%hs %s"), *GetNameSafe(PropertyOwner), *PropertyName, __func__, *Tag.ToString());
+		UE_VLOG(PropertyOwner, LogGameplayCues, VeryVerbose, TEXT("%s::%hs %s"), *PropertyName, __func__, *Tag.ToString());
+	}
+
+
 	int32 CountDelta = 0;
 
 	// This is the old behavior (tentatively removed in UE5.5):  A single callback per multiple removals
@@ -390,6 +480,12 @@ void FActiveGameplayCueContainer::SetOwner(UAbilitySystemComponent* InOwner)
 	}
 }
 
+UAbilitySystemComponent* FActiveGameplayCueContainer::GetOwner() const
+{
+	return Owner;
+}
+
+
 // ----------------------------------------------------------------------------------------
 
 FMinimalGameplayCueReplicationProxy::FMinimalGameplayCueReplicationProxy()
diff --git a/Engine/Plugins/Runtime/GameplayAbilities/Source/GameplayAbilities/Public/GameplayCueInterface.h b/Engine/Plugins/Runtime/GameplayAbilities/Source/GameplayAbilities/Public/GameplayCueInterface.h
--- a/Engine/Plugins/Runtime/GameplayAbilities/Source/GameplayAbilities/Public/GameplayCueInterface.h
+++ b/Engine/Plugins/Runtime/GameplayAbilities/Source/GameplayAbilities/Public/GameplayCueInterface.h
@@ -136,6 +136,7 @@ struct GAMEPLAYABILITIES_API FActiveGameplayCueContainer : public FFastArraySeri
 	TArray< FActiveGameplayCue >	GameplayCues;
 
 	void SetOwner(UAbilitySystemComponent* InOwner);
+	UAbilitySystemComponent* GetOwner() const;
 
 	/** Should this container only replicate in minimal replication mode */
 	bool bMinimalReplication;
