Fix Server Crash when removing All Gameplay Cues

#tests Tested force calling RemoveAllGameplayCues when we have multiple of the same tags. No Crash with fix. Server Crash without it.
#rnx

diff --git a/Engine/Plugins/Runtime/GameplayAbilities/Source/GameplayAbilities/Private/AbilitySystemComponent.cpp b/Engine/Plugins/Runtime/GameplayAbilities/Source/GameplayAbilities/Private/AbilitySystemComponent.cpp
--- a/Engine/Plugins/Runtime/GameplayAbilities/Source/GameplayAbilities/Private/AbilitySystemComponent.cpp
+++ b/Engine/Plugins/Runtime/GameplayAbilities/Source/GameplayAbilities/Private/AbilitySystemComponent.cpp
@@ -46,6 +46,9 @@ static FAutoConsoleVariableRef CVarReplicateGameplayAbilitiesToOwnerOnly(TEXT("A
 static bool bForceReplicationAlsoUpdatesReplicatedProxyInterface = true;
 static FAutoConsoleVariableRef CVarForceReplicationAlsoUpdatesReplicatedProxyInterface(TEXT("AbilitySystem.Fix.ForceReplicationAlsoUpdatesReplicatedProxyInterface"), bForceReplicationAlsoUpdatesReplicatedProxyInterface, TEXT("Default: True.  When true, Calling ForceReplication() on the AbilitySystemComponent will also call ForceReplication() on the ReplicationProxy to ensure prompt replication of Cues and Tags"));
 
+static bool bSafeRemoveAllGameplayCues = true;
+static FAutoConsoleVariableRef CVarSafeRemoveAllGameplayCues(TEXT("AbilitySystem.Fix.SafeRemovalAllGameplayCues"), bSafeRemoveAllGameplayCues, TEXT("Default: True. When true, Calling RemoveAllGameplayCues on the AbilitySystemComponent, duplicate GameplayCues will be removed safely, avoiding Index out of bounds errors."));
+
 UAbilitySystemComponent::UAbilitySystemComponent(const FObjectInitializer& ObjectInitializer)
 	: Super(ObjectInitializer)
 	, GameplayTagCountContainer()
@@ -1417,9 +1420,21 @@ void UAbilitySystemComponent::RemoveGameplayCue_Internal(const FGameplayTag Game
 
 void UAbilitySystemComponent::RemoveAllGameplayCues()
 {
-	for (int32 i = (ActiveGameplayCues.GameplayCues.Num() - 1); i >= 0; --i)
+	if (bSafeRemoveAllGameplayCues)
+	{
+		while (!ActiveGameplayCues.GameplayCues.IsEmpty())
+		{
+			RemoveGameplayCue(ActiveGameplayCues.GameplayCues.Last().GameplayCueTag);
+		}
+	}
+	else
 	{
-		RemoveGameplayCue(ActiveGameplayCues.GameplayCues[i].GameplayCueTag);
+		// NOTE: This code is dangerous as `RemoveGameplayCue` can remove multiple elements from the array
+		// This will result in Index out of Bounds errors as the Iteration Index will be out of sync with the Number of Elements
+		for (int32 i = (ActiveGameplayCues.GameplayCues.Num() - 1); i >= 0; --i)
+		{
+			RemoveGameplayCue(ActiveGameplayCues.GameplayCues[i].GameplayCueTag);
+		}
 	}
 }
 
