PR #12035: GameplayCues: OriginalTag is never set for GCNs.

diff --git a/Engine/Plugins/Runtime/GameplayAbilities/Source/GameplayAbilities/Private/GameplayCueManager.cpp b/Engine/Plugins/Runtime/GameplayAbilities/Source/GameplayAbilities/Private/GameplayCueManager.cpp
--- a/Engine/Plugins/Runtime/GameplayAbilities/Source/GameplayAbilities/Private/GameplayCueManager.cpp
+++ b/Engine/Plugins/Runtime/GameplayAbilities/Source/GameplayAbilities/Private/GameplayCueManager.cpp
@@ -179,6 +179,11 @@ void UGameplayCueManager::RouteGameplayCue(AActor* TargetActor, FGameplayTag Gam
 	QUICK_SCOPE_CYCLE_COUNTER(STAT_GameplayCueManager_RouteGameplayCue);
 #endif
 
+	if (!Parameters.OriginalTag.IsValid())
+	{
+		Parameters.OriginalTag = GameplayCueTag;
+	}
+
 	// If we want to ignore interfaces, set the pointer to null
 	IGameplayCueInterface* GameplayCueInterface = !(Options & EGameplayCueExecutionOptions::IgnoreInterfaces) ? Cast<IGameplayCueInterface>(TargetActor) : nullptr;
 	bool bAcceptsCue = true;
