PR #12007: Fixed MontageTask Interrupt After BlendOut

Added a CVar "AbilitySystem.PlayMontage.FireInterruptOnAnimEndInterrupt" which can be set to false if this causes issues.

diff --git a/Engine/Plugins/Runtime/GameplayAbilities/Source/GameplayAbilities/Private/Abilities/Tasks/AbilityTask_PlayMontageAndWait.cpp b/Engine/Plugins/Runtime/GameplayAbilities/Source/GameplayAbilities/Private/Abilities/Tasks/AbilityTask_PlayMontageAndWait.cpp
--- a/Engine/Plugins/Runtime/GameplayAbilities/Source/GameplayAbilities/Private/Abilities/Tasks/AbilityTask_PlayMontageAndWait.cpp
+++ b/Engine/Plugins/Runtime/GameplayAbilities/Source/GameplayAbilities/Private/Abilities/Tasks/AbilityTask_PlayMontageAndWait.cpp
@@ -12,6 +12,9 @@
 static bool GUseAggressivePlayMontageAndWaitEndTask = true;
 static FAutoConsoleVariableRef CVarAggressivePlayMontageAndWaitEndTask(TEXT("AbilitySystem.PlayMontage.AggressiveEndTask"), GUseAggressivePlayMontageAndWaitEndTask, TEXT("This should be set to true in order to avoid multiple callbacks off an AbilityTask_PlayMontageAndWait node"));
 
+static bool GPlayMontageAndWaitFireInterruptOnAnimEndInterrupt = true;
+static FAutoConsoleVariableRef CVarPlayMontageAndWaitFireInterruptOnAnimEndInterrupt(TEXT("AbilitySystem.PlayMontage.FireInterruptOnAnimEndInterrupt"), GPlayMontageAndWaitFireInterruptOnAnimEndInterrupt, TEXT("This is a fix that will cause AbilityTask_PlayMontageAndWait to fire its Interrupt event if the underlying AnimInstance ends in an interrupted"));
+
 void UAbilityTask_PlayMontageAndWait::OnMontageBlendingOut(UAnimMontage* Montage, bool bInterrupted)
 {
 	const bool bPlayingThisMontage = (Montage == MontageToPlay) && Ability && Ability->GetCurrentMontage() == MontageToPlay;
@@ -38,6 +41,7 @@ void UAbilityTask_PlayMontageAndWait::OnMontageBlendingOut(UAnimMontage* Montage
 	{
 		if (bInterrupted)
 		{
+            bAllowInterruptAfterBlendOut = false;
 			OnInterrupted.Broadcast();
 
 			if (GUseAggressivePlayMontageAndWaitEndTask)
@@ -73,6 +77,7 @@ void UAbilityTask_PlayMontageAndWait::OnGameplayAbilityCancelled()
 		// Let the BP handle the interrupt as well
 		if (ShouldBroadcastAbilityTaskDelegates())
 		{
+			bAllowInterruptAfterBlendOut = false;
 			OnInterrupted.Broadcast();
 		}
 	}
@@ -92,6 +97,13 @@ void UAbilityTask_PlayMontageAndWait::OnMontageEnded(UAnimMontage* Montage, bool
 			OnCompleted.Broadcast();
 		}
 	}
+	else if(bAllowInterruptAfterBlendOut && GPlayMontageAndWaitFireInterruptOnAnimEndInterrupt)
+	{
+		if (ShouldBroadcastAbilityTaskDelegates())
+		{
+			OnInterrupted.Broadcast();
+		}
+	}
 
 	EndTask();
 }
