PR #12048: [GAS] Fix FGameplayEffectQuery::IsEmpty incorrectly returning true even if a CustomDelegate is bound

diff --git a/Engine/Plugins/Runtime/GameplayAbilities/Source/GameplayAbilities/Private/GameplayEffect.cpp b/Engine/Plugins/Runtime/GameplayAbilities/Source/GameplayAbilities/Private/GameplayEffect.cpp
--- a/Engine/Plugins/Runtime/GameplayAbilities/Source/GameplayAbilities/Private/GameplayEffect.cpp
+++ b/Engine/Plugins/Runtime/GameplayAbilities/Source/GameplayAbilities/Private/GameplayEffect.cpp
@@ -5840,6 +5840,8 @@ bool FGameplayEffectQuery::IsEmpty() const
 {
 	return 
 	(
+		!CustomMatchDelegate.IsBound() &&
+		!CustomMatchDelegate_BP.IsBound() &&
 		OwningTagQuery.IsEmpty() &&
 		EffectTagQuery.IsEmpty() &&
 		SourceAggregateTagQuery.IsEmpty() &&
