Fix for UGameplayAbility::GetSourceObject not returning a valid object when it is in the pending list.

diff --git a/Engine/Plugins/Runtime/GameplayAbilities/Source/GameplayAbilities/Private/Abilities/GameplayAbility.cpp b/Engine/Plugins/Runtime/GameplayAbilities/Source/GameplayAbilities/Private/Abilities/GameplayAbility.cpp
--- a/Engine/Plugins/Runtime/GameplayAbilities/Source/GameplayAbilities/Private/Abilities/GameplayAbility.cpp
+++ b/Engine/Plugins/Runtime/GameplayAbilities/Source/GameplayAbilities/Private/Abilities/GameplayAbility.cpp
@@ -1830,7 +1830,7 @@ UObject* UGameplayAbility::GetSourceObject(FGameplayAbilitySpecHandle Handle, co
 		UAbilitySystemComponent* const AbilitySystemComponent = ActorInfo->AbilitySystemComponent.Get();
 		if (AbilitySystemComponent != nullptr)
 		{
-			FGameplayAbilitySpec* AbilitySpec = AbilitySystemComponent->FindAbilitySpecFromHandle(Handle);
+			FGameplayAbilitySpec* AbilitySpec = AbilitySystemComponent->FindAbilitySpecFromHandle(Handle, EConsiderPending::All);
 			if (AbilitySpec)
 			{
 				return AbilitySpec->SourceObject.Get();
