[GAS] Added missing closing parenthesis to Input.ini

diff --git a/Engine/Plugins/Runtime/GameplayAbilities/Config/Input.ini b/Engine/Plugins/Runtime/GameplayAbilities/Config/Input.ini
--- a/Engine/Plugins/Runtime/GameplayAbilities/Config/Input.ini
+++ b/Engine/Plugins/Runtime/GameplayAbilities/Config/Input.ini
@@ -1,2 +1,2 @@
 [/Script/EngineSettings.ConsoleSettings]
-+ManualAutoCompleteList=(Command="ShowDebug AbilitySystem",Desc="Toggles debug information about the Gameplay Ability System"
\ No newline at end of file
++ManualAutoCompleteList=(Command="ShowDebug AbilitySystem",Desc="Toggles debug information about the Gameplay Ability System")
\ No newline at end of file
