Return the asset tags as per the function name in get owned gameplay tags. Functionality to access Granted Tags will be moved in the future.

diff --git a/Engine/Plugins/Runtime/GameplayAbilities/Source/GameplayAbilities/Private/GameplayEffect.cpp b/Engine/Plugins/Runtime/GameplayAbilities/Source/GameplayAbilities/Private/GameplayEffect.cpp
--- a/Engine/Plugins/Runtime/GameplayAbilities/Source/GameplayAbilities/Private/GameplayEffect.cpp
+++ b/Engine/Plugins/Runtime/GameplayAbilities/Source/GameplayAbilities/Private/GameplayEffect.cpp
@@ -93,6 +93,10 @@ namespace UE::GameplayEffect
 	FAutoConsoleVariableRef CVarSkipUnmappedReferencesCheckForGameplayCues{ TEXT("AbilitySystem.Fix.SkipUnmappedReferencesCheckForGameplayCues"), bSkipUnmappedReferencesCheckForGameplayCues,
 		TEXT("Skip the bHasMoreUnmappedReferences check for GameplayCues which never worked as intended and causes issues when set properly (may be deprecated soon)"), ECVF_Default };
 
+	int32 ReturnAssetTagsViaGetOwnedGameplayTags = 2;
+	FAutoConsoleVariableRef CVarReturnAssetTagsViaGetOwnedGameplayTags{ TEXT("AbilitySystem.Fix.GameplayEffects.ReturnAssetTagsViaGetOwnedGameplayTags"), ReturnAssetTagsViaGetOwnedGameplayTags ,
+		TEXT("0 = Return the asset tags as per the function name in get owned gameplay tags. 1 =  Use legacy behavior. If legacy is desired, change code to GetGrantedTags. 2 = Legacy behavior but ensure to catch bad use. 3. Legacy behavior but always ensure to catch multiple issues."), ECVF_Default };
+
 #if WITH_EDITOR
 	namespace EditorOnly
 	{
@@ -220,8 +224,25 @@ void UGameplayEffect::PostInitProperties()
 
 void UGameplayEffect::GetOwnedGameplayTags(FGameplayTagContainer& TagContainer) const
 {
-	UE_LOG(LogGameplayEffects, Warning, TEXT("%hs on %s: The implementation and method name did not match.  Use GetGrantedTags() to get the tags Granted to the Actor this GameplayEffect is applied to."), __func__, *GetName());
- 	TagContainer.AppendTags(GetGrantedTags());
+	if (UE::GameplayEffect::ReturnAssetTagsViaGetOwnedGameplayTags == 0)
+	{
+		TagContainer.AppendTags(GetAssetTags());
+	}
+	else
+	{	
+		// Default Engine Behavior, to be deprecated. Set to ensure to warn of potential complications 
+		if (UE::GameplayEffect::ReturnAssetTagsViaGetOwnedGameplayTags == 2)
+		{
+			// Warn once. To check for additional accessors, set ReturnAssetTagsViaGetOwnedGameplayTags = 3
+			ensureMsgf(false, TEXT("%hs on %s: The implementation and method name did not match.  Use GetGrantedTags() to get the tags Granted to the Actor this GameplayEffect is applied to. Change ReturnAssetTagsViaGetOwnedGameplayTags to get the Asset Tags on the GE, or to modify the behavior of this ensure"), __func__, *GetName());
+		}
+		else if (UE::GameplayEffect::ReturnAssetTagsViaGetOwnedGameplayTags == 3)
+		{
+			ensureAlwaysMsgf(false, TEXT("%hs on %s: The implementation and method name did not match.  Use GetGrantedTags() to get the tags Granted to the Actor this GameplayEffect is applied to. Change ReturnAssetTagsViaGetOwnedGameplayTags to get the Asset Tags on the GE, or to modify the behavior of this ensure"), __func__, *GetName());
+		}
+
+		TagContainer.AppendTags(GetGrantedTags());
+	}
 }
 
 bool UGameplayEffect::HasMatchingGameplayTag(FGameplayTag TagToCheck) const
