Workaround for ASC infinitely looping while attempting to remove all cues.

diff --git a/Engine/Plugins/Runtime/GameplayAbilities/Source/GameplayAbilities/Private/AbilitySystemComponent.cpp b/Engine/Plugins/Runtime/GameplayAbilities/Source/GameplayAbilities/Private/AbilitySystemComponent.cpp
--- a/Engine/Plugins/Runtime/GameplayAbilities/Source/GameplayAbilities/Private/AbilitySystemComponent.cpp
+++ b/Engine/Plugins/Runtime/GameplayAbilities/Source/GameplayAbilities/Private/AbilitySystemComponent.cpp
@@ -46,8 +46,11 @@ static FAutoConsoleVariableRef CVarReplicateGameplayAbilitiesToOwnerOnly(TEXT("A
 static bool bForceReplicationAlsoUpdatesReplicatedProxyInterface = true;
 static FAutoConsoleVariableRef CVarForceReplicationAlsoUpdatesReplicatedProxyInterface(TEXT("AbilitySystem.Fix.ForceReplicationAlsoUpdatesReplicatedProxyInterface"), bForceReplicationAlsoUpdatesReplicatedProxyInterface, TEXT("Default: True.  When true, Calling ForceReplication() on the AbilitySystemComponent will also call ForceReplication() on the ReplicationProxy to ensure prompt replication of Cues and Tags"));
 
-static bool bSafeRemoveAllGameplayCues = true;
-static FAutoConsoleVariableRef CVarSafeRemoveAllGameplayCues(TEXT("AbilitySystem.Fix.SafeRemovalAllGameplayCues"), bSafeRemoveAllGameplayCues, TEXT("Default: True. When true, Calling RemoveAllGameplayCues on the AbilitySystemComponent, duplicate GameplayCues will be removed safely, avoiding Index out of bounds errors."));
+static int32 bSafeRemoveAllGameplayCuesMode = 1;
+static FAutoConsoleVariableRef CVarSafeRemoveAllGameplayCues(
+	TEXT("AbilitySystem.Fix.SafeRemovalAllGameplayCuesMode"),
+	bSafeRemoveAllGameplayCuesMode,
+	TEXT("2: Remove only gameplay cues that were present at the beginning of removal. 1: Remove all gameplay cues until empty (potentially unsafe on client). 0: Remove while iterating (unsafe when duplicates are present)."));
 
 UAbilitySystemComponent::UAbilitySystemComponent(const FObjectInitializer& ObjectInitializer)
 	: Super(ObjectInitializer)
@@ -1420,7 +1423,19 @@ void UAbilitySystemComponent::RemoveGameplayCue_Internal(const FGameplayTag Game
 
 void UAbilitySystemComponent::RemoveAllGameplayCues()
 {
-	if (bSafeRemoveAllGameplayCues)
+	if (bSafeRemoveAllGameplayCuesMode == 2)
+	{
+		TArray<FGameplayTag, TInlineAllocator<16>> GameplayCueTagsToRemove;
+		for (const FActiveGameplayCue& Cue : ActiveGameplayCues.GameplayCues)
+		{
+			GameplayCueTagsToRemove.Emplace(Cue.GameplayCueTag);
+		}
+		for (const FGameplayTag& CueTagToRemove : GameplayCueTagsToRemove)
+		{
+			RemoveGameplayCue(CueTagToRemove);
+		}
+	}
+	else if (bSafeRemoveAllGameplayCuesMode == 1)
 	{
 		while (!ActiveGameplayCues.GameplayCues.IsEmpty())
 		{
