Cleanup ASC from client hang fix

diff --git a/Engine/Plugins/Runtime/GameplayAbilities/Source/GameplayAbilities/Private/AbilitySystemComponent.cpp b/Engine/Plugins/Runtime/GameplayAbilities/Source/GameplayAbilities/Private/AbilitySystemComponent.cpp
--- a/Engine/Plugins/Runtime/GameplayAbilities/Source/GameplayAbilities/Private/AbilitySystemComponent.cpp
+++ b/Engine/Plugins/Runtime/GameplayAbilities/Source/GameplayAbilities/Private/AbilitySystemComponent.cpp
@@ -46,12 +46,6 @@ static FAutoConsoleVariableRef CVarReplicateGameplayAbilitiesToOwnerOnly(TEXT("A
 static bool bForceReplicationAlsoUpdatesReplicatedProxyInterface = true;
 static FAutoConsoleVariableRef CVarForceReplicationAlsoUpdatesReplicatedProxyInterface(TEXT("AbilitySystem.Fix.ForceReplicationAlsoUpdatesReplicatedProxyInterface"), bForceReplicationAlsoUpdatesReplicatedProxyInterface, TEXT("Default: True.  When true, Calling ForceReplication() on the AbilitySystemComponent will also call ForceReplication() on the ReplicationProxy to ensure prompt replication of Cues and Tags"));
 
-static int32 bSafeRemoveAllGameplayCuesMode = 2;
-static FAutoConsoleVariableRef CVarSafeRemoveAllGameplayCues(
-	TEXT("AbilitySystem.Fix.SafeRemovalAllGameplayCuesMode"),
-	bSafeRemoveAllGameplayCuesMode,
-	TEXT("2: Remove only gameplay cues that were present at the beginning of removal. 1: Remove all gameplay cues until empty (potentially unsafe on client). 0: Remove while iterating (unsafe when duplicates are present)."));
-
 UAbilitySystemComponent::UAbilitySystemComponent(const FObjectInitializer& ObjectInitializer)
 	: Super(ObjectInitializer)
 	, GameplayTagCountContainer()
@@ -1423,33 +1417,14 @@ void UAbilitySystemComponent::RemoveGameplayCue_Internal(const FGameplayTag Game
 
 void UAbilitySystemComponent::RemoveAllGameplayCues()
 {
-	if (bSafeRemoveAllGameplayCuesMode == 2)
-	{
-		TArray<FGameplayTag, TInlineAllocator<16>> GameplayCueTagsToRemove;
-		for (const FActiveGameplayCue& Cue : ActiveGameplayCues.GameplayCues)
-		{
-			GameplayCueTagsToRemove.Emplace(Cue.GameplayCueTag);
-		}
-		for (const FGameplayTag& CueTagToRemove : GameplayCueTagsToRemove)
-		{
-			RemoveGameplayCue(CueTagToRemove);
-		}
-	}
-	else if (bSafeRemoveAllGameplayCuesMode == 1)
+	TArray<FGameplayTag, TInlineAllocator<16>> GameplayCueTagsToRemove;
+	for (const FActiveGameplayCue& Cue : ActiveGameplayCues.GameplayCues)
 	{
-		while (!ActiveGameplayCues.GameplayCues.IsEmpty())
-		{
-			RemoveGameplayCue(ActiveGameplayCues.GameplayCues.Last().GameplayCueTag);
-		}
+		GameplayCueTagsToRemove.Emplace(Cue.GameplayCueTag);
 	}
-	else
+	for (const FGameplayTag& CueTagToRemove : GameplayCueTagsToRemove)
 	{
-		// NOTE: This code is dangerous as `RemoveGameplayCue` can remove multiple elements from the array
-		// This will result in Index out of Bounds errors as the Iteration Index will be out of sync with the Number of Elements
-		for (int32 i = (ActiveGameplayCues.GameplayCues.Num() - 1); i >= 0; --i)
-		{
-			RemoveGameplayCue(ActiveGameplayCues.GameplayCues[i].GameplayCueTag);
-		}
+		RemoveGameplayCue(CueTagToRemove);
 	}
 }
 
