[GAS] Removing PackageMap assumptions from the Gameplay Debugger (it's not valid in Replay and can therefore crash when using during Replay)

diff --git a/Engine/Plugins/Runtime/GameplayAbilities/Source/GameplayAbilities/Private/GameplayDebuggerCategory_Abilities.cpp b/Engine/Plugins/Runtime/GameplayAbilities/Source/GameplayAbilities/Private/GameplayDebuggerCategory_Abilities.cpp
--- a/Engine/Plugins/Runtime/GameplayAbilities/Source/GameplayAbilities/Private/GameplayDebuggerCategory_Abilities.cpp
+++ b/Engine/Plugins/Runtime/GameplayAbilities/Source/GameplayAbilities/Private/GameplayDebuggerCategory_Abilities.cpp
@@ -85,7 +85,7 @@ void FGameplayDebuggerCategory_Abilities::OnShowGameplayAttributesToggle()
 void FGameplayDebuggerCategory_Abilities::FRepData::Serialize(FArchive& Ar)
 {
 	bool bSuccess;
-	OwnedTags.NetSerialize(Ar, ClientPackageMap.Get(), bSuccess);
+	OwnedTags.NetSerialize(Ar, nullptr, bSuccess);
 
 	Ar << TagCounts;
 
@@ -143,10 +143,6 @@ void FGameplayDebuggerCategory_Abilities::CollectData(APlayerController* OwnerPC
 {
 	if (const UAbilitySystemComponent* AbilityComp = UAbilitySystemGlobals::GetAbilitySystemComponentFromActor(DebugActor))
 	{
-		// Save off the package map for serialization over the network
-		UNetConnection* NetConnection = OwnerPC->GetNetConnection();
-		DataPack.ClientPackageMap = NetConnection ? NetConnection->PackageMap : nullptr;
-
 		AbilityComp->GetOwnedGameplayTags(DataPack.OwnedTags);
 
 		// Copy over the tag counts
diff --git a/Engine/Plugins/Runtime/GameplayAbilities/Source/GameplayAbilities/Public/GameplayDebuggerCategory_Abilities.h b/Engine/Plugins/Runtime/GameplayAbilities/Source/GameplayAbilities/Public/GameplayDebuggerCategory_Abilities.h
--- a/Engine/Plugins/Runtime/GameplayAbilities/Source/GameplayAbilities/Public/GameplayDebuggerCategory_Abilities.h
+++ b/Engine/Plugins/Runtime/GameplayAbilities/Source/GameplayAbilities/Public/GameplayDebuggerCategory_Abilities.h
@@ -49,9 +49,6 @@ protected:
 
 	struct FRepData
 	{
-		// to aid in NetSerialize
-		TWeakObjectPtr<UPackageMap>	ClientPackageMap;
-
 		FGameplayTagContainer OwnedTags;
 		TArray<int32> TagCounts;
 
