[GAS] Adding GameplayTagFilter to API that works with gameplay tag events

The big projects that use a lot of gameplay tags may want to limit the tags they use to identify gameplay ability events. Adding a gameplay tag filter allows such a project to use tag category remapping to adjust all filtering to their gameplay tag structure.

To enable this feature, use the Category Remapping settings to remap "GameplayEventTagsCategory" under "Project Settings / GameplayTags".

diff --git a/Engine/Plugins/Runtime/GameplayAbilities/Source/GameplayAbilities/Public/Abilities/Async/AbilityAsync_WaitGameplayEvent.h b/Engine/Plugins/Runtime/GameplayAbilities/Source/GameplayAbilities/Public/Abilities/Async/AbilityAsync_WaitGameplayEvent.h
--- a/Engine/Plugins/Runtime/GameplayAbilities/Source/GameplayAbilities/Public/Abilities/Async/AbilityAsync_WaitGameplayEvent.h
+++ b/Engine/Plugins/Runtime/GameplayAbilities/Source/GameplayAbilities/Public/Abilities/Async/AbilityAsync_WaitGameplayEvent.h
@@ -22,7 +22,7 @@ public:
 	 * If used in an ability graph, this async action will wait even after activation ends. It's recommended to use WaitGameplayEvent instead.
 	 */
 	UFUNCTION(BlueprintCallable, Category = "Ability|Async", meta = (DefaultToSelf = "TargetActor", BlueprintInternalUseOnly = "TRUE"))
-	static UAbilityAsync_WaitGameplayEvent* WaitGameplayEventToActor(AActor* TargetActor, FGameplayTag EventTag, bool OnlyTriggerOnce = false, bool OnlyMatchExact = true);
+	static UAbilityAsync_WaitGameplayEvent* WaitGameplayEventToActor(AActor* TargetActor, UPARAM(meta=(GameplayTagFilter="GameplayEventTagsCategory")) FGameplayTag EventTag, bool OnlyTriggerOnce = false, bool OnlyMatchExact = true);
 
 	DECLARE_DYNAMIC_MULTICAST_DELEGATE_OneParam(FEventReceivedDelegate, FGameplayEventData, Payload);
 
diff --git a/Engine/Plugins/Runtime/GameplayAbilities/Source/GameplayAbilities/Public/Abilities/GameplayAbility.h b/Engine/Plugins/Runtime/GameplayAbilities/Source/GameplayAbilities/Public/Abilities/GameplayAbility.h
--- a/Engine/Plugins/Runtime/GameplayAbilities/Source/GameplayAbilities/Public/Abilities/GameplayAbility.h
+++ b/Engine/Plugins/Runtime/GameplayAbilities/Source/GameplayAbilities/Public/Abilities/GameplayAbility.h
@@ -551,7 +551,7 @@ protected:
 
 	/** Sends a gameplay event, also creates a prediction window */
 	UFUNCTION(BlueprintCallable, Category = Ability)
-	virtual void SendGameplayEvent(FGameplayTag EventTag, FGameplayEventData Payload);
+	virtual void SendGameplayEvent(UPARAM(meta=(GameplayTagFilter="GameplayEventTagsCategory")) FGameplayTag EventTag, FGameplayEventData Payload);
 
 	// --------------------------------------
 	//	CanActivate
diff --git a/Engine/Plugins/Runtime/GameplayAbilities/Source/GameplayAbilities/Public/Abilities/GameplayAbilityTypes.h b/Engine/Plugins/Runtime/GameplayAbilities/Source/GameplayAbilities/Public/Abilities/GameplayAbilityTypes.h
--- a/Engine/Plugins/Runtime/GameplayAbilities/Source/GameplayAbilities/Public/Abilities/GameplayAbilityTypes.h
+++ b/Engine/Plugins/Runtime/GameplayAbilities/Source/GameplayAbilities/Public/Abilities/GameplayAbilityTypes.h
@@ -243,7 +243,7 @@ struct GAMEPLAYABILITIES_API FGameplayEventData
 	}
 	
 	/** Tag of the event that triggered this */
-	UPROPERTY(EditAnywhere, BlueprintReadWrite, Category = GameplayAbilityTriggerPayload)
+	UPROPERTY(EditAnywhere, BlueprintReadWrite, Category = GameplayAbilityTriggerPayload, meta=(GameplayTagFilter="GameplayEventTagsCategory"))
 	FGameplayTag EventTag;
 
 	/** The instigator of the event */
diff --git a/Engine/Plugins/Runtime/GameplayAbilities/Source/GameplayAbilities/Public/Abilities/Tasks/AbilityTask_WaitGameplayEvent.h b/Engine/Plugins/Runtime/GameplayAbilities/Source/GameplayAbilities/Public/Abilities/Tasks/AbilityTask_WaitGameplayEvent.h
--- a/Engine/Plugins/Runtime/GameplayAbilities/Source/GameplayAbilities/Public/Abilities/Tasks/AbilityTask_WaitGameplayEvent.h
+++ b/Engine/Plugins/Runtime/GameplayAbilities/Source/GameplayAbilities/Public/Abilities/Tasks/AbilityTask_WaitGameplayEvent.h
@@ -26,7 +26,7 @@ class GAMEPLAYABILITIES_API UAbilityTask_WaitGameplayEvent : public UAbilityTask
 	 * If OnlyMatchExact = false it will trigger for nested tags
 	 */
 	UFUNCTION(BlueprintCallable, Category = "Ability|Tasks", meta = (HidePin = "OwningAbility", DefaultToSelf = "OwningAbility", BlueprintInternalUseOnly = "TRUE"))
-	static UAbilityTask_WaitGameplayEvent* WaitGameplayEvent(UGameplayAbility* OwningAbility, FGameplayTag EventTag, AActor* OptionalExternalTarget=nullptr, bool OnlyTriggerOnce=false, bool OnlyMatchExact = true);
+	static UAbilityTask_WaitGameplayEvent* WaitGameplayEvent(UGameplayAbility* OwningAbility, UPARAM(meta=(GameplayTagFilter="GameplayEventTagsCategory"))FGameplayTag EventTag, AActor* OptionalExternalTarget=nullptr, bool OnlyTriggerOnce=false, bool OnlyMatchExact = true);
 
 	void SetExternalTarget(AActor* Actor);
 
diff --git a/Engine/Plugins/Runtime/GameplayAbilities/Source/GameplayAbilities/Public/AbilitySystemBlueprintLibrary.h b/Engine/Plugins/Runtime/GameplayAbilities/Source/GameplayAbilities/Public/AbilitySystemBlueprintLibrary.h
--- a/Engine/Plugins/Runtime/GameplayAbilities/Source/GameplayAbilities/Public/AbilitySystemBlueprintLibrary.h
+++ b/Engine/Plugins/Runtime/GameplayAbilities/Source/GameplayAbilities/Public/AbilitySystemBlueprintLibrary.h
@@ -36,7 +36,7 @@ class GAMEPLAYABILITIES_API UAbilitySystemBlueprintLibrary : public UBlueprintFu
 	 * found, the event will not be sent.
 	 */
 	UFUNCTION(BlueprintCallable, Category = Ability, Meta = (Tooltip = "This function can be used to trigger an ability on the actor in question with useful payload data."))
-	static void SendGameplayEventToActor(AActor* Actor, FGameplayTag EventTag, FGameplayEventData Payload);
+	static void SendGameplayEventToActor(AActor* Actor, UPARAM(meta=(GameplayTagFilter="GameplayEventTagsCategory")) FGameplayTag EventTag, FGameplayEventData Payload);
 	
 	// -------------------------------------------------------------------------------
 	//		Attribute
