Fix for AbilitySystemGlobals stomping GameplayAbilitiesDeveloperSettings when it does not have any valid tags to contribute (e.g. Lyra).

diff --git a/Engine/Plugins/Runtime/GameplayAbilities/Source/GameplayAbilities/Private/AbilitySystemGlobals.cpp b/Engine/Plugins/Runtime/GameplayAbilities/Source/GameplayAbilities/Private/AbilitySystemGlobals.cpp
--- a/Engine/Plugins/Runtime/GameplayAbilities/Source/GameplayAbilities/Private/AbilitySystemGlobals.cpp
+++ b/Engine/Plugins/Runtime/GameplayAbilities/Source/GameplayAbilities/Private/AbilitySystemGlobals.cpp
@@ -330,7 +330,7 @@ void UAbilitySystemGlobals::PerformDeveloperSettingsUpgrade()
 {
 	auto SyncToDeveloperSettings = [](FGameplayTag& DeveloperSettingsTag, const FGameplayTag& OurTag)
 	{
-		if (DeveloperSettingsTag != OurTag)
+		if (OurTag.IsValid() && DeveloperSettingsTag != OurTag)
 		{
 			DeveloperSettingsTag = OurTag;
 			return true;
