[GAS] When an ability is blocked by a tag, we should report which tag caused the block.  Added a generic tag for when the native or blueprint CanActivateAbility function blocks activation.

diff --git a/Engine/Plugins/Runtime/GameplayAbilities/Source/GameplayAbilities/Private/Abilities/GameplayAbility.cpp b/Engine/Plugins/Runtime/GameplayAbilities/Source/GameplayAbilities/Private/Abilities/GameplayAbility.cpp
--- a/Engine/Plugins/Runtime/GameplayAbilities/Source/GameplayAbilities/Private/Abilities/GameplayAbility.cpp
+++ b/Engine/Plugins/Runtime/GameplayAbilities/Source/GameplayAbilities/Private/Abilities/GameplayAbility.cpp
@@ -12,6 +12,7 @@
 #include "AbilitySystemPrivate.h"
 #include "Abilities/Tasks/AbilityTask.h"
 #include "Components/SkeletalMeshComponent.h"
+#include "GameplayAbilitiesDeveloperSettings.h"
 #include "GameplayCue_Types.h"
 #include "GameFramework/CharacterMovementComponent.h"
 #include "GameFramework/Character.h"
@@ -314,88 +315,86 @@ EDataValidationResult UGameplayAbility::IsDataValid(FDataValidationContext& Cont
 
 bool UGameplayAbility::DoesAbilitySatisfyTagRequirements(const UAbilitySystemComponent& AbilitySystemComponent, const FGameplayTagContainer* SourceTags, const FGameplayTagContainer* TargetTags, OUT FGameplayTagContainer* OptionalRelevantTags) const
 {
+	// Define a common lambda to check for blocked tags
 	bool bBlocked = false;
-	bool bMissing = false;
-
-	UAbilitySystemGlobals& AbilitySystemGlobals = UAbilitySystemGlobals::Get();
-	const FGameplayTag& BlockedTag = AbilitySystemGlobals.ActivateFailTagsBlockedTag;
-	const FGameplayTag& MissingTag = AbilitySystemGlobals.ActivateFailTagsMissingTag;
-	
-	// Check if any of this ability's tags are currently blocked
-	if (AbilitySystemComponent.AreAbilityTagsBlocked(GetAssetTags()))
-	{
-		bBlocked = true;
-	}
-
-	// Check to see the required/blocked tags for this ability
-	if (ActivationBlockedTags.Num() || ActivationRequiredTags.Num())
+	auto CheckForBlocked = [&](const FGameplayTagContainer& ContainerA, const FGameplayTagContainer& ContainerB)
 	{
-		static FGameplayTagContainer AbilitySystemComponentTags;
-		AbilitySystemComponentTags.Reset();
-
-		AbilitySystemComponent.GetOwnedGameplayTags(AbilitySystemComponentTags);
-
-		if (AbilitySystemComponentTags.HasAny(ActivationBlockedTags))
+		// Do we not have any tags in common?  Then we're not blocked
+		if (ContainerA.IsEmpty() || ContainerB.IsEmpty() || !ContainerA.HasAny(ContainerB))
 		{
-			bBlocked = true;
+			return;
 		}
 
-		if (!AbilitySystemComponentTags.HasAll(ActivationRequiredTags))
+		if (OptionalRelevantTags)
 		{
-			bMissing = true;
-		}
-	}
-
-	if (SourceTags != nullptr)
-	{
-		if (SourceBlockedTags.Num() || SourceRequiredTags.Num())
-		{
-			if (SourceTags->HasAny(SourceBlockedTags))
+			// Ensure the global blocking tag is only added once
+			if (!bBlocked)
 			{
-				bBlocked = true;
+				UAbilitySystemGlobals& AbilitySystemGlobals = UAbilitySystemGlobals::Get();
+				const FGameplayTag& BlockedTag = AbilitySystemGlobals.ActivateFailTagsBlockedTag;
+				OptionalRelevantTags->AddTag(BlockedTag);
 			}
 
-			if (!SourceTags->HasAll(SourceRequiredTags))
-			{
-				bMissing = true;
-			}
+			// Now append all the blocking tags
+			OptionalRelevantTags->AppendMatchingTags(ContainerA, ContainerB);
 		}
-	}
 
-	if (TargetTags != nullptr)
+		bBlocked = true;
+	};
+
+	// Define a common lambda to check for missing required tags
+	bool bMissing = false;
+	auto CheckForRequired = [&](const FGameplayTagContainer& TagsToCheck, const FGameplayTagContainer& RequiredTags)
 	{
-		if (TargetBlockedTags.Num() || TargetRequiredTags.Num())
+		// Do we have no requirements, or have met all requirements?  Then nothing's missing
+		if (RequiredTags.IsEmpty() || TagsToCheck.HasAll(RequiredTags))
 		{
-			if (TargetTags->HasAny(TargetBlockedTags))
-			{
-				bBlocked = true;
-			}
+			return;
+		}
 
-			if (!TargetTags->HasAll(TargetRequiredTags))
+		if (OptionalRelevantTags)
+		{
+			// Ensure the global missing tag is only added once
+			if (!bMissing)
 			{
-				bMissing = true;
+				UAbilitySystemGlobals& AbilitySystemGlobals = UAbilitySystemGlobals::Get();
+				const FGameplayTag& MissingTag = AbilitySystemGlobals.ActivateFailTagsMissingTag;
+				OptionalRelevantTags->AddTag(MissingTag);
 			}
+
+			FGameplayTagContainer MissingTags = RequiredTags; 
+			MissingTags.RemoveTags(TagsToCheck.GetGameplayTagParents());
+			OptionalRelevantTags->AppendTags(MissingTags);
 		}
+
+		bMissing = true;
+	};
+
+	// Start by checking all of the blocked tags first (so OptionalRelevantTags will contain blocked tags first)
+	CheckForBlocked(AbilitySystemComponent.GetBlockedAbilityTags(), GetAssetTags());
+	CheckForBlocked(AbilitySystemComponent.GetOwnedGameplayTags(), ActivationBlockedTags);
+	if (SourceTags != nullptr)
+	{
+		CheckForBlocked(*SourceTags, SourceBlockedTags);
+	}
+	if (TargetTags != nullptr)
+	{
+		CheckForBlocked(*TargetTags, TargetBlockedTags);
 	}
 
-	if (bBlocked)
+	// Now check all required tags
+	CheckForRequired(AbilitySystemComponent.GetOwnedGameplayTags(), ActivationRequiredTags);
+	if (SourceTags != nullptr)
 	{
-		if (OptionalRelevantTags && BlockedTag.IsValid())
-		{
-			OptionalRelevantTags->AddTag(BlockedTag);
-		}
-		return false;
+		CheckForRequired(*SourceTags, SourceRequiredTags);
 	}
-	if (bMissing)
+	if (TargetTags != nullptr)
 	{
-		if (OptionalRelevantTags && MissingTag.IsValid())
-		{
-			OptionalRelevantTags->AddTag(MissingTag);
-		}
-		return false;
+		CheckForRequired(*TargetTags, TargetRequiredTags);
 	}
-	
-	return true;
+
+	// We succeeded if there were no blocked tags and no missing required tags	
+	return !bBlocked && !bMissing;
 }
 
 bool UGameplayAbility::ShouldActivateAbility(ENetRole Role) const
@@ -425,8 +424,6 @@ bool UGameplayAbility::CanActivateAbility(const FGameplayAbilitySpecHandle Handl
 	static FGameplayTagContainer DummyContainer;
 	DummyContainer.Reset();
 
-	FGameplayTagContainer& OutTags = OptionalRelevantTags ? *OptionalRelevantTags : DummyContainer;
-
 	// make sure the ability system component is valid, if not bail out.
 	UAbilitySystemComponent* const AbilitySystemComponent = ActorInfo->AbilitySystemComponent.Get();
 	if (!AbilitySystemComponent)
@@ -506,7 +503,8 @@ bool UGameplayAbility::CanActivateAbility(const FGameplayAbilitySpecHandle Handl
 
 	if (bHasBlueprintCanUse)
 	{
-		if (K2_CanActivateAbility(*ActorInfo, Handle, OutTags) == false)
+		FGameplayTagContainer K2FailTags;
+		if (K2_CanActivateAbility(*ActorInfo, Handle, K2FailTags) == false)
 		{
 			if (FScopedCanActivateAbilityLogEnabler::IsLoggingEnabled())
 			{
@@ -514,6 +512,17 @@ bool UGameplayAbility::CanActivateAbility(const FGameplayAbilitySpecHandle Handl
 				UE_VLOG(ActorInfo->OwnerActor.Get(), VLogAbilitySystem, Verbose, TEXT("CanActivateAbility on %s failed, Blueprint override returned false"), *GetNameSafe(Spec->Ability));
 			}
 
+			if (OptionalRelevantTags)
+			{
+				const FGameplayTag& FailTag = GetDefault<UGameplayAbilitiesDeveloperSettings>()->ActivateFailCanActivateAbilityTag;
+				if (FailTag.IsValid())
+				{
+					OptionalRelevantTags->AddTag(FailTag);
+				}
+
+				OptionalRelevantTags->AppendTags(K2FailTags);
+			}
+
 			return false;
 		}
 	}
@@ -1034,23 +1043,25 @@ bool UGameplayAbility::CheckCooldown(const FGameplayAbilitySpecHandle Handle, co
 	}
 
 	const FGameplayTagContainer* CooldownTags = GetCooldownTags();
-	if (CooldownTags)
+	if (CooldownTags && !CooldownTags->IsEmpty())
 	{
-		if (CooldownTags->Num() > 0)
+		if (UAbilitySystemComponent* AbilitySystemComponent = ActorInfo->AbilitySystemComponent.Get())
 		{
-			if (UAbilitySystemComponent* AbilitySystemComponent = ActorInfo->AbilitySystemComponent.Get())
+			if (AbilitySystemComponent->HasAnyMatchingGameplayTags(*CooldownTags))
 			{
-				if (AbilitySystemComponent->HasAnyMatchingGameplayTags(*CooldownTags))
+				if (OptionalRelevantTags)
 				{
-					const FGameplayTag& CooldownTag = UAbilitySystemGlobals::Get().ActivateFailCooldownTag;
-
-					if (OptionalRelevantTags && CooldownTag.IsValid())
+					const FGameplayTag& FailCooldownTag = UAbilitySystemGlobals::Get().ActivateFailCooldownTag;
+					if (FailCooldownTag.IsValid())
 					{
-						OptionalRelevantTags->AddTag(CooldownTag);
+						OptionalRelevantTags->AddTag(FailCooldownTag);
 					}
 
-					return false;
+					// Let the caller know which tags were blocking
+					OptionalRelevantTags->AppendMatchingTags(AbilitySystemComponent->GetOwnedGameplayTags(), *CooldownTags);
 				}
+
+				return false;
 			}
 		}
 	}
diff --git a/Engine/Plugins/Runtime/GameplayAbilities/Source/GameplayAbilities/Private/AbilitySystemComponent_Abilities.cpp b/Engine/Plugins/Runtime/GameplayAbilities/Source/GameplayAbilities/Private/AbilitySystemComponent_Abilities.cpp
--- a/Engine/Plugins/Runtime/GameplayAbilities/Source/GameplayAbilities/Private/AbilitySystemComponent_Abilities.cpp
+++ b/Engine/Plugins/Runtime/GameplayAbilities/Source/GameplayAbilities/Private/AbilitySystemComponent_Abilities.cpp
@@ -15,6 +15,7 @@
 #include "TimerManager.h"
 #include "AbilitySystemLog.h"
 #include "AttributeSet.h"
+#include "GameplayAbilitiesDeveloperSettings.h"
 #include "GameplayPrediction.h"
 #include "GameplayEffectTypes.h"
 #include "GameplayAbilitySpec.h"
@@ -1786,6 +1787,12 @@ bool UAbilitySystemComponent::InternalTryActivateAbility(FGameplayAbilitySpecHan
 		FScopedCanActivateAbilityLogEnabler LogEnabler;
 		if (!AbilitySource->CanActivateAbility(Handle, ActorInfo, SourceTags, TargetTags, &InternalTryActivateAbilityFailureTags))
 		{
+			// At least let the user know that the native CanActivateAbility rejected it
+			if (InternalTryActivateAbilityFailureTags.IsEmpty())
+			{
+				InternalTryActivateAbilityFailureTags.AddTag(GetDefault<UGameplayAbilitiesDeveloperSettings>()->ActivateFailCanActivateAbilityTag);
+			}
+
 			// CanActivateAbility with LogEnabler will have UE_LOG/UE_VLOG so don't add more failure logs here
 			NotifyAbilityFailed(Handle, AbilitySource, InternalTryActivateAbilityFailureTags);
 			return false;
diff --git a/Engine/Plugins/Runtime/GameplayAbilities/Source/GameplayAbilities/Public/GameplayAbilitiesDeveloperSettings.h b/Engine/Plugins/Runtime/GameplayAbilities/Source/GameplayAbilities/Public/GameplayAbilitiesDeveloperSettings.h
--- a/Engine/Plugins/Runtime/GameplayAbilities/Source/GameplayAbilities/Public/GameplayAbilitiesDeveloperSettings.h
+++ b/Engine/Plugins/Runtime/GameplayAbilities/Source/GameplayAbilities/Public/GameplayAbilitiesDeveloperSettings.h
@@ -75,6 +75,10 @@ public:
 	UPROPERTY(Config, EditDefaultsOnly, Category=Gameplay)
 	bool ReplicateActivationOwnedTags = true;
 
+	/** TryActive failed due to GameplayAbility's CanActivateAbility function (Blueprint or Native) */
+	UPROPERTY(Config, EditDefaultsOnly, Category=Gameplay, meta = (ConfigRestartRequired=true))
+	FGameplayTag ActivateFailCanActivateAbilityTag;
+
 	/** TryActivate failed due to being on cooldown */
 	UPROPERTY(Config, EditDefaultsOnly, Category=Gameplay, meta = (ConfigRestartRequired=true))
 	FGameplayTag ActivateFailCooldownTag; 
