[GAS] Ability System Globals copies data back from GameplayAbilitiesDeveloperSettings.

Added some documentation acknowledging these two features and where the division lies between the two.

diff --git a/Engine/Plugins/Runtime/GameplayAbilities/README.md b/Engine/Plugins/Runtime/GameplayAbilities/README.md
--- a/Engine/Plugins/Runtime/GameplayAbilities/README.md
+++ b/Engine/Plugins/Runtime/GameplayAbilities/README.md
@@ -236,6 +236,14 @@ There is documentation for how the Gameplay Prediction mechanisms work at the to
 
 ---
 
+# Ability System Globals
+
+There is a class called [AbilitySystemGlobals](./Source/GameplayAbilities/Public/AbilitySystemGlobals.h) which provide project customization points for how to handle specific base Ability System scenarios.  For example, there a functions you can override to implement derived classes of types used throughout the code (such as `AllocGameplayEffectContext`).
+
+In UE5.5, a lot of these settings have started migrating to the [GameplayAbilitiesDeveloperSettings](./Source/GameplayAbilities/Public/GameplayAbilitiesDeveloperSettings.h) (which can be accessed using the Editor and choose the Project Settings menu item).  The rough division of responsibilities:  If it's a global setting (like a variable) then it should be configurable through Gameplay Abilities Developer Settings; if it's functionality (such as allocating project-specific classes) it should be in [AbilitySystemGlobals](./Source/GameplayAbilities/Public/AbilitySystemGlobals.h).
+
+---
+
 # Debugging the Gameplay Ability System
 
 ## Legacy ShowDebug Functionality
diff --git a/Engine/Plugins/Runtime/GameplayAbilities/Source/GameplayAbilities/Private/AbilitySystemGlobals.cpp b/Engine/Plugins/Runtime/GameplayAbilities/Source/GameplayAbilities/Private/AbilitySystemGlobals.cpp
--- a/Engine/Plugins/Runtime/GameplayAbilities/Source/GameplayAbilities/Private/AbilitySystemGlobals.cpp
+++ b/Engine/Plugins/Runtime/GameplayAbilities/Source/GameplayAbilities/Private/AbilitySystemGlobals.cpp
@@ -328,11 +328,11 @@ PRAGMA_ENABLE_DEPRECATION_WARNINGS
 
 void UAbilitySystemGlobals::PerformDeveloperSettingsUpgrade()
 {
-	auto SyncToDeveloperSettings = [](FGameplayTag& DeveloperSettingsTag, const FGameplayTag& OurTag)
+	auto SyncTag = [](FGameplayTag& DestinationTag, const FGameplayTag& OurTag)
 	{
-		if (OurTag.IsValid() && DeveloperSettingsTag != OurTag)
+		if (OurTag.IsValid() && DestinationTag != OurTag)
 		{
-			DeveloperSettingsTag = OurTag;
+			DestinationTag = OurTag;
 			return true;
 		}
 
@@ -342,11 +342,11 @@ void UAbilitySystemGlobals::PerformDeveloperSettingsUpgrade()
 	UGameplayAbilitiesDeveloperSettings* DeveloperSettings = GetMutableDefault<UGameplayAbilitiesDeveloperSettings>();
 
 	bool bUpgraded = false;
-	bUpgraded |= SyncToDeveloperSettings(DeveloperSettings->ActivateFailCooldownTag, ActivateFailCooldownTag);
-	bUpgraded |= SyncToDeveloperSettings(DeveloperSettings->ActivateFailCostTag, ActivateFailCostTag);
-	bUpgraded |= SyncToDeveloperSettings(DeveloperSettings->ActivateFailNetworkingTag, ActivateFailNetworkingTag);
-	bUpgraded |= SyncToDeveloperSettings(DeveloperSettings->ActivateFailTagsBlockedTag, ActivateFailTagsBlockedTag);
-	bUpgraded |= SyncToDeveloperSettings(DeveloperSettings->ActivateFailTagsMissingTag, ActivateFailTagsMissingTag);
+	bUpgraded |= SyncTag(DeveloperSettings->ActivateFailCooldownTag, ActivateFailCooldownTag);
+	bUpgraded |= SyncTag(DeveloperSettings->ActivateFailCostTag, ActivateFailCostTag);
+	bUpgraded |= SyncTag(DeveloperSettings->ActivateFailNetworkingTag, ActivateFailNetworkingTag);
+	bUpgraded |= SyncTag(DeveloperSettings->ActivateFailTagsBlockedTag, ActivateFailTagsBlockedTag);
+	bUpgraded |= SyncTag(DeveloperSettings->ActivateFailTagsMissingTag, ActivateFailTagsMissingTag);
 
 	if (bUpgraded)
 	{
@@ -358,6 +358,13 @@ void UAbilitySystemGlobals::PerformDeveloperSettingsUpgrade()
 			UE_LOG(LogAbilitySystem, Warning, TEXT("AbilitySystemGlobals config file (DefaultGame.ini) couldn't be saved. Make sure the file is writable to update it."));
 		}
 	}
+
+	// Now that the upgrade is done, copy any settings set in the DeveloperSettings back to here (so calls to UAbilitySystemGlobals::Get().SomeTag work)
+	SyncTag(ActivateFailCooldownTag, DeveloperSettings->ActivateFailCooldownTag);
+	SyncTag(ActivateFailCostTag, DeveloperSettings->ActivateFailCostTag);
+	SyncTag(ActivateFailNetworkingTag, DeveloperSettings->ActivateFailNetworkingTag);
+	SyncTag(ActivateFailTagsBlockedTag, DeveloperSettings->ActivateFailTagsBlockedTag);
+	SyncTag(ActivateFailTagsMissingTag, DeveloperSettings->ActivateFailTagsMissingTag);
 }
 
 void UAbilitySystemGlobals::InitTargetDataScriptStructCache()
