Use FObjectKey as the unique connection identifier in FPredictionKey.

diff --git a/Engine/Plugins/Runtime/GameplayAbilities/Source/GameplayAbilities/Private/GameplayPrediction.cpp b/Engine/Plugins/Runtime/GameplayAbilities/Source/GameplayAbilities/Private/GameplayPrediction.cpp
--- a/Engine/Plugins/Runtime/GameplayAbilities/Source/GameplayAbilities/Private/GameplayPrediction.cpp
+++ b/Engine/Plugins/Runtime/GameplayAbilities/Source/GameplayAbilities/Private/GameplayPrediction.cpp
@@ -124,7 +124,7 @@ bool FPredictionKey::NetSerialize(FArchive& Ar, class UPackageMap* Map, bool& bO
 		 *	or if the owning connection is this connection (Server only sends the prediction key to the client who gave it to us)
 		 *  or if this is a server initiated key (valid on all connections)
 		 */		
-		ValidKeyForConnection = (PredictiveConnectionKey == 0 || ((UPTRINT)Map == PredictiveConnectionKey) || bIsServerInitiated) && (Current > 0);
+		ValidKeyForConnection = (Current > 0) && (bIsServerInitiated || (PredictiveConnectionObjectKey == FObjectKey()) || (PredictiveConnectionObjectKey == FObjectKey(Map)));
 	}
 	Ar.SerializeBits(&ValidKeyForConnection, 1);
 
@@ -159,7 +159,7 @@ bool FPredictionKey::NetSerialize(FArchive& Ar, class UPackageMap* Map, bool& bO
 		// We are reading this key: the connection that gave us this key is the predictive connection, and we will only serialize this key back to it.
 		if (!bIsServerInitiated)
 		{
-			PredictiveConnectionKey = (UPTRINT)Map;
+			PredictiveConnectionObjectKey = FObjectKey(Map);
 		}
 	}
 
diff --git a/Engine/Plugins/Runtime/GameplayAbilities/Source/GameplayAbilities/Private/Serialization/PredictionKeyNetSerializer.cpp b/Engine/Plugins/Runtime/GameplayAbilities/Source/GameplayAbilities/Private/Serialization/PredictionKeyNetSerializer.cpp
--- a/Engine/Plugins/Runtime/GameplayAbilities/Source/GameplayAbilities/Private/Serialization/PredictionKeyNetSerializer.cpp
+++ b/Engine/Plugins/Runtime/GameplayAbilities/Source/GameplayAbilities/Private/Serialization/PredictionKeyNetSerializer.cpp
@@ -73,7 +73,7 @@ void FPredictionKeyNetSerializer::Serialize(FNetSerializationContext& Context, c
 	FNetBitStreamWriter* Writer = Context.GetBitStreamWriter();
 
 	Writer->WriteBool(Value.bIsServerInitiated);
-	const bool bShouldWriteKey = (Value.CurrentKey > 0) && (Value.bIsServerInitiated || (Value.OwningConnectionId == 0 || Value.OwningConnectionId == Context.GetLocalConnectionId()));
+	const bool bShouldWriteKey = (Value.CurrentKey > 0) && (Value.bIsServerInitiated || (Value.OwningConnectionId == 0) || (Value.OwningConnectionId == Context.GetLocalConnectionId()));
 	if (Writer->WriteBool(bShouldWriteKey))
 	{
 		Writer->WriteBits(Value.CurrentKey, 15U);
@@ -112,11 +112,11 @@ void FPredictionKeyNetSerializer::Quantize(FNetSerializationContext&, const FNet
 		Target.CurrentKey = uint16(Source.Current);
 
 		// The owning connection is only relevant if the key is valid.
-		if (UPackageMapClient* PackageMap = Cast<UPackageMapClient>(reinterpret_cast<UPackageMap*>(Source.GetPredictiveConnectionKey())))
+		if (UPackageMapClient* PackageMap = Cast<UPackageMapClient>(Source.PredictiveConnectionObjectKey.ResolveObjectPtr()))
 		{
 			if (const UNetConnection* Connection = PackageMap->GetConnection())
 			{
-				Target.OwningConnectionId = Connection->GetConnectionId();
+				Target.OwningConnectionId = Connection->GetParentConnectionId();
 			}
 		}
 	}
@@ -134,7 +134,7 @@ void FPredictionKeyNetSerializer::Dequantize(FNetSerializationContext& Context,
 		UObject* UserData = Context.GetLocalConnectionUserData(Source.OwningConnectionId);
 		if (UNetConnection* NetConnection = Cast<UNetConnection>(UserData))
 		{
-			Target.PredictiveConnectionKey = UPTRINT(NetConnection->PackageMap);
+			Target.PredictiveConnectionObjectKey = FObjectKey(NetConnection->PackageMap);
 		}
 		else
 		{
@@ -150,8 +150,7 @@ bool FPredictionKeyNetSerializer::IsEqual(FNetSerializationContext&, const FNetI
 		const QuantizedType& Value0 = *reinterpret_cast<const QuantizedType*>(Args.Source0);
 		const QuantizedType& Value1 = *reinterpret_cast<const QuantizedType*>(Args.Source1);
 
-		const bool bIsEqual = (Value0.CurrentKey == Value1.CurrentKey) & 
-			(Value0.bIsServerInitiated == Value1.bIsServerInitiated) & (Value0.OwningConnectionId == Value1.OwningConnectionId);
+		const bool bIsEqual = (Value0.CurrentKey == Value1.CurrentKey) & (Value0.bIsServerInitiated == Value1.bIsServerInitiated) & (Value0.OwningConnectionId == Value1.OwningConnectionId);
 
 		return bIsEqual;
 	}
diff --git a/Engine/Plugins/Runtime/GameplayAbilities/Source/GameplayAbilities/Public/GameplayPrediction.h b/Engine/Plugins/Runtime/GameplayAbilities/Source/GameplayAbilities/Public/GameplayPrediction.h
--- a/Engine/Plugins/Runtime/GameplayAbilities/Source/GameplayAbilities/Public/GameplayPrediction.h
+++ b/Engine/Plugins/Runtime/GameplayAbilities/Source/GameplayAbilities/Public/GameplayPrediction.h
@@ -5,6 +5,8 @@
 #include "Engine/NetDriver.h"
 #include "Engine/NetSerialization.h"
 #include "Net/Serialization/FastArraySerializer.h"
+#include "UObject/ObjectKey.h"
+#include "Templates/TypeCompatibleBytes.h"
 #include "GameplayPrediction.generated.h"
 
 class UAbilitySystemComponent;
@@ -355,12 +357,12 @@ struct GAMEPLAYABILITIES_API FPredictionKey
 	/** Was this PredictionKey received from a NetSerialize or created locally? */
 	bool WasReceived() const
 	{
-		return PredictiveConnectionKey != 0;
+		return PredictiveConnectionObjectKey != FObjectKey();
 	}
 
 	bool WasLocallyGenerated() const
 	{
-		return (Current > 0) && (PredictiveConnectionKey == 0);
+		return (Current > 0) && (PredictiveConnectionObjectKey == FObjectKey());
 	}
 
 	bool operator==(const FPredictionKey& Other) const
@@ -385,7 +387,7 @@ struct GAMEPLAYABILITIES_API FPredictionKey
 		return ((InKey.Current << 1) | (InKey.bIsServerInitiated & 1));
 	}
 
-	UPTRINT GetPredictiveConnectionKey() const { return PredictiveConnectionKey; }
+	uint64 GetPredictiveConnectionKey() const { return BitCast<uint64>(PredictiveConnectionObjectKey); }
 
 private:
 	friend UE::Net::FPredictionKeyNetSerializer;
@@ -399,7 +401,7 @@ private:
 	}
 
 	/** On the server, uniquely identifies network connection this was serialized on/from.  See NetSerialize for additional information. */
-	UPTRINT PredictiveConnectionKey = 0;
+	FObjectKey PredictiveConnectionObjectKey;
 };
 PRAGMA_ENABLE_DEPRECATION_WARNINGS
 
