Remove CoreMinimal include from GameplayEffectComponent and GameplayEffectCustomApplicationRequirement.
Also updated GENERATED_BODY from legacy version.

#rnx

diff --git a/Engine/Plugins/Runtime/GameplayAbilities/Source/GameplayAbilities/Private/GameplayEffectCustomApplicationRequirement.cpp b/Engine/Plugins/Runtime/GameplayAbilities/Source/GameplayAbilities/Private/GameplayEffectCustomApplicationRequirement.cpp
--- a/Engine/Plugins/Runtime/GameplayAbilities/Source/GameplayAbilities/Private/GameplayEffectCustomApplicationRequirement.cpp
+++ b/Engine/Plugins/Runtime/GameplayAbilities/Source/GameplayAbilities/Private/GameplayEffectCustomApplicationRequirement.cpp
@@ -4,12 +4,6 @@
 
 #include UE_INLINE_GENERATED_CPP_BY_NAME(GameplayEffectCustomApplicationRequirement)
 
-UGameplayEffectCustomApplicationRequirement::UGameplayEffectCustomApplicationRequirement(const FObjectInitializer& ObjectInitializer)
-	: Super(ObjectInitializer)
-{
-
-}
-
 bool UGameplayEffectCustomApplicationRequirement::CanApplyGameplayEffect_Implementation(const UGameplayEffect* GameplayEffect, const FGameplayEffectSpec& Spec, UAbilitySystemComponent* ASC) const
 {
 	return true;
diff --git a/Engine/Plugins/Runtime/GameplayAbilities/Source/GameplayAbilities/Public/GameplayEffectComponent.h b/Engine/Plugins/Runtime/GameplayAbilities/Source/GameplayAbilities/Public/GameplayEffectComponent.h
--- a/Engine/Plugins/Runtime/GameplayAbilities/Source/GameplayAbilities/Public/GameplayEffectComponent.h
+++ b/Engine/Plugins/Runtime/GameplayAbilities/Source/GameplayAbilities/Public/GameplayEffectComponent.h
@@ -2,7 +2,7 @@
 
 #pragma once
 
-#include "CoreMinimal.h"
+#include "UObject/Object.h"
 #include "GameplayEffectComponent.generated.h"
 
 struct FActiveGameplayEffect;
diff --git a/Engine/Plugins/Runtime/GameplayAbilities/Source/GameplayAbilities/Public/GameplayEffectCustomApplicationRequirement.h b/Engine/Plugins/Runtime/GameplayAbilities/Source/GameplayAbilities/Public/GameplayEffectCustomApplicationRequirement.h
--- a/Engine/Plugins/Runtime/GameplayAbilities/Source/GameplayAbilities/Public/GameplayEffectCustomApplicationRequirement.h
+++ b/Engine/Plugins/Runtime/GameplayAbilities/Source/GameplayAbilities/Public/GameplayEffectCustomApplicationRequirement.h
@@ -2,7 +2,7 @@
 
 #pragma once
 
-#include "CoreMinimal.h"
+#include "UObject/Object.h"
 #include "GameplayEffectCustomApplicationRequirement.generated.h"
 
 struct FGameplayEffectSpec;
@@ -13,10 +13,9 @@ class UAbilitySystemComponent;
 UCLASS(BlueprintType, Blueprintable, Abstract)
 class GAMEPLAYABILITIES_API UGameplayEffectCustomApplicationRequirement : public UObject
 {
+	GENERATED_BODY()
 
 public:
-	GENERATED_UCLASS_BODY()
-	
 	/** Return whether the gameplay effect should be applied or not */
 	UFUNCTION(BlueprintNativeEvent, Category="Calculation")
 	bool CanApplyGameplayEffect(const UGameplayEffect* GameplayEffect, const FGameplayEffectSpec& Spec, UAbilitySystemComponent* ASC) const;
