Adding clearer documentation to RemoveActiveGameplayEffect.

#tests Tested in editor

diff --git a/Engine/Plugins/Runtime/GameplayAbilities/Source/GameplayAbilities/Public/AbilitySystemComponent.h b/Engine/Plugins/Runtime/GameplayAbilities/Source/GameplayAbilities/Public/AbilitySystemComponent.h
--- a/Engine/Plugins/Runtime/GameplayAbilities/Source/GameplayAbilities/Public/AbilitySystemComponent.h
+++ b/Engine/Plugins/Runtime/GameplayAbilities/Source/GameplayAbilities/Public/AbilitySystemComponent.h
@@ -342,7 +342,13 @@ class GAMEPLAYABILITIES_API UAbilitySystemComponent : public UGameplayTasksCompo
 	/** Gets the FActiveGameplayEffect based on the passed in Handle */
 	const UGameplayEffect* GetGameplayEffectDefForHandle(FActiveGameplayEffectHandle Handle);
 
-	/** Removes GameplayEffect by Handle. StacksToRemove=-1 will remove all stacks. */
+	/** 
+	* Removes the specified GameplayEffect by Handle.
+	* @param Handle								The handle of the gameplay effect to remove
+	* @param StacksToRemove						Number of stacks to remove, -1 means remove all
+	* 
+	* @return Whether the effect was successfully removed.
+	*/
 	UFUNCTION(BlueprintCallable, BlueprintAuthorityOnly, Category = GameplayEffects)
 	virtual bool RemoveActiveGameplayEffect(FActiveGameplayEffectHandle Handle, int32 StacksToRemove=-1);
 
