PR #11970: Fix network compatibility between editor and non-editor builds

diff --git a/Engine/Plugins/Runtime/GameplayAbilities/Source/GameplayAbilities/Public/Abilities/GameplayAbilityRepAnimMontage.h b/Engine/Plugins/Runtime/GameplayAbilities/Source/GameplayAbilities/Public/Abilities/GameplayAbilityRepAnimMontage.h
--- a/Engine/Plugins/Runtime/GameplayAbilities/Source/GameplayAbilities/Public/Abilities/GameplayAbilityRepAnimMontage.h
+++ b/Engine/Plugins/Runtime/GameplayAbilities/Source/GameplayAbilities/Public/Abilities/GameplayAbilityRepAnimMontage.h
@@ -22,12 +22,6 @@ struct GAMEPLAYABILITIES_API FGameplayAbilityRepAnimMontage
 {
 	GENERATED_USTRUCT_BODY()
 
-#if WITH_EDITORONLY_DATA
-	/** AnimMontage ref */
-	UPROPERTY(meta = (DeprecatedProperty, DeprecationMessage = "Use the GetAnimMontage function instead"))
-	TObjectPtr<UAnimMontage> AnimMontage_DEPRECATED;
-#endif
-
 	/** Animation ref. When playing a dynamic montage this points to the AnimSequence the montage was created from */
 	UPROPERTY()
 	TObjectPtr<UAnimSequenceBase> Animation;
diff --git a/Engine/Source/Runtime/Engine/Classes/GameFramework/Character.h b/Engine/Source/Runtime/Engine/Classes/GameFramework/Character.h
--- a/Engine/Source/Runtime/Engine/Classes/GameFramework/Character.h
+++ b/Engine/Source/Runtime/Engine/Classes/GameFramework/Character.h
@@ -48,12 +48,6 @@ struct FRepRootMotionMontage
 {
 	GENERATED_USTRUCT_BODY()
 
-#if WITH_EDITORONLY_DATA
-	/** AnimMontage providing Root Motion */
-	UPROPERTY(meta = (DeprecatedProperty, DeprecationMessage = "Use the GetAnimMontage function instead"))
-	TObjectPtr<UAnimMontage> AnimMontage_DEPRECATED = nullptr;
-#endif
-	
 	/** Animation providing Root Motion */
 	UPROPERTY()
 	TObjectPtr<UAnimSequenceBase> Animation = nullptr;
