Gameplay Effect / Engineering: Gameplay Effect Handles are now no longer removed from the global map while the Gameplay Effect is locked.
Instead they will now be removed while handling pending removes after the Gameplay Effect is no longer locked.

diff --git a/Engine/Plugins/Runtime/GameplayAbilities/Source/GameplayAbilities/Private/GameplayEffect.cpp b/Engine/Plugins/Runtime/GameplayAbilities/Source/GameplayAbilities/Private/GameplayEffect.cpp
--- a/Engine/Plugins/Runtime/GameplayAbilities/Source/GameplayAbilities/Private/GameplayEffect.cpp
+++ b/Engine/Plugins/Runtime/GameplayAbilities/Source/GameplayAbilities/Private/GameplayEffect.cpp
@@ -4549,9 +4549,6 @@ bool FActiveGameplayEffectsContainer::InternalRemoveActiveGameplayEffect(int32 I
 			}
 		}
 
-		// Remove this handle from the global map
-		Effect.Handle.RemoveFromGlobalMap();
-
 		bool ModifiedArray = false;
 
 		// Finally remove the ActiveGameplayEffect
@@ -4566,6 +4563,9 @@ bool FActiveGameplayEffectsContainer::InternalRemoveActiveGameplayEffect(int32 I
 		{
 			// Not locked, so do the removal right away.
 
+			// Remove this handle from the global map
+			Effect.Handle.RemoveFromGlobalMap();
+
 			// If we are not scope locked, then there is no way this idx should be referring to something on the pending add list.
 			// It is possible to remove a GE that is pending add, but it would happen while the scope lock is still in effect, resulting
 			// in a pending remove being set.
@@ -6141,6 +6141,10 @@ void FActiveGameplayEffectsContainer::DecrementLock()
 			if (Effect.IsPendingRemove)
 			{
 				UE_LOG(LogGameplayEffects, Verbose, TEXT("%s: Finish PendingRemove: %s. Auth: %d"), *GetNameSafe(Owner->GetOwnerActor()), *Effect.GetDebugString(), IsNetAuthority());
+				
+				// Remove this handle from the global map
+				Effect.Handle.RemoveFromGlobalMap();
+
 				GameplayEffects_Internal.RemoveAtSwap(idx, EAllowShrinking::No);
 				ModifiedArray = true;
 				PendingRemoves--;
