Add checks to prevent adding invalid tags (consistent with FGameplayTagContainer).

FGameplayTagCountContainer:
- Added ensure in UpdateTagMap_Internal. Logic should return early before reaching this point.
- Early returns in UpdateTagCount and SetTagCount if the tag is invalid.

FMinimalReplicationTagCountMap:
- Early returns in AddTag and SetTagCount if the tag is invalid.

GameplayCues:
- Early returns in GameplayCueManager's InvokeGameplayCueExecuted and InvokeGameplayCueExecuted.
- Early return in UAbilitySystemComponent::AddGameplayCue_Internal.

diff --git a/Engine/Plugins/Runtime/GameplayAbilities/Source/GameplayAbilities/Private/AbilitySystemComponent.cpp b/Engine/Plugins/Runtime/GameplayAbilities/Source/GameplayAbilities/Private/AbilitySystemComponent.cpp
--- a/Engine/Plugins/Runtime/GameplayAbilities/Source/GameplayAbilities/Private/AbilitySystemComponent.cpp
+++ b/Engine/Plugins/Runtime/GameplayAbilities/Source/GameplayAbilities/Private/AbilitySystemComponent.cpp
@@ -1340,6 +1340,11 @@ void UAbilitySystemComponent::AddGameplayCue_Internal(const FGameplayTag Gamepla
 {
 	if (IsOwnerActorAuthoritative())
 	{
+		if (!GameplayCueTag.IsValid())
+		{
+			return;
+		}
+
 		const bool bWasInList = GameplayCueContainer.HasCue(GameplayCueTag);
 
 		ForceReplication();
diff --git a/Engine/Plugins/Runtime/GameplayAbilities/Source/GameplayAbilities/Private/GameplayCueManager.cpp b/Engine/Plugins/Runtime/GameplayAbilities/Source/GameplayAbilities/Private/GameplayCueManager.cpp
--- a/Engine/Plugins/Runtime/GameplayAbilities/Source/GameplayAbilities/Private/GameplayCueManager.cpp
+++ b/Engine/Plugins/Runtime/GameplayAbilities/Source/GameplayAbilities/Private/GameplayCueManager.cpp
@@ -1345,6 +1345,11 @@ void UGameplayCueManager::InvokeGameplayCueExecuted_FromSpec(UAbilitySystemCompo
 
 void UGameplayCueManager::InvokeGameplayCueExecuted(UAbilitySystemComponent* OwningComponent, const FGameplayTag GameplayCueTag, FPredictionKey PredictionKey, FGameplayEffectContextHandle EffectContext)
 {
+	if (!GameplayCueTag.IsValid())
+	{
+		return;
+	}
+
 	if (EnableSuppressCuesOnGameplayCueManager && OwningComponent && OwningComponent->bSuppressGameplayCues)
 	{
 		return;
@@ -1365,6 +1370,11 @@ void UGameplayCueManager::InvokeGameplayCueExecuted(UAbilitySystemComponent* Own
 
 void UGameplayCueManager::InvokeGameplayCueExecuted_WithParams(UAbilitySystemComponent* OwningComponent, const FGameplayTag GameplayCueTag, FPredictionKey PredictionKey, FGameplayCueParameters GameplayCueParameters)
 {
+	if (!GameplayCueTag.IsValid())
+	{
+		return;
+	}
+
 	if (EnableSuppressCuesOnGameplayCueManager && OwningComponent && OwningComponent->bSuppressGameplayCues)
 	{
 		return;
diff --git a/Engine/Plugins/Runtime/GameplayAbilities/Source/GameplayAbilities/Private/GameplayEffectTypes.cpp b/Engine/Plugins/Runtime/GameplayAbilities/Source/GameplayAbilities/Private/GameplayEffectTypes.cpp
--- a/Engine/Plugins/Runtime/GameplayAbilities/Source/GameplayAbilities/Private/GameplayEffectTypes.cpp
+++ b/Engine/Plugins/Runtime/GameplayAbilities/Source/GameplayAbilities/Private/GameplayEffectTypes.cpp
@@ -616,6 +616,11 @@ bool FGameplayTagCountContainer::GatherTagChangeDelegates(const FGameplayTag& Ta
 
 bool FGameplayTagCountContainer::UpdateTagMap_Internal(const FGameplayTag& Tag, int32 CountDelta)
 {
+	if (!ensureMsgf(Tag.IsValid(), TEXT("FGameplayTagCountContainer attempted to update tag map with an invalid Tag (None)!")))
+	{
+		return false;
+	}
+
 	if (!UpdateExplicitTags(Tag, CountDelta, false))
 	{
 		return false;
diff --git a/Engine/Plugins/Runtime/GameplayAbilities/Source/GameplayAbilities/Public/GameplayEffectTypes.h b/Engine/Plugins/Runtime/GameplayAbilities/Source/GameplayAbilities/Public/GameplayEffectTypes.h
--- a/Engine/Plugins/Runtime/GameplayAbilities/Source/GameplayAbilities/Public/GameplayEffectTypes.h
+++ b/Engine/Plugins/Runtime/GameplayAbilities/Source/GameplayAbilities/Public/GameplayEffectTypes.h
@@ -1147,7 +1147,7 @@ struct GAMEPLAYABILITIES_API FGameplayTagCountContainer
 	 */
 	FORCEINLINE bool UpdateTagCount(const FGameplayTag& Tag, int32 CountDelta)
 	{
-		if (CountDelta != 0)
+		if (Tag.IsValid() && CountDelta != 0)
 		{
 			return UpdateTagMap_Internal(Tag, CountDelta);
 		}
@@ -1185,6 +1185,11 @@ struct GAMEPLAYABILITIES_API FGameplayTagCountContainer
 	 */
 	FORCEINLINE bool SetTagCount(const FGameplayTag& Tag, int32 NewCount)
 	{
+		if (!Tag.IsValid())
+		{
+			return false;
+		}
+
 		int32 ExistingCount = 0;
 		if (int32* Ptr  = ExplicitTagCountMap.Find(Tag))
 		{
@@ -1585,6 +1590,11 @@ struct GAMEPLAYABILITIES_API FMinimalReplicationTagCountMap
 
 	void AddTag(const FGameplayTag& Tag)
 	{
+		if (!Tag.IsValid())
+		{
+			return;
+		}
+		
 		MapID++;
 		TagMap.FindOrAdd(Tag)++;
 	}
@@ -1609,6 +1619,11 @@ struct GAMEPLAYABILITIES_API FMinimalReplicationTagCountMap
 
 	void SetTagCount(const FGameplayTag& Tag, int32 NewCount)
 	{
+		if (!Tag.IsValid())
+		{
+			return;
+		}
+
 		MapID++;
 		int32& Count = TagMap.FindOrAdd(Tag);
 		Count = FMath::Max(0, NewCount);
