Update UAbilitySystemComponent::GetAnimatingAbility to be const.

#rnx

diff --git a/Engine/Plugins/Runtime/GameplayAbilities/Source/GameplayAbilities/Private/AbilitySystemComponent_Abilities.cpp b/Engine/Plugins/Runtime/GameplayAbilities/Source/GameplayAbilities/Private/AbilitySystemComponent_Abilities.cpp
--- a/Engine/Plugins/Runtime/GameplayAbilities/Source/GameplayAbilities/Private/AbilitySystemComponent_Abilities.cpp
+++ b/Engine/Plugins/Runtime/GameplayAbilities/Source/GameplayAbilities/Private/AbilitySystemComponent_Abilities.cpp
@@ -3748,7 +3748,7 @@ bool UAbilitySystemComponent::IsAnimatingAbility(UGameplayAbility* InAbility) co
 	return (LocalAnimMontageInfo.AnimatingAbility.Get() == InAbility);
 }
 
-UGameplayAbility* UAbilitySystemComponent::GetAnimatingAbility()
+UGameplayAbility* UAbilitySystemComponent::GetAnimatingAbility() const
 {
 	return LocalAnimMontageInfo.AnimatingAbility.Get();
 }
diff --git a/Engine/Plugins/Runtime/GameplayAbilities/Source/GameplayAbilities/Public/AbilitySystemComponent.h b/Engine/Plugins/Runtime/GameplayAbilities/Source/GameplayAbilities/Public/AbilitySystemComponent.h
--- a/Engine/Plugins/Runtime/GameplayAbilities/Source/GameplayAbilities/Public/AbilitySystemComponent.h
+++ b/Engine/Plugins/Runtime/GameplayAbilities/Source/GameplayAbilities/Public/AbilitySystemComponent.h
@@ -1486,7 +1486,7 @@ class GAMEPLAYABILITIES_API UAbilitySystemComponent : public UGameplayTasksCompo
 	bool IsAnimatingAbility(UGameplayAbility* Ability) const;
 
 	/** Returns the current animating ability */
-	UGameplayAbility* GetAnimatingAbility();
+	UGameplayAbility* GetAnimatingAbility() const;
 
 	/** Returns montage that is currently playing */
 	UAnimMontage* GetCurrentMontage() const;
