Add Blueprint auto cast functions for FScalableFloat's to single and double precision float values. This is a small but nice UX improvement when working with scalable floats in blueprints.

diff --git a/Engine/Plugins/Runtime/GameplayAbilities/Source/GameplayAbilities/Private/AbilitySystemBlueprintLibrary.cpp b/Engine/Plugins/Runtime/GameplayAbilities/Source/GameplayAbilities/Private/AbilitySystemBlueprintLibrary.cpp
--- a/Engine/Plugins/Runtime/GameplayAbilities/Source/GameplayAbilities/Private/AbilitySystemBlueprintLibrary.cpp
+++ b/Engine/Plugins/Runtime/GameplayAbilities/Source/GameplayAbilities/Private/AbilitySystemBlueprintLibrary.cpp
@@ -1260,3 +1260,13 @@ bool UAbilitySystemBlueprintLibrary::NotEqual_GameplayAbilitySpecHandle(const FG
 {
 	return A != B;
 }
+
+float UAbilitySystemBlueprintLibrary::Conv_ScalableFloatToFloat(const FScalableFloat& Input, float Level)
+{
+	return Input.GetValueAtLevel(Level);
+}
+
+double UAbilitySystemBlueprintLibrary::Conv_ScalableFloatToDouble(const FScalableFloat& Input, float Level)
+{
+	return static_cast<double>(Input.GetValueAtLevel(Level));
+}
diff --git a/Engine/Plugins/Runtime/GameplayAbilities/Source/GameplayAbilities/Public/AbilitySystemBlueprintLibrary.h b/Engine/Plugins/Runtime/GameplayAbilities/Source/GameplayAbilities/Public/AbilitySystemBlueprintLibrary.h
--- a/Engine/Plugins/Runtime/GameplayAbilities/Source/GameplayAbilities/Public/AbilitySystemBlueprintLibrary.h
+++ b/Engine/Plugins/Runtime/GameplayAbilities/Source/GameplayAbilities/Public/AbilitySystemBlueprintLibrary.h
@@ -14,6 +14,7 @@
 #include "Abilities/GameplayAbilityTypes.h"
 #include "Kismet/BlueprintFunctionLibrary.h"
 #include "Abilities/GameplayAbilityTargetDataFilter.h"
+#include "ScalableFloat.h"
 #include "AbilitySystemBlueprintLibrary.generated.h"
 
 class UAbilitySystemComponent;
@@ -456,4 +457,24 @@ class GAMEPLAYABILITIES_API UAbilitySystemBlueprintLibrary : public UBlueprintFu
 	/** Inequality operator for two Gameplay Ability Spec Handles */
 	UFUNCTION(BlueprintPure, Category = "Ability|GameplayAbility", meta = (DisplayName = "Not Equal (Gameplay Ability Spec Handle)", CompactNodeTitle = "!=", ScriptOperator = "!="))
 	static bool NotEqual_GameplayAbilitySpecHandle(const FGameplayAbilitySpecHandle& A, const FGameplayAbilitySpecHandle& B);
+
+	/**
+	 * Converts the given Scalable float to the value at the given level.
+	 * 
+	 * @param Input	The scalable float to get the value from 
+	 * @param Level The Level of which to get the value at
+	 * @return The single-precision float value at the given level from the scalable float
+	 */
+	UFUNCTION(BlueprintPure, Category="ScalableFloat", meta=(DisplayName="Get Value At Level (single-precision)", BlueprintAutocast, IgnoreTypePromotion))
+	static float Conv_ScalableFloatToFloat(const FScalableFloat& Input, float Level = 0.0f);
+
+	/**
+	 * Converts the given Scalable float to the value at the given level.
+	 * 
+	 * @param Input	The scalable float to get the value from 
+	 * @param Level The Level of which to get the value at
+	 * @return The double-precision float value at the given level from the scalable float
+	 */
+	UFUNCTION(BlueprintPure, Category="ScalableFloat", meta=(DisplayName="Get Value At Level (double-precision)", BlueprintAutocast, IgnoreTypePromotion))
+	static double Conv_ScalableFloatToDouble(const FScalableFloat& Input, float Level = 0.0f);
 };
