Initial OnRep_SpawnedAttributes from DefaultStartingData no longer clears attribute aggregators for already replicated ActiveGameplayEffects. This fixes a bug where client-side attributes will have the wrong value when modifiers are applied before the first time SpawnedAttributes are replicated over from the server.

diff --git a/Engine/Plugins/Runtime/GameplayAbilities/Source/GameplayAbilities/Private/AbilitySystemComponent.cpp b/Engine/Plugins/Runtime/GameplayAbilities/Source/GameplayAbilities/Private/AbilitySystemComponent.cpp
--- a/Engine/Plugins/Runtime/GameplayAbilities/Source/GameplayAbilities/Private/AbilitySystemComponent.cpp
+++ b/Engine/Plugins/Runtime/GameplayAbilities/Source/GameplayAbilities/Private/AbilitySystemComponent.cpp
@@ -3103,17 +3103,57 @@ void UAbilitySystemComponent::OnRep_SpawnedAttributes(const TArray<UAttributeSet
 		}
 	}
 
-	// Find the attribute sets that got removed
+	// Create a map of attribute set class to latest object to detect removed and replaced attribute sets.
+	TMap<TSubclassOf<UAttributeSet>, UAttributeSet*> NewAttributeSetClasses;
+	for (UAttributeSet* NewAttributeSet : SpawnedAttributes)
+	{
+		if (IsValid(NewAttributeSet))
+		{
+			NewAttributeSetClasses.Add(NewAttributeSet->GetClass(), NewAttributeSet);
+		}
+	}
+
+	// Find the attribute sets that got removed or replaced by one of the same class. Replacement here means that the an attribute set instance
+	// is no longer in SpawnedAttributes, but a new one exists of the same class. This happens when the client and server both created an attribute set
+	// and the server's instance is replicated over, for example when using DefaultStartingData. For removed attribute sets, their aggregators should be 
+	// cleaned up. For replaced attribute sets, the aggregators should remain valid.
 	for (UAttributeSet* PreviousAttributeSet : PreviousSpawnedAttributes)
 	{
-		if (PreviousAttributeSet && SpawnedAttributes.Find(PreviousAttributeSet) == INDEX_NONE)
+		if (PreviousAttributeSet)
 		{
+			UClass* AttributeSetClass = PreviousAttributeSet->GetClass();
 			TArray<FGameplayAttribute> Attributes;
 			UAttributeSet::GetAttributesFromSetClass(PreviousAttributeSet->GetClass(), Attributes);
-			for (const FGameplayAttribute& Attribute : Attributes)
+
+			// Case: Attribute set removed and not replaced
+			if (!NewAttributeSetClasses.Contains(PreviousAttributeSet->GetClass()))
+			{
+				// The attribute set no longer exists after removal on the server.
+				// Remove any aggregators for all the attributes in it.
+				for (const FGameplayAttribute& Attribute : Attributes)
+				{
+					ABILITY_LOG(Log, TEXT("Cleaning up aggregator for attribute '%s' due to OnRep_SpawnedAttributes detecting removal of '%s' without replacement"), *Attribute.GetName(), *PreviousAttributeSet->GetName());
+					ActiveGameplayEffects.CleanupAttributeAggregator(Attribute);
+				}
+			}
+			// Case: Attribute set replaced by server-replicated object
+			else if (PreviousAttributeSet != NewAttributeSetClasses[AttributeSetClass])
 			{
-				ABILITY_LOG(Log, TEXT("Cleaning up aggregator for attribute '%s' due to OnRep_SpawnedAttributes detecting removal of '%s'"), *Attribute.GetName(), *PreviousAttributeSet->GetName());
-				ActiveGameplayEffects.CleanupAttributeAggregator(Attribute);
+				UAttributeSet* NewAttributeSet = NewAttributeSetClasses[AttributeSetClass];
+				for (const FGameplayAttribute& Attribute : Attributes)
+				{
+					ABILITY_LOG(Log, TEXT("Keeping aggregator for attribute '%s' due to OnRep_SpawnedAttributes detecting replacement '%s' for '%s'"), *Attribute.GetName(), *NewAttributeSet->GetName(), *PreviousAttributeSet->GetName());
+
+					// Since we're replacing the attribute set with the server replicated one, give it a chance to do setup on the 
+					// existing aggregator like metadata and binding to its delegates. However, only do this for non-system attributes, 
+					// mirroring behavior in FindOrCreateAttributeAggregator(). Also don't actually created the aggregator if it doesn't 
+					// exist. Aggregators are created on-demand when they are first modified.
+					if (Attribute.IsSystemAttribute() == false && ActiveGameplayEffects.AttributeAggregatorMap.Contains(Attribute))
+					{
+						FAggregatorRef& Aggregator = ActiveGameplayEffects.FindOrCreateAttributeAggregator(Attribute);
+						NewAttributeSet->OnAttributeAggregatorCreated(Attribute, Aggregator.Get());
+					}
+				}
 			}
 		}
 	}
