Added an optional _DontCreate method to get ASC that implementer of the interface can use to avoid creating ASC when it's not, for perfmem reasons

diff --git a/Engine/Plugins/Runtime/GameplayAbilities/Source/GameplayAbilities/Private/AbilitySystemGlobals.cpp b/Engine/Plugins/Runtime/GameplayAbilities/Source/GameplayAbilities/Private/AbilitySystemGlobals.cpp
--- a/Engine/Plugins/Runtime/GameplayAbilities/Source/GameplayAbilities/Private/AbilitySystemGlobals.cpp
+++ b/Engine/Plugins/Runtime/GameplayAbilities/Source/GameplayAbilities/Private/AbilitySystemGlobals.cpp
@@ -252,6 +252,22 @@ UAbilitySystemComponent* UAbilitySystemGlobals::GetAbilitySystemComponentFromAct
 	return nullptr;
 }
 
+UAbilitySystemComponent* UAbilitySystemGlobals::GetAbilitySystemComponentFromActor_DontCreate(const AActor* Actor, bool LookForComponent)
+{
+	if (const IAbilitySystemInterface* const ASI = Cast<IAbilitySystemInterface>(Actor))
+	{
+		return ASI->GetAbilitySystemComponent_DontCreate();
+	}
+
+	if (Actor && LookForComponent)
+	{
+		// Fall back to a component search to better support BP-only actors
+		return Actor->FindComponentByClass<UAbilitySystemComponent>();
+	}
+
+	return nullptr;
+}
+
 bool UAbilitySystemGlobals::ShouldPredictTargetGameplayEffects() const
 {
 	return GetDefault<UGameplayAbilitiesDeveloperSettings>()->PredictTargetGameplayEffects;
diff --git a/Engine/Plugins/Runtime/GameplayAbilities/Source/GameplayAbilities/Public/AbilitySystemGlobals.h b/Engine/Plugins/Runtime/GameplayAbilities/Source/GameplayAbilities/Public/AbilitySystemGlobals.h
--- a/Engine/Plugins/Runtime/GameplayAbilities/Source/GameplayAbilities/Public/AbilitySystemGlobals.h
+++ b/Engine/Plugins/Runtime/GameplayAbilities/Source/GameplayAbilities/Public/AbilitySystemGlobals.h
@@ -71,6 +71,9 @@ class GAMEPLAYABILITIES_API UAbilitySystemGlobals : public UObject
 
 	/** Searches the passed in actor for an ability system component, will use IAbilitySystemInterface or fall back to a component search */
 	static UAbilitySystemComponent* GetAbilitySystemComponentFromActor(const AActor* Actor, bool LookForComponent=true);
+	
+	/** Same as 'GetAbilitySystemComponentFromActor', but calls the DontCreate method on the interface */
+	static UAbilitySystemComponent* GetAbilitySystemComponentFromActor_DontCreate(const AActor* Actor, bool LookForComponent=true);
 
 	/** Should allocate a project specific AbilityActorInfo struct. Caller is responsible for deallocation */
 	virtual FGameplayAbilityActorInfo* AllocAbilityActorInfo() const;
diff --git a/Engine/Plugins/Runtime/GameplayAbilities/Source/GameplayAbilities/Public/AbilitySystemInterface.h b/Engine/Plugins/Runtime/GameplayAbilities/Source/GameplayAbilities/Public/AbilitySystemInterface.h
--- a/Engine/Plugins/Runtime/GameplayAbilities/Source/GameplayAbilities/Public/AbilitySystemInterface.h
+++ b/Engine/Plugins/Runtime/GameplayAbilities/Source/GameplayAbilities/Public/AbilitySystemInterface.h
@@ -22,4 +22,10 @@ class GAMEPLAYABILITIES_API IAbilitySystemInterface
 
 	/** Returns the ability system component to use for this actor. It may live on another actor, such as a Pawn using the PlayerState's component */
 	virtual UAbilitySystemComponent* GetAbilitySystemComponent() const = 0;
+	
+	/** If the implementer supports this, avoid creating the ASC if it's nullptr. */
+	virtual UAbilitySystemComponent* GetAbilitySystemComponent_DontCreate() const
+	{
+		return GetAbilitySystemComponent();
+	}
 };
