Add optional PlayCount parameter to PlayAnimation functions

#okforversepublic https://docs.google.com/document/d/1EFFxkaN_KJOhN8itTwyp5u_zJ0_b8_XUNsZLJAlttKg/edit?tab=t.0#heading=h.m30bw1da0jpg

diff --git a/Engine/Plugins/Runtime/GameplayAbilities/Source/GameplayAbilities/Private/Abilities/Tasks/AbilityTask_PlayAnimAndWait.cpp b/Engine/Plugins/Runtime/GameplayAbilities/Source/GameplayAbilities/Private/Abilities/Tasks/AbilityTask_PlayAnimAndWait.cpp
--- a/Engine/Plugins/Runtime/GameplayAbilities/Source/GameplayAbilities/Private/Abilities/Tasks/AbilityTask_PlayAnimAndWait.cpp
+++ b/Engine/Plugins/Runtime/GameplayAbilities/Source/GameplayAbilities/Private/Abilities/Tasks/AbilityTask_PlayAnimAndWait.cpp
@@ -84,24 +84,26 @@ UAbilityTask_PlayAnimAndWait* UAbilityTask_PlayAnimAndWait::CreatePlayAnimAndWai
 	FName TaskInstanceName, 
 	UAnimSequence* AnimSequence, 
 	FName SlotName, 
-	float BlendInTime, 
-	float BlendOutTime, 
-	float InPlayRate,
+	float BlendInTime /*= 0.25f*/,
+	float BlendOutTime /*= 0.25f*/,
+	float InPlayRate /*= 1.0f*/,
+	float InPlayCount /*= 1.0f*/,
 	float StartTimeSeconds /*= 0.0f*/, 
 	bool bStopWhenAbilityEnds /*= true*/, 
 	float AnimRootMotionTranslationScale /*= 1.f*/)
 {
-	UAbilityTask_PlayAnimAndWait* MyObj = NewAbilityTask<UAbilityTask_PlayAnimAndWait>(OwningAbility, TaskInstanceName);
-	MyObj->AnimSequenceToPlay = AnimSequence;
-	MyObj->SlotName = SlotName;
-	MyObj->BlendInTime = BlendInTime;
-	MyObj->BlendOutTime = BlendOutTime;
-	MyObj->PlayRate = InPlayRate;
-	MyObj->AnimRootMotionTranslationScale = AnimRootMotionTranslationScale;
-	MyObj->bStopWhenAbilityEnds = bStopWhenAbilityEnds;
-	MyObj->StartTimeSeconds = StartTimeSeconds;
-
-	return MyObj;
+	UAbilityTask_PlayAnimAndWait* PlayAnimTask = NewAbilityTask<UAbilityTask_PlayAnimAndWait>(OwningAbility, TaskInstanceName);
+	PlayAnimTask->AnimSequenceToPlay = AnimSequence;
+	PlayAnimTask->SlotName = SlotName;
+	PlayAnimTask->BlendInTime = BlendInTime;
+	PlayAnimTask->BlendOutTime = BlendOutTime;
+	PlayAnimTask->PlayRate = InPlayRate;
+	PlayAnimTask->PlayCount = InPlayCount;
+	PlayAnimTask->AnimRootMotionTranslationScale = AnimRootMotionTranslationScale;
+	PlayAnimTask->bStopWhenAbilityEnds = bStopWhenAbilityEnds;
+	PlayAnimTask->StartTimeSeconds = StartTimeSeconds;
+
+	return PlayAnimTask;
 }
 
 void UAbilityTask_PlayAnimAndWait::Activate()
@@ -119,7 +121,7 @@ void UAbilityTask_PlayAnimAndWait::Activate()
 		UAnimInstance* AnimInstance = ActorInfo->GetAnimInstance();
 		if (AnimInstance != nullptr)
 		{
-			UAnimMontage* NewDynamicMontage = ASC->PlaySlotAnimationAsDynamicMontage(Ability, Ability->GetCurrentActivationInfo(), AnimSequenceToPlay, SlotName, BlendInTime, BlendOutTime, PlayRate, StartTimeSeconds);
+			UAnimMontage* NewDynamicMontage = ASC->PlaySlotAnimationAsDynamicMontage_WithFractionalLoops(Ability, Ability->GetCurrentActivationInfo(), AnimSequenceToPlay, SlotName, BlendInTime, BlendOutTime, PlayRate, StartTimeSeconds, PlayCount);
 			if (NewDynamicMontage)
 			{
 				
diff --git a/Engine/Plugins/Runtime/GameplayAbilities/Source/GameplayAbilities/Private/AbilitySystemComponent_Abilities.cpp b/Engine/Plugins/Runtime/GameplayAbilities/Source/GameplayAbilities/Private/AbilitySystemComponent_Abilities.cpp
--- a/Engine/Plugins/Runtime/GameplayAbilities/Source/GameplayAbilities/Private/AbilitySystemComponent_Abilities.cpp
+++ b/Engine/Plugins/Runtime/GameplayAbilities/Source/GameplayAbilities/Private/AbilitySystemComponent_Abilities.cpp
@@ -2989,6 +2989,11 @@ void UAbilitySystemComponent::ClearDebugInstantEffects()
 // ---------------------------------------------------------------------------
 
 float UAbilitySystemComponent::PlayMontage(UGameplayAbility* InAnimatingAbility, FGameplayAbilityActivationInfo ActivationInfo, UAnimMontage* NewAnimMontage, float InPlayRate, FName StartSectionName, float StartTimeSeconds)
+{
+	return PlayMontageInternal(InAnimatingAbility, ActivationInfo, NewAnimMontage, InPlayRate, StartSectionName, StartTimeSeconds, [](FGameplayAbilityRepAnimMontage&){});
+}
+
+float UAbilitySystemComponent::PlayMontageInternal(UGameplayAbility* InAnimatingAbility, FGameplayAbilityActivationInfo ActivationInfo, UAnimMontage* NewAnimMontage, float InPlayRate, FName StartSectionName, float StartTimeSeconds, TFunctionRef<void(FGameplayAbilityRepAnimMontage&)> MutateRepAnimMontageFunction)
 {
 	float Duration = -1.f;
 
@@ -3059,6 +3064,9 @@ float UAbilitySystemComponent::PlayMontage(UGameplayAbility* InAnimatingAbility,
 					MutableRepAnimMontageInfo.BlendOutTime = NewAnimMontage->GetDefaultBlendInTime();
 				}
 
+				// Give the caller the opportunity to modify the rep info for this montage, used by "slot animation as dynamic montage" to replicate play count parameter
+				MutateRepAnimMontageFunction(MutableRepAnimMontageInfo);
+
 				// Update parameters that change during Montage life time.
 				AnimMontage_UpdateReplicatedData();
 			}
@@ -3282,12 +3290,26 @@ void UAbilitySystemComponent::OnRep_ReplicatedAnimMontage()
 				}
 				else
 				{
-					PlaySlotAnimationAsDynamicMontageSimulated(
-						ConstRepAnimMontageInfo.Animation,
-						ConstRepAnimMontageInfo.SlotName,
-						ConstRepAnimMontageInfo.BlendTime,
-						ConstRepAnimMontageInfo.BlendOutTime,
-						ConstRepAnimMontageInfo.PlayRate);
+					if (ConstRepAnimMontageInfo.PlayCount == 1.0f)
+					{
+						PlaySlotAnimationAsDynamicMontageSimulated(
+							ConstRepAnimMontageInfo.Animation,
+							ConstRepAnimMontageInfo.SlotName,
+							ConstRepAnimMontageInfo.BlendTime,
+							ConstRepAnimMontageInfo.BlendOutTime,
+							ConstRepAnimMontageInfo.PlayRate);
+					}
+					else
+					{
+						PlaySlotAnimationAsDynamicMontage_WithFractionalLoopsSimulated(
+							ConstRepAnimMontageInfo.Animation,
+							ConstRepAnimMontageInfo.SlotName,
+							ConstRepAnimMontageInfo.BlendTime,
+							ConstRepAnimMontageInfo.BlendOutTime,
+							ConstRepAnimMontageInfo.PlayRate,
+							ConstRepAnimMontageInfo.PlayCount
+						);
+					}
 				}
 			}
 
@@ -3654,6 +3676,23 @@ UAnimMontage* UAbilitySystemComponent::PlaySlotAnimationAsDynamicMontageSimulate
 	return DynamicMontage;
 }
 
+UAnimMontage* UAbilitySystemComponent::PlaySlotAnimationAsDynamicMontage_WithFractionalLoopsSimulated(UAnimSequenceBase* AnimAsset, FName SlotName, float BlendInTime, float BlendOutTime, float InPlayRate, float PlayCount)
+{
+	UAnimMontage* DynamicMontage = UAnimMontage::CreateSlotAnimationAsDynamicMontage_WithFractionalLoops(AnimAsset, SlotName, BlendInTime, BlendOutTime, PlayCount);
+	PlayMontageSimulated(DynamicMontage, InPlayRate, NAME_None);
+	return DynamicMontage;
+}
+
+UAnimMontage* UAbilitySystemComponent::PlaySlotAnimationAsDynamicMontage_WithFractionalLoops(UGameplayAbility* AnimatingAbility, FGameplayAbilityActivationInfo ActivationInfo, UAnimSequenceBase* AnimAsset, FName SlotName, float BlendInTime, float BlendOutTime, float InPlayRate, float StartTimeSeconds, float PlayCount)
+{
+	UAnimMontage* DynamicMontage = UAnimMontage::CreateSlotAnimationAsDynamicMontage_WithFractionalLoops(AnimAsset, SlotName, BlendInTime, BlendOutTime, PlayCount);
+	PlayMontageInternal(AnimatingAbility, ActivationInfo, DynamicMontage, InPlayRate, NAME_None, StartTimeSeconds, [PlayCount](FGameplayAbilityRepAnimMontage& MutableRepAnimMontageInfo)
+		{
+			MutableRepAnimMontageInfo.PlayCount = PlayCount;
+		});
+	return DynamicMontage;
+}
+
 UAnimMontage* UAbilitySystemComponent::GetCurrentMontage() const
 {
 	UAnimInstance* AnimInstance = AbilityActorInfo.IsValid() ? AbilityActorInfo->GetAnimInstance() : nullptr;
diff --git a/Engine/Plugins/Runtime/GameplayAbilities/Source/GameplayAbilities/Private/GameplayAbilityRepAnimMontage.cpp b/Engine/Plugins/Runtime/GameplayAbilities/Source/GameplayAbilities/Private/GameplayAbilityRepAnimMontage.cpp
--- a/Engine/Plugins/Runtime/GameplayAbilities/Source/GameplayAbilities/Private/GameplayAbilityRepAnimMontage.cpp
+++ b/Engine/Plugins/Runtime/GameplayAbilities/Source/GameplayAbilities/Private/GameplayAbilityRepAnimMontage.cpp
@@ -91,6 +91,11 @@ bool FGameplayAbilityRepAnimMontage::NetSerialize(FArchive& Ar, class UPackageMa
 		Ar << SlotName;
 	}
 
+	if (Ar.EngineNetVer() >= FEngineNetworkCustomVersion::MontagePlayCountSerialization)
+	{
+		Ar << PlayCount;
+	}
+
 	bOutSuccess = true;
 	return true;
 }
diff --git a/Engine/Plugins/Runtime/GameplayAbilities/Source/GameplayAbilities/Private/Serialization/GameplayAbilityRepAnimMontageNetSerializer.cpp b/Engine/Plugins/Runtime/GameplayAbilities/Source/GameplayAbilities/Private/Serialization/GameplayAbilityRepAnimMontageNetSerializer.cpp
--- a/Engine/Plugins/Runtime/GameplayAbilities/Source/GameplayAbilities/Private/Serialization/GameplayAbilityRepAnimMontageNetSerializer.cpp
+++ b/Engine/Plugins/Runtime/GameplayAbilities/Source/GameplayAbilities/Private/Serialization/GameplayAbilityRepAnimMontageNetSerializer.cpp
@@ -44,7 +44,7 @@ struct FGameplayAbilityRepAnimMontageNetSerializer
 
 	struct FQuantizedType
 	{
-		alignas(16) uint8 GameplayAbilityRepAnimMontage[80];
+		alignas(16) uint8 GameplayAbilityRepAnimMontage[84];
 
 		float Position;
 		float BlendOutTime;
diff --git a/Engine/Plugins/Runtime/GameplayAbilities/Source/GameplayAbilities/Public/Abilities/GameplayAbilityRepAnimMontage.h b/Engine/Plugins/Runtime/GameplayAbilities/Source/GameplayAbilities/Public/Abilities/GameplayAbilityRepAnimMontage.h
--- a/Engine/Plugins/Runtime/GameplayAbilities/Source/GameplayAbilities/Public/Abilities/GameplayAbilityRepAnimMontage.h
+++ b/Engine/Plugins/Runtime/GameplayAbilities/Source/GameplayAbilities/Public/Abilities/GameplayAbilityRepAnimMontage.h
@@ -46,6 +46,10 @@ struct GAMEPLAYABILITIES_API FGameplayAbilityRepAnimMontage
 	UPROPERTY(NotReplicated)
 	float BlendOutTime;
 
+	/** Fractional loop count for non-montages */
+	UPROPERTY()
+	float PlayCount;
+
 	/** NextSectionID */
 	UPROPERTY()
 	uint8 NextSectionID;
@@ -84,6 +88,7 @@ struct GAMEPLAYABILITIES_API FGameplayAbilityRepAnimMontage
 	Position(0.f),
 	BlendTime(0.f),
 	BlendOutTime(0.f),
+	PlayCount(0.f),
 	NextSectionID(0),
 	PlayInstanceId(0),
 	bRepPosition(true),
diff --git a/Engine/Plugins/Runtime/GameplayAbilities/Source/GameplayAbilities/Public/Abilities/Tasks/AbilityTask_PlayAnimAndWait.h b/Engine/Plugins/Runtime/GameplayAbilities/Source/GameplayAbilities/Public/Abilities/Tasks/AbilityTask_PlayAnimAndWait.h
--- a/Engine/Plugins/Runtime/GameplayAbilities/Source/GameplayAbilities/Public/Abilities/Tasks/AbilityTask_PlayAnimAndWait.h
+++ b/Engine/Plugins/Runtime/GameplayAbilities/Source/GameplayAbilities/Public/Abilities/Tasks/AbilityTask_PlayAnimAndWait.h
@@ -61,7 +61,7 @@ public:
 	UFUNCTION(BlueprintCallable, Category="Ability|Tasks", meta = (DisplayName="Play Anim And Wait",
 		HidePin = "OwningAbility", DefaultToSelf = "OwningAbility", BlueprintInternalUseOnly = "TRUE"))
 	static UAbilityTask_PlayAnimAndWait* CreatePlayAnimAndWaitProxy(UGameplayAbility* OwningAbility,
-		FName TaskInstanceName, UAnimSequence* AnimSequence, FName SlotName = NAME_None, float BlendInTime = 0.25f, float BlendOutTime = 0.25f, float InPlayRate = 1.0f, float StartTimeSeconds = 0.0f, bool bStopWhenAbilityEnds = true, float AnimRootMotionTranslationScale = 1.0f);
+		FName TaskInstanceName, UAnimSequence* AnimSequence, FName SlotName = NAME_None, float BlendInTime = 0.25f, float BlendOutTime = 0.25f, float InPlayRate = 1.0f, float InPlayCount = 1.0f, float StartTimeSeconds = 0.0f, bool bStopWhenAbilityEnds = true, float AnimRootMotionTranslationScale = 1.0f);
 
 	virtual void Activate() override;
 
@@ -89,6 +89,7 @@ protected:
 	float BlendInTime = 0.25f;
 	float BlendOutTime = 0.25f;
 	float PlayRate = 1.0f;
+	float PlayCount = 1.0f;
 	float AnimRootMotionTranslationScale = 1.0f;
 	float StartTimeSeconds = 0.0f;
 	bool bStopWhenAbilityEnds = true;
diff --git a/Engine/Plugins/Runtime/GameplayAbilities/Source/GameplayAbilities/Public/AbilitySystemComponent.h b/Engine/Plugins/Runtime/GameplayAbilities/Source/GameplayAbilities/Public/AbilitySystemComponent.h
--- a/Engine/Plugins/Runtime/GameplayAbilities/Source/GameplayAbilities/Public/AbilitySystemComponent.h
+++ b/Engine/Plugins/Runtime/GameplayAbilities/Source/GameplayAbilities/Public/AbilitySystemComponent.h
@@ -1450,14 +1450,18 @@ class GAMEPLAYABILITIES_API UAbilitySystemComponent : public UGameplayTasksCompo
 
 	/** Plays a montage and handles replication and prediction based on passed in ability/activation info */
 	virtual float PlayMontage(UGameplayAbility* AnimatingAbility, FGameplayAbilityActivationInfo ActivationInfo, UAnimMontage* Montage, float InPlayRate, FName StartSectionName = NAME_None, float StartTimeSeconds = 0.0f);
-	
+
 	virtual UAnimMontage* PlaySlotAnimationAsDynamicMontage(UGameplayAbility* AnimatingAbility, FGameplayAbilityActivationInfo ActivationInfo, UAnimSequenceBase* AnimAsset, FName SlotName, float BlendInTime, float BlendOutTime, float InPlayRate = 1.f, float StartTimeSeconds = 0.0f);
 
+	virtual UAnimMontage* PlaySlotAnimationAsDynamicMontage_WithFractionalLoops(UGameplayAbility* AnimatingAbility, FGameplayAbilityActivationInfo ActivationInfo, UAnimSequenceBase* AnimAsset, FName SlotName, float BlendInTime, float BlendOutTime, float InPlayRate = 1.f, float StartTimeSeconds = 0.0f, float PlayCount = 1.0f);
+
 	/** Plays a montage without updating replication/prediction structures. Used by simulated proxies when replication tells them to play a montage. */
 	virtual float PlayMontageSimulated(UAnimMontage* Montage, float InPlayRate, FName StartSectionName = NAME_None);
 	
 	virtual UAnimMontage* PlaySlotAnimationAsDynamicMontageSimulated(UAnimSequenceBase* AnimAsset, FName SlotName, float BlendInTime, float BlendOutTime, float InPlayRate = 1.f);
 
+	virtual UAnimMontage* PlaySlotAnimationAsDynamicMontage_WithFractionalLoopsSimulated(UAnimSequenceBase* AnimAsset, FName SlotName, float BlendInTime, float BlendOutTime, float InPlayRate = 1.f, float PlayCount = 1.0f);
+
 	/** Stops whatever montage is currently playing. Expectation is caller should only be stopping it if they are the current animating ability (or have good reason not to check) */
 	virtual void CurrentMontageStop(float OverrideBlendOutTime = -1.0f);
 
@@ -1996,5 +2000,8 @@ protected:
 	};
 
 	TArray<FAbilityListLockActiveChange*> AbilityListLockActiveChanges;
+
+private:
+	float PlayMontageInternal(UGameplayAbility* AnimatingAbility, FGameplayAbilityActivationInfo ActivationInfo, UAnimMontage* Montage, float InPlayRate, FName StartSectionName, float StartTimeSeconds, TFunctionRef<void(FGameplayAbilityRepAnimMontage&)> MutateRepAnimMontageFunction);
 	
 };
diff --git a/Engine/Source/Runtime/Core/Public/Misc/EngineNetworkCustomVersion.h b/Engine/Source/Runtime/Core/Public/Misc/EngineNetworkCustomVersion.h
--- a/Engine/Source/Runtime/Core/Public/Misc/EngineNetworkCustomVersion.h
+++ b/Engine/Source/Runtime/Core/Public/Misc/EngineNetworkCustomVersion.h
@@ -89,7 +89,7 @@ struct FEngineNetworkCustomVersion
 		PredictionKeyBaseNotReplicated = 34,	// Bump version to stop FPredictionKey::Base from being replicated (it was unused).
 		RepMoveOptionalAcceleration = 35,		// Bump version to support serialization changes to RepMove for optional Acceleration
 		CustomExports = 36,						// Bump version to support CustomExports (such as NetTokens used for supporting "Dynamic Replication" of GameplayTags -- see FGameplayTag::NetSerialize_Packed)		// -----<new versions can be added above this line>-------------------------------------------------
-
+		MontagePlayCountSerialization = 37,		// Bump version to support serialization changes for dynamic montage play count 
 
 		// -----<new versions can be added above this line>-------------------------------------------------
 		VersionPlusOne,
diff --git a/Engine/Source/Runtime/Engine/Classes/Animation/AnimMontage.h b/Engine/Source/Runtime/Engine/Classes/Animation/AnimMontage.h
--- a/Engine/Source/Runtime/Engine/Classes/Animation/AnimMontage.h
+++ b/Engine/Source/Runtime/Engine/Classes/Animation/AnimMontage.h
@@ -958,6 +958,9 @@ public:
 	UFUNCTION(BlueprintCallable, Category = "Animation")
 	static ENGINE_API UAnimMontage* CreateSlotAnimationAsDynamicMontage_WithBlendSettings(UAnimSequenceBase* Asset, FName SlotNodeName, const FMontageBlendSettings& BlendInSettings, const FMontageBlendSettings& BlendOutSettings, float InPlayRate = 1.f, int32 LoopCount = 1, float InBlendOutTriggerTime = -1.f);
 
+	/** Utility function to create dynamic montage from AnimSequence */
+	ENGINE_API static UAnimMontage* CreateSlotAnimationAsDynamicMontage_WithFractionalLoops(UAnimSequenceBase* Asset, FName SlotNodeName, float BlendInTime = 0.25f, float BlendOutTime = 0.25f, float LoopCount = 1.0f, float BlendOutTriggerTime = -1.f);
+
 	//~Begin Time Stretch Curve
 public:
 
diff --git a/Engine/Source/Runtime/Engine/Private/Animation/AnimMontage.cpp b/Engine/Source/Runtime/Engine/Private/Animation/AnimMontage.cpp
--- a/Engine/Source/Runtime/Engine/Private/Animation/AnimMontage.cpp
+++ b/Engine/Source/Runtime/Engine/Private/Animation/AnimMontage.cpp
@@ -2996,6 +2996,58 @@ UAnimMontage* UAnimMontage::CreateSlotAnimationAsDynamicMontage(UAnimSequenceBas
 	return CreateSlotAnimationAsDynamicMontage_WithBlendSettings(Asset, SlotNodeName, BlendInSettings, BlendOutSettings, InPlayRate, LoopCount, BlendOutTriggerTime);
 }
 
+UAnimMontage* UAnimMontage::CreateSlotAnimationAsDynamicMontage_WithFractionalLoops(UAnimSequenceBase* Asset, FName SlotNodeName, float BlendInTime, float BlendOutTime, float LoopCount, float BlendOutTriggerTime)
+{
+	if (!ensure(LoopCount > 0.0))
+	{
+		return nullptr;
+	}
+
+	UAnimMontage* BaseMontage = CreateSlotAnimationAsDynamicMontage(Asset, SlotNodeName, BlendInTime, BlendOutTime, 1.0f, 1, BlendOutTriggerTime);
+	if (BaseMontage == nullptr)
+	{
+		return nullptr;
+	}
+
+	if (LoopCount == 1.0)
+	{
+		return BaseMontage;
+	}
+
+	FSlotAnimationTrack& SlotTrack = BaseMontage->SlotAnimTracks[0];
+	FAnimSegment& OriginalSegment = SlotTrack.AnimTrack.AnimSegments[0];
+
+	FCompositeSection& Section = BaseMontage->CompositeSections[0];
+	Section.Link(Asset, 0.0);
+
+	if (LoopCount < 1.0)
+	{
+		// Need to shorten the existing segment
+		OriginalSegment.AnimEndTime = Asset->GetPlayLength() * LoopCount;
+	}
+	else
+	{
+		int32 NumFullLoops = FMath::IsFinite(LoopCount) ? static_cast<int32>(LoopCount) : MAX_int32;
+		float FractionalLoop = LoopCount - NumFullLoops;
+
+		// Need to lengthen the existing segment
+		OriginalSegment.LoopingCount = NumFullLoops;
+
+		// Need to add a new fractional segment onto the end of the montage
+		FAnimSegment& NewSegment = SlotTrack.AnimTrack.AnimSegments.AddDefaulted_GetRef();
+		NewSegment.SetAnimReference(Asset, true);
+
+		NewSegment.StartPos = Asset->GetPlayLength() * NumFullLoops;
+		NewSegment.AnimEndTime = Asset->GetPlayLength() * FractionalLoop;
+	}
+
+	PRAGMA_DISABLE_DEPRECATION_WARNINGS
+	BaseMontage->SequenceLength = Asset->GetPlayLength() * LoopCount;
+	PRAGMA_ENABLE_DEPRECATION_WARNINGS
+
+	return BaseMontage;
+}
+
 UAnimMontage* UAnimMontage::CreateSlotAnimationAsDynamicMontage_WithBlendSettings(UAnimSequenceBase* Asset, FName SlotNodeName, const FMontageBlendSettings& BlendInSettings, const FMontageBlendSettings& BlendOutSettings, float InPlayRate, int32 LoopCount, float InBlendOutTriggerTime)
 {
 	// create temporary montage and play
