Fix back compat in UAbilityTask_PlayAnimAndWait with new PlayCount optional param

diff --git a/Engine/Plugins/Runtime/GameplayAbilities/Source/GameplayAbilities/Private/Abilities/Tasks/AbilityTask_PlayAnimAndWait.cpp b/Engine/Plugins/Runtime/GameplayAbilities/Source/GameplayAbilities/Private/Abilities/Tasks/AbilityTask_PlayAnimAndWait.cpp
--- a/Engine/Plugins/Runtime/GameplayAbilities/Source/GameplayAbilities/Private/Abilities/Tasks/AbilityTask_PlayAnimAndWait.cpp
+++ b/Engine/Plugins/Runtime/GameplayAbilities/Source/GameplayAbilities/Private/Abilities/Tasks/AbilityTask_PlayAnimAndWait.cpp
@@ -87,10 +87,10 @@ UAbilityTask_PlayAnimAndWait* UAbilityTask_PlayAnimAndWait::CreatePlayAnimAndWai
 	float BlendInTime /*= 0.25f*/,
 	float BlendOutTime /*= 0.25f*/,
 	float InPlayRate /*= 1.0f*/,
-	float InPlayCount /*= 1.0f*/,
 	float StartTimeSeconds /*= 0.0f*/, 
 	bool bStopWhenAbilityEnds /*= true*/, 
-	float AnimRootMotionTranslationScale /*= 1.f*/)
+	float AnimRootMotionTranslationScale /*= 1.f*/,
+	float InPlayCount /*= 1.0f*/)
 {
 	UAbilityTask_PlayAnimAndWait* PlayAnimTask = NewAbilityTask<UAbilityTask_PlayAnimAndWait>(OwningAbility, TaskInstanceName);
 	PlayAnimTask->AnimSequenceToPlay = AnimSequence;
diff --git a/Engine/Plugins/Runtime/GameplayAbilities/Source/GameplayAbilities/Public/Abilities/Tasks/AbilityTask_PlayAnimAndWait.h b/Engine/Plugins/Runtime/GameplayAbilities/Source/GameplayAbilities/Public/Abilities/Tasks/AbilityTask_PlayAnimAndWait.h
--- a/Engine/Plugins/Runtime/GameplayAbilities/Source/GameplayAbilities/Public/Abilities/Tasks/AbilityTask_PlayAnimAndWait.h
+++ b/Engine/Plugins/Runtime/GameplayAbilities/Source/GameplayAbilities/Public/Abilities/Tasks/AbilityTask_PlayAnimAndWait.h
@@ -61,7 +61,7 @@ public:
 	UFUNCTION(BlueprintCallable, Category="Ability|Tasks", meta = (DisplayName="Play Anim And Wait",
 		HidePin = "OwningAbility", DefaultToSelf = "OwningAbility", BlueprintInternalUseOnly = "TRUE"))
 	static UAbilityTask_PlayAnimAndWait* CreatePlayAnimAndWaitProxy(UGameplayAbility* OwningAbility,
-		FName TaskInstanceName, UAnimSequence* AnimSequence, FName SlotName = NAME_None, float BlendInTime = 0.25f, float BlendOutTime = 0.25f, float InPlayRate = 1.0f, float InPlayCount = 1.0f, float StartTimeSeconds = 0.0f, bool bStopWhenAbilityEnds = true, float AnimRootMotionTranslationScale = 1.0f);
+		FName TaskInstanceName, UAnimSequence* AnimSequence, FName SlotName = NAME_None, float BlendInTime = 0.25f, float BlendOutTime = 0.25f, float InPlayRate = 1.0f, float StartTimeSeconds = 0.0f, bool bStopWhenAbilityEnds = true, float AnimRootMotionTranslationScale = 1.0f, float InPlayCount = 1.0f);
 
 	virtual void Activate() override;
 
