PR #12621: Correct spelling mistake in GameplayAbilitySpecHandle.h comment

diff --git a/Engine/Plugins/Runtime/GameplayAbilities/Source/GameplayAbilities/Public/GameplayAbilitySpecHandle.h b/Engine/Plugins/Runtime/GameplayAbilities/Source/GameplayAbilities/Public/GameplayAbilitySpecHandle.h
--- a/Engine/Plugins/Runtime/GameplayAbilities/Source/GameplayAbilities/Public/GameplayAbilitySpecHandle.h
+++ b/Engine/Plugins/Runtime/GameplayAbilities/Source/GameplayAbilities/Public/GameplayAbilitySpecHandle.h
@@ -6,7 +6,7 @@
 
 
 /**
- *	This file exists in addition so that GameplayEffect.h can use FGameplayAbilitySpec without having to include GameplayAbilityTypes.h which has depancies on
+ *	This file exists in addition so that GameplayEffect.h can use FGameplayAbilitySpec without having to include GameplayAbilityTypes.h which has dependencies on
  *	GameplayEffect.h
  */
 
