Fix GAs missing ReplicatedInstances in replays.

Normally ability instances are added and replicated in UAbilitySystemComponent::CreateNewInstanceOfAbility, but that is only called on the server (for replicated abilities). The result is that an ability's ReplicatedInstances array will be of the correct size in replays, but the elements will be nullptr as they did not replicate to replays.

CVar to revert behavior: AbilitySystem.RecordAbilityInstancesOnClientSpecUpdates

#rnx

diff --git a/Engine/Plugins/Runtime/GameplayAbilities/Source/GameplayAbilities/Private/GameplayAbilityTypes.cpp b/Engine/Plugins/Runtime/GameplayAbilities/Source/GameplayAbilities/Private/GameplayAbilityTypes.cpp
--- a/Engine/Plugins/Runtime/GameplayAbilities/Source/GameplayAbilities/Private/GameplayAbilityTypes.cpp
+++ b/Engine/Plugins/Runtime/GameplayAbilities/Source/GameplayAbilities/Private/GameplayAbilityTypes.cpp
@@ -18,6 +18,8 @@
 
 //----------------------------------------------------------------------
 
+static TAutoConsoleVariable<bool> CVarRecordAbilityInstancesOnClientSpecUpdates(TEXT("AbilitySystem.RecordAbilityInstancesOnClientSpecUpdates"), true, TEXT("Set this to false if encountering issues with the logic for replicating ability instances in the spec's PreReplicatedRemove, PostReplicatedChange, or PostReplicatedAdd."));
+
 void FGameplayAbilityActorInfo::InitFromActor(AActor *InOwnerActor, AActor *InAvatarActor, UAbilitySystemComponent* InAbilitySystemComponent)
 {
 	check(InOwnerActor);
@@ -236,6 +238,18 @@ void FGameplayAbilitySpec::PreReplicatedRemove(const struct FGameplayAbilitySpec
 		UE_VLOG(InArraySerializer.Owner->GetOwner(), VLogAbilitySystem, Verbose, TEXT("OnRemoveAbility (Non-Auth): [%s] %s. Level: %d"), *Handle.ToString(), *GetNameSafe(Ability), Level);
 
 		FScopedAbilityListLock AblityListLock(*InArraySerializer.Owner);
+
+		if (CVarRecordAbilityInstancesOnClientSpecUpdates.GetValueOnGameThread())
+		{
+			for (UGameplayAbility* AbilityInstance : ReplicatedInstances)
+			{
+				if (AbilityInstance)
+				{
+					InArraySerializer.Owner->RemoveReplicatedInstancedAbility(AbilityInstance);
+				}
+			}
+		}
+		
 		InArraySerializer.Owner->OnRemoveAbility(*this);
 	}
 }
@@ -246,6 +260,17 @@ void FGameplayAbilitySpec::PostReplicatedChange(const struct FGameplayAbilitySpe
 	{
 		UE_LOG(LogAbilitySystem, Verbose, TEXT("%s: AbilitySpecChanged (Non-Auth): [%s] %s. Level: %d"), *GetNameSafe(InArraySerializer.Owner->GetOwner()), *Handle.ToString(), *GetNameSafe(Ability), Level)
 		UE_VLOG(InArraySerializer.Owner->GetOwner(), VLogAbilitySystem, Verbose, TEXT("AbilitySpecChanged (Non-Auth): [%s] %s. Level: %d"), *Handle.ToString(), *GetNameSafe(Ability), Level);
+
+		if (CVarRecordAbilityInstancesOnClientSpecUpdates.GetValueOnGameThread())
+		{
+			for (UGameplayAbility* AbilityInstance : ReplicatedInstances)
+			{
+				if (AbilityInstance)
+				{
+					InArraySerializer.Owner->AddReplicatedInstancedAbility(AbilityInstance);
+				}
+			}
+		}
 	}
 }
 
@@ -256,6 +281,17 @@ void FGameplayAbilitySpec::PostReplicatedAdd(const struct FGameplayAbilitySpecCo
 		UE_LOG(LogAbilitySystem, Verbose, TEXT("%s: OnGiveAbility (Non-Auth): [%s] %s. Level: %d"), *GetNameSafe(InArraySerializer.Owner->GetOwner()), *Handle.ToString(), *GetNameSafe(Ability), Level)
 		UE_VLOG(InArraySerializer.Owner->GetOwner(), VLogAbilitySystem, Verbose, TEXT("OnGiveAbility (Non-Auth): [%s] %s. Level: %d"), *Handle.ToString(), *GetNameSafe(Ability), Level);
 
+		if (CVarRecordAbilityInstancesOnClientSpecUpdates.GetValueOnGameThread())
+		{
+			for (UGameplayAbility* AbilityInstance : ReplicatedInstances)
+			{
+				if (AbilityInstance)
+				{
+					InArraySerializer.Owner->AddReplicatedInstancedAbility(AbilityInstance);
+				}
+			}
+		}
+		
 		InArraySerializer.Owner->OnGiveAbility(*this);
 	}
 }
