ATTRIBUTE_ACCESSORS macro which was commented as a suggestion in AttributeSet.h is now usable out of the box as ATTRIBUTE_ACCESSORS_BASIC.

diff --git a/Engine/Plugins/Runtime/GameplayAbilities/Source/GameplayAbilities/Public/AttributeSet.h b/Engine/Plugins/Runtime/GameplayAbilities/Source/GameplayAbilities/Public/AttributeSet.h
--- a/Engine/Plugins/Runtime/GameplayAbilities/Source/GameplayAbilities/Public/AttributeSet.h
+++ b/Engine/Plugins/Runtime/GameplayAbilities/Source/GameplayAbilities/Public/AttributeSet.h
@@ -453,3 +453,16 @@ private:
 		PropertyName.SetBaseValue(NewVal); \
 		PropertyName.SetCurrentValue(NewVal); \
 	}
+
+/**
+ * Instead of copying the ATTRIBUTE_ACCESSORS macro from the comment above, you can also immediately 
+ * use this one if it fulfills your needs. Named differently to not cause name collision in existing 
+ * projects that have copied ATTRIBUTE_ACCESSORS already:
+ * 
+ *  ATTRIBUTE_ACCESSORS_BASIC(UMyHealthSet, Health)
+ */
+#define ATTRIBUTE_ACCESSORS_BASIC(ClassName, PropertyName) \
+	GAMEPLAYATTRIBUTE_PROPERTY_GETTER(ClassName, PropertyName) \
+	GAMEPLAYATTRIBUTE_VALUE_GETTER(PropertyName) \
+	GAMEPLAYATTRIBUTE_VALUE_SETTER(PropertyName) \
+	GAMEPLAYATTRIBUTE_VALUE_INITTER(PropertyName)
\ No newline at end of file
