Fixing Garbage reference in AbilityTask_ApplyRootMotion_Base

diff --git a/Engine/Plugins/Runtime/GameplayAbilities/Source/GameplayAbilities/Private/Abilities/Tasks/AbilityTask_ApplyRootMotionConstantForce.cpp b/Engine/Plugins/Runtime/GameplayAbilities/Source/GameplayAbilities/Private/Abilities/Tasks/AbilityTask_ApplyRootMotionConstantForce.cpp
--- a/Engine/Plugins/Runtime/GameplayAbilities/Source/GameplayAbilities/Private/Abilities/Tasks/AbilityTask_ApplyRootMotionConstantForce.cpp
+++ b/Engine/Plugins/Runtime/GameplayAbilities/Source/GameplayAbilities/Private/Abilities/Tasks/AbilityTask_ApplyRootMotionConstantForce.cpp
@@ -60,7 +60,7 @@ void UAbilityTask_ApplyRootMotionConstantForce::SharedInitAndApply()
 		StartTime = GetWorld()->GetTimeSeconds();
 		EndTime = StartTime + Duration;
 
-		if (MovementComponent)
+		if (MovementComponent.IsValid())
 		{
 			ForceName = ForceName.IsNone() ? FName("AbilityTaskApplyRootMotionConstantForce"): ForceName;
 			TSharedPtr<FRootMotionSource_ConstantForce> ConstantForce = MakeShared<FRootMotionSource_ConstantForce>();
@@ -145,7 +145,7 @@ void UAbilityTask_ApplyRootMotionConstantForce::PreDestroyFromReplication()
 
 void UAbilityTask_ApplyRootMotionConstantForce::OnDestroy(bool AbilityIsEnding)
 {
-	if (MovementComponent)
+	if (MovementComponent.IsValid())
 	{
 		MovementComponent->RemoveRootMotionSourceByID(RootMotionSourceID);
 	}
diff --git a/Engine/Plugins/Runtime/GameplayAbilities/Source/GameplayAbilities/Private/Abilities/Tasks/AbilityTask_ApplyRootMotionJumpForce.cpp b/Engine/Plugins/Runtime/GameplayAbilities/Source/GameplayAbilities/Private/Abilities/Tasks/AbilityTask_ApplyRootMotionJumpForce.cpp
--- a/Engine/Plugins/Runtime/GameplayAbilities/Source/GameplayAbilities/Private/Abilities/Tasks/AbilityTask_ApplyRootMotionJumpForce.cpp
+++ b/Engine/Plugins/Runtime/GameplayAbilities/Source/GameplayAbilities/Private/Abilities/Tasks/AbilityTask_ApplyRootMotionJumpForce.cpp
@@ -96,7 +96,7 @@ void UAbilityTask_ApplyRootMotionJumpForce::SharedInitAndApply()
 		StartTime = GetWorld()->GetTimeSeconds();
 		EndTime = StartTime + Duration;
 
-		if (MovementComponent)
+		if (MovementComponent.IsValid())
 		{
 			ForceName = ForceName.IsNone() ? FName("AbilityTaskApplyRootMotionJumpForce") : ForceName;
 			TSharedPtr<FRootMotionSource_JumpForce> JumpForce = MakeShared<FRootMotionSource_JumpForce>();
@@ -206,7 +206,7 @@ void UAbilityTask_ApplyRootMotionJumpForce::OnDestroy(bool AbilityIsEnding)
 		Character->LandedDelegate.RemoveDynamic(this, &UAbilityTask_ApplyRootMotionJumpForce::OnLandedCallback);
 	}
 
-	if (MovementComponent)
+	if (MovementComponent.IsValid())
 	{
 		MovementComponent->RemoveRootMotionSourceByID(RootMotionSourceID);
 	}
diff --git a/Engine/Plugins/Runtime/GameplayAbilities/Source/GameplayAbilities/Private/Abilities/Tasks/AbilityTask_ApplyRootMotionMoveToActorForce.cpp b/Engine/Plugins/Runtime/GameplayAbilities/Source/GameplayAbilities/Private/Abilities/Tasks/AbilityTask_ApplyRootMotionMoveToActorForce.cpp
--- a/Engine/Plugins/Runtime/GameplayAbilities/Source/GameplayAbilities/Private/Abilities/Tasks/AbilityTask_ApplyRootMotionMoveToActorForce.cpp
+++ b/Engine/Plugins/Runtime/GameplayAbilities/Source/GameplayAbilities/Private/Abilities/Tasks/AbilityTask_ApplyRootMotionMoveToActorForce.cpp
@@ -159,7 +159,7 @@ void UAbilityTask_ApplyRootMotionMoveToActorForce::SharedInitAndApply()
 		StartTime = GetWorld()->GetTimeSeconds();
 		EndTime = StartTime + Duration;
 
-		if (MovementComponent)
+		if (MovementComponent.IsValid())
 		{
 			if (bSetNewMovementMode)
 			{
@@ -212,7 +212,7 @@ FVector UAbilityTask_ApplyRootMotionMoveToActorForce::CalculateTargetOffset() co
 	
 	if (OffsetAlignment == ERootMotionMoveToActorTargetOffsetType::AlignFromTargetToSource)
 	{
-		if (MovementComponent)
+		if (MovementComponent.IsValid())
 		{
 			FVector ToSource = MovementComponent->GetActorLocation() - TargetActorLocation;
 			ToSource.Z = 0.f;
@@ -285,7 +285,7 @@ bool UAbilityTask_ApplyRootMotionMoveToActorForce::UpdateTargetLocation(float De
 
 void UAbilityTask_ApplyRootMotionMoveToActorForce::SetRootMotionTargetLocation(FVector NewTargetLocation)
 {
-	if (MovementComponent)
+	if (MovementComponent.IsValid())
 	{
 		TSharedPtr<FRootMotionSource> RMS = MovementComponent->GetRootMotionSourceByID(RootMotionSourceID);
 		if (RMS.IsValid())
@@ -396,7 +396,7 @@ void UAbilityTask_ApplyRootMotionMoveToActorForce::OnDestroy(bool AbilityIsEndin
 		UAbilityTask_ApplyRootMotion_Base::OnTargetActorSwapped.Remove( TargetActorSwapHandle );
 	}
 
-	if (MovementComponent)
+	if (MovementComponent.IsValid())
 	{
 		MovementComponent->RemoveRootMotionSourceByID(RootMotionSourceID);
 
diff --git a/Engine/Plugins/Runtime/GameplayAbilities/Source/GameplayAbilities/Private/Abilities/Tasks/AbilityTask_ApplyRootMotionMoveToForce.cpp b/Engine/Plugins/Runtime/GameplayAbilities/Source/GameplayAbilities/Private/Abilities/Tasks/AbilityTask_ApplyRootMotionMoveToForce.cpp
--- a/Engine/Plugins/Runtime/GameplayAbilities/Source/GameplayAbilities/Private/Abilities/Tasks/AbilityTask_ApplyRootMotionMoveToForce.cpp
+++ b/Engine/Plugins/Runtime/GameplayAbilities/Source/GameplayAbilities/Private/Abilities/Tasks/AbilityTask_ApplyRootMotionMoveToForce.cpp
@@ -51,7 +51,7 @@ void UAbilityTask_ApplyRootMotionMoveToForce::SharedInitAndApply()
 		StartTime = GetWorld()->GetTimeSeconds();
 		EndTime = StartTime + Duration;
 
-		if (MovementComponent)
+		if (MovementComponent.IsValid())
 		{
 			if (bSetNewMovementMode)
 			{
@@ -154,7 +154,7 @@ void UAbilityTask_ApplyRootMotionMoveToForce::PreDestroyFromReplication()
 
 void UAbilityTask_ApplyRootMotionMoveToForce::OnDestroy(bool AbilityIsEnding)
 {
-	if (MovementComponent)
+	if (MovementComponent.IsValid())
 	{
 		MovementComponent->RemoveRootMotionSourceByID(RootMotionSourceID);
 
diff --git a/Engine/Plugins/Runtime/GameplayAbilities/Source/GameplayAbilities/Private/Abilities/Tasks/AbilityTask_ApplyRootMotionRadialForce.cpp b/Engine/Plugins/Runtime/GameplayAbilities/Source/GameplayAbilities/Private/Abilities/Tasks/AbilityTask_ApplyRootMotionRadialForce.cpp
--- a/Engine/Plugins/Runtime/GameplayAbilities/Source/GameplayAbilities/Private/Abilities/Tasks/AbilityTask_ApplyRootMotionRadialForce.cpp
+++ b/Engine/Plugins/Runtime/GameplayAbilities/Source/GameplayAbilities/Private/Abilities/Tasks/AbilityTask_ApplyRootMotionRadialForce.cpp
@@ -55,7 +55,7 @@ void UAbilityTask_ApplyRootMotionRadialForce::SharedInitAndApply()
 		StartTime = GetWorld()->GetTimeSeconds();
 		EndTime = StartTime + Duration;
 
-		if (MovementComponent)
+		if (MovementComponent.IsValid())
 		{
 			ForceName = ForceName.IsNone() ? FName("AbilityTaskApplyRootMotionRadialForce") : ForceName;
 			TSharedPtr<FRootMotionSource_RadialForce> RadialForce = MakeShared<FRootMotionSource_RadialForce>();
@@ -150,7 +150,7 @@ void UAbilityTask_ApplyRootMotionRadialForce::PreDestroyFromReplication()
 
 void UAbilityTask_ApplyRootMotionRadialForce::OnDestroy(bool AbilityIsEnding)
 {
-	if (MovementComponent)
+	if (MovementComponent.IsValid())
 	{
 		MovementComponent->RemoveRootMotionSourceByID(RootMotionSourceID);
 	}
diff --git a/Engine/Plugins/Runtime/GameplayAbilities/Source/GameplayAbilities/Private/Abilities/Tasks/AbilityTask_ApplyRootMotion_Base.cpp b/Engine/Plugins/Runtime/GameplayAbilities/Source/GameplayAbilities/Private/Abilities/Tasks/AbilityTask_ApplyRootMotion_Base.cpp
--- a/Engine/Plugins/Runtime/GameplayAbilities/Source/GameplayAbilities/Private/Abilities/Tasks/AbilityTask_ApplyRootMotion_Base.cpp
+++ b/Engine/Plugins/Runtime/GameplayAbilities/Source/GameplayAbilities/Private/Abilities/Tasks/AbilityTask_ApplyRootMotion_Base.cpp
@@ -45,7 +45,7 @@ void UAbilityTask_ApplyRootMotion_Base::InitSimulatedTask(UGameplayTasksComponen
 
 bool UAbilityTask_ApplyRootMotion_Base::HasTimedOut() const
 {
-	const TSharedPtr<FRootMotionSource> RMS = (MovementComponent ? MovementComponent->GetRootMotionSourceByID(RootMotionSourceID) : nullptr);
+	const TSharedPtr<FRootMotionSource> RMS = (MovementComponent.IsValid() ? MovementComponent->GetRootMotionSourceByID(RootMotionSourceID) : nullptr);
 	if (!RMS.IsValid())
 	{
 		return true;
diff --git a/Engine/Plugins/Runtime/GameplayAbilities/Source/GameplayAbilities/Public/Abilities/Tasks/AbilityTask_ApplyRootMotion_Base.h b/Engine/Plugins/Runtime/GameplayAbilities/Source/GameplayAbilities/Public/Abilities/Tasks/AbilityTask_ApplyRootMotion_Base.h
--- a/Engine/Plugins/Runtime/GameplayAbilities/Source/GameplayAbilities/Public/Abilities/Tasks/AbilityTask_ApplyRootMotion_Base.h
+++ b/Engine/Plugins/Runtime/GameplayAbilities/Source/GameplayAbilities/Public/Abilities/Tasks/AbilityTask_ApplyRootMotion_Base.h
@@ -43,7 +43,7 @@ protected:
 	float FinishClampVelocity;
 
 	UPROPERTY()
-	TObjectPtr<UCharacterMovementComponent> MovementComponent; 
+	TWeakObjectPtr<UCharacterMovementComponent> MovementComponent; 
 	
 	uint16 RootMotionSourceID;
 
