Fix for a rare crash that can occur when FGameplayTagBlueprintPropertyMap has delegates with invalid RegisteredOwner

diff --git a/Engine/Plugins/Runtime/GameplayAbilities/Source/GameplayAbilities/Private/GameplayEffectTypes.cpp b/Engine/Plugins/Runtime/GameplayAbilities/Source/GameplayAbilities/Private/GameplayEffectTypes.cpp
--- a/Engine/Plugins/Runtime/GameplayAbilities/Source/GameplayAbilities/Private/GameplayEffectTypes.cpp
+++ b/Engine/Plugins/Runtime/GameplayAbilities/Source/GameplayAbilities/Private/GameplayEffectTypes.cpp
@@ -778,6 +778,12 @@ void FGameplayTagBlueprintPropertyMap::Unregister()
 
 void FGameplayTagBlueprintPropertyMap::GameplayTagEventCallback(const FGameplayTag Tag, int32 NewCount, TWeakObjectPtr<UObject> RegisteredOwner)
 {
+	if (!RegisteredOwner.IsValid() || !CachedOwner.IsValid())
+	{
+		ABILITY_LOG(Warning, TEXT("FGameplayTagBlueprintPropertyMap::GameplayTagEventCallback called with corrupted RegisteredOwner or corrupted CachedOwner!"));
+		return;
+	}
+	
 	// If the index and serial don't match with registered owner, the memory might be trashed so abort
 	if (!ensure(RegisteredOwner.HasSameIndexAndSerialNumber(CachedOwner)))
 	{
