Fixing a crash in FMinimalReplicationTagCountMap::UpdateOwnerTagMap.  If we received a Tag Count Update which in turn causes a callback to remove that same tag, it would invalidate our iterator and cause a crash.  Instead, I warn about this recursive situation and remove the entries in a second iteration pass.

diff --git a/Engine/Plugins/Runtime/GameplayAbilities/Source/GameplayAbilities/Private/GameplayEffectTypes.cpp b/Engine/Plugins/Runtime/GameplayAbilities/Source/GameplayAbilities/Private/GameplayEffectTypes.cpp
--- a/Engine/Plugins/Runtime/GameplayAbilities/Source/GameplayAbilities/Private/GameplayEffectTypes.cpp
+++ b/Engine/Plugins/Runtime/GameplayAbilities/Source/GameplayAbilities/Private/GameplayEffectTypes.cpp
@@ -1474,8 +1474,15 @@ void FMinimalReplicationTagCountMap::UpdateOwnerTagMap()
 	{
 		for (auto It = TagMap.CreateIterator(); It; ++It)
 		{
-			Owner->SetTagMapCount(It->Key, It->Value);
+			const FGameplayTag Tag = It->Key;
+			const int32 NewValue = It->Value;
 
+			Owner->SetTagMapCount(Tag, NewValue);
+			UE_CLOG((NewValue != It->Value), LogAbilitySystem, Warning, TEXT("Callback for SetTagMapCount(%s, %d) caused its own TagCount to change to %d (likely an unwanted recursion in your registered TagChanged Listeners)"), *Tag.ToString(), NewValue, It->Value);
+		}
+
+		for (auto It = TagMap.CreateIterator(); It; ++It)
+		{
 			// Remove tags with a count of zero from the map. This prevents them
 			// from being replicated incorrectly when recording client replays.
 			if (It->Value == 0)
