Fixed warning C4244 in GameplayAbilities plugin by adjusting variable type.

#12728

diff --git a/Engine/Plugins/Runtime/GameplayAbilities/Source/GameplayAbilities/Public/AbilitySystemLog.h b/Engine/Plugins/Runtime/GameplayAbilities/Source/GameplayAbilities/Public/AbilitySystemLog.h
--- a/Engine/Plugins/Runtime/GameplayAbilities/Source/GameplayAbilities/Public/AbilitySystemLog.h
+++ b/Engine/Plugins/Runtime/GameplayAbilities/Source/GameplayAbilities/Public/AbilitySystemLog.h
@@ -55,7 +55,7 @@ GAMEPLAYABILITIES_API DECLARE_LOG_CATEGORY_EXTERN(LogGameplayEffects, Display, A
 	if( FVisualLogger::IsRecording() ) \
 	{ \
 		static const FName GraphName("Attribute Graph"); \
-		const float CurrentTime = Actor->GetWorld() ? Actor->GetWorld()->GetTimeSeconds() : 0.f; \
+		const double CurrentTime = Actor->GetWorld() ? Actor->GetWorld()->GetTimeSeconds() : 0.0; \
 		const FVector2D OldPt(CurrentTime, OldValue); \
 		const FVector2D NewPt(CurrentTime, NewValue); \
 		const FName LineName(*AttributeName); \
