[GAS]  Adding back the explicit call to AreAbilityTagsBlocked so that users can continue to choose to override that function rather than DoesAbilitySatisfyTagRequirements.  I should note that it is *much* more preferable to use/override DoesAbilitySatisfyTagRequirements as AreAbilityTagsBlocked is not a strong enough indication if you can use an Ability or not.

I also added some better logging to figure out why an ability was not activated properly.

diff --git a/Engine/Plugins/Runtime/GameplayAbilities/Source/GameplayAbilities/Private/Abilities/GameplayAbility.cpp b/Engine/Plugins/Runtime/GameplayAbilities/Source/GameplayAbilities/Private/Abilities/GameplayAbility.cpp
--- a/Engine/Plugins/Runtime/GameplayAbilities/Source/GameplayAbilities/Private/Abilities/GameplayAbility.cpp
+++ b/Engine/Plugins/Runtime/GameplayAbilities/Source/GameplayAbilities/Private/Abilities/GameplayAbility.cpp
@@ -393,6 +393,18 @@ bool UGameplayAbility::DoesAbilitySatisfyTagRequirements(const UAbilitySystemCom
 		CheckForRequired(*TargetTags, TargetRequiredTags);
 	}
 
+	if (!bBlocked && !bMissing)
+	{
+		// If it's a custom implementation that blocks, we can't specify exactly which tag so just use the generic
+		bBlocked = AbilitySystemComponent.AreAbilityTagsBlocked(GetAssetTags());
+		if (bBlocked && OptionalRelevantTags)
+		{
+			UAbilitySystemGlobals& AbilitySystemGlobals = UAbilitySystemGlobals::Get();
+			const FGameplayTag& BlockedTag = AbilitySystemGlobals.ActivateFailTagsBlockedTag;
+			OptionalRelevantTags->AddTag(BlockedTag);
+		}
+	}
+
 	// We succeeded if there were no blocked tags and no missing required tags	
 	return !bBlocked && !bMissing;
 }
@@ -464,8 +476,8 @@ bool UGameplayAbility::CanActivateAbility(const FGameplayAbilitySpecHandle Handl
 	{
 		if (FScopedCanActivateAbilityLogEnabler::IsLoggingEnabled())
 		{
-			UE_LOG(LogAbilitySystem, Verbose, TEXT("%s: %s could not be activated due to Cooldown"), *GetNameSafe(ActorInfo->OwnerActor.Get()), *GetNameSafe(Spec->Ability));
-			UE_VLOG(ActorInfo->OwnerActor.Get(), VLogAbilitySystem, Verbose, TEXT("%s could not be activated due to Cooldown"), *GetNameSafe(Spec->Ability));
+			UE_LOG(LogAbilitySystem, Verbose, TEXT("%s: %s could not be activated due to Cooldown (%s)"), *GetNameSafe(ActorInfo->OwnerActor.Get()), *GetNameSafe(Spec->Ability), OptionalRelevantTags ? *OptionalRelevantTags->ToStringSimple() : TEXT("Unknown"));
+			UE_VLOG(ActorInfo->OwnerActor.Get(), VLogAbilitySystem, Verbose, TEXT("%s could not be activated due to Cooldown (%s)"), *GetNameSafe(Spec->Ability), OptionalRelevantTags ? *OptionalRelevantTags->ToStringSimple() : TEXT("Unknown"));
 		}
 		return false;
 	}
@@ -474,8 +486,8 @@ bool UGameplayAbility::CanActivateAbility(const FGameplayAbilitySpecHandle Handl
 	{
 		if (FScopedCanActivateAbilityLogEnabler::IsLoggingEnabled())
 		{
-			UE_LOG(LogAbilitySystem, Verbose, TEXT("%s: %s could not be activated due to Cost"), *GetNameSafe(ActorInfo->OwnerActor.Get()), *GetNameSafe(Spec->Ability));
-			UE_VLOG(ActorInfo->OwnerActor.Get(), VLogAbilitySystem, Verbose, TEXT("%s could not be activated due to Cost"), *GetNameSafe(Spec->Ability));
+			UE_LOG(LogAbilitySystem, Verbose, TEXT("%s: %s could not be activated due to Cost (%s)"), *GetNameSafe(ActorInfo->OwnerActor.Get()), *GetNameSafe(Spec->Ability), OptionalRelevantTags ? *OptionalRelevantTags->ToStringSimple() : TEXT("Unknown"));
+			UE_VLOG(ActorInfo->OwnerActor.Get(), VLogAbilitySystem, Verbose, TEXT("%s could not be activated due to Cost (%s)"), *GetNameSafe(Spec->Ability), OptionalRelevantTags ? *OptionalRelevantTags->ToStringSimple() : TEXT("Unknown"));
 		}
 		return false;
 	}
@@ -484,8 +496,8 @@ bool UGameplayAbility::CanActivateAbility(const FGameplayAbilitySpecHandle Handl
 	{	// If the ability's tags are blocked, or if it has a "Blocking" tag or is missing a "Required" tag, then it can't activate.
 		if (FScopedCanActivateAbilityLogEnabler::IsLoggingEnabled())
 		{
-			UE_LOG(LogAbilitySystem, Verbose, TEXT("%s: %s could not be activated due to Blocking Tags or Missing Required Tags"), *GetNameSafe(ActorInfo->OwnerActor.Get()), *GetNameSafe(Spec->Ability));
-			UE_VLOG(ActorInfo->OwnerActor.Get(), VLogAbilitySystem, Verbose, TEXT("%s could not be activated due to Blocking Tags or Missing Required Tags"), *GetNameSafe(Spec->Ability));
+			UE_LOG(LogAbilitySystem, Verbose, TEXT("%s: %s could not be activated due to Blocking Tags or Missing Required Tags (%s)"), *GetNameSafe(ActorInfo->OwnerActor.Get()), *GetNameSafe(Spec->Ability), OptionalRelevantTags ? *OptionalRelevantTags->ToStringSimple() : TEXT("Unknown"));
+			UE_VLOG(ActorInfo->OwnerActor.Get(), VLogAbilitySystem, Verbose, TEXT("%s could not be activated due to Blocking Tags or Missing Required Tags (%s)"), *GetNameSafe(Spec->Ability), OptionalRelevantTags ? *OptionalRelevantTags->ToStringSimple() : TEXT("Unknown"));
 		}
 		return false;
 	}
