[GAS] Allowing AbiiltySystem.Ability.Activate to distinguish between previously granted abilities and ones that need to activate immediately.  This helps debug situations where you can't activate a granted ability.

diff --git a/Engine/Plugins/Runtime/GameplayAbilities/Source/GameplayAbilities/Private/AbilitySystemCheatManagerExtension.cpp b/Engine/Plugins/Runtime/GameplayAbilities/Source/GameplayAbilities/Private/AbilitySystemCheatManagerExtension.cpp
--- a/Engine/Plugins/Runtime/GameplayAbilities/Source/GameplayAbilities/Private/AbilitySystemCheatManagerExtension.cpp
+++ b/Engine/Plugins/Runtime/GameplayAbilities/Source/GameplayAbilities/Private/AbilitySystemCheatManagerExtension.cpp
@@ -319,16 +319,33 @@ void UAbilitySystemCheatManagerExtension::AbilityActivate(const FString& Partial
 	LOG_SCOPE_VERBOSITY_OVERRIDE(LogAbilitySystem, ELogVerbosity::VeryVerbose);
 	FScopedCanActivateAbilityLogGatherer LogGatherer{ LogAbilitySystem };
 	FScopedCanActivateAbilityLogEnabler LogEnabler;
-	bool bSuccess = ASC->TryActivateAbilityByClass(ActivateAbilityClass);
-	if (bSuccess)
+	if (FGameplayAbilitySpec* GrantedAbilitySpec = ASC->FindAbilitySpecFromClass(ActivateAbilityClass))
 	{
-		UE_LOG(LogConsoleResponse, Log, TEXT("Successfully Activated Granted Ability '%s' on Player '%s'."), *GetNameSafe(ActivateAbilityClass.Get()), *GetNameSafe(PC));
+		constexpr bool bAllowRemoteActivation = false;
+		bool bSuccess = ASC->TryActivateAbility(GrantedAbilitySpec->Handle, bAllowRemoteActivation);
+#if !NO_LOGGING
+		if (bSuccess)
+		{
+			UE_LOG(LogConsoleResponse, Log, TEXT("Successfully Activated previously Granted Ability '%s' on Player '%s'."), *GetNameSafe(ActivateAbilityClass.Get()), *GetNameSafe(PC));
+		}
+		else
+		{
+			UE_LOG(LogConsoleResponse, Log, TEXT("Failed to Activate previously Granted Ability '%s' on Player '%s'. Logs:"), *GetNameSafe(ActivateAbilityClass.Get()), *GetNameSafe(PC));
+
+			TArray<FScopedCanActivateAbilityLogGatherer::FLogEntry> CapturedLogs = LogGatherer.GetCapturedLogs();
+			for (const FScopedCanActivateAbilityLogGatherer::FLogEntry& LogEntry : CapturedLogs)
+			{
+				FMsg::LogV(__FILE__, __LINE__, LogConsoleResponse.GetCategoryName(), LogEntry.Verbosity, *LogEntry.Text, {});
+			}
+		}
+#endif
 		return;
 	}
 
-	// It wasn't granted (or we failed to activate it even though it was granted).  We can't really differentiate those two, so let's grant it, then activate it.
+	// It wasn't granted, let's grant it, then activate it.
 	FGameplayAbilitySpec AbilitySpec{ ActivateAbilityClass };
 	FGameplayAbilitySpecHandle SpecHandle = ASC->GiveAbilityAndActivateOnce(AbilitySpec);
+#if !NO_LOGGING
 	if (SpecHandle.IsValid())
 	{
 		UE_LOG(LogConsoleResponse, Log, TEXT("Successfully Granted, then Activated '%s' on Player '%s'."), *GetNameSafe(ActivateAbilityClass.Get()), *GetNameSafe(PC));
@@ -337,14 +354,13 @@ void UAbilitySystemCheatManagerExtension::AbilityActivate(const FString& Partial
 	{
 		UE_LOG(LogConsoleResponse, Log, TEXT("Failed to Grant & Activate '%s' on Player '%s'. Logs:"), *GetNameSafe(ActivateAbilityClass.Get()), *GetNameSafe(PC));
 
-#if !NO_LOGGING
 		TArray<FScopedCanActivateAbilityLogGatherer::FLogEntry> CapturedLogs = LogGatherer.GetCapturedLogs();
 		for (const FScopedCanActivateAbilityLogGatherer::FLogEntry& LogEntry : CapturedLogs)
 		{
 			FMsg::LogV(__FILE__, __LINE__, LogConsoleResponse.GetCategoryName(), LogEntry.Verbosity, *LogEntry.Text, {});
 		}
-#endif
 	}
+#endif
 }
 
 void UAbilitySystemCheatManagerExtension::AbilityCancel(const FString& PartialName) const
