GameplayCueNotify_Actors that receive multiple WhileActive/OnBecameRelevant events will now ensure that OnRemove/OnCeaseRelevant events are executed afterwards.

The console variable "AbilitySystem.GameplayCueNotifyMultipleRemoveEvents 0" can be used to force a single OnRemove/OnCeaseRelevant event per instance in all situations.

[RN] minor plugin/ability

diff --git a/Engine/Plugins/Runtime/GameplayAbilities/Source/GameplayAbilities/Private/GameplayCueNotify_Actor.cpp b/Engine/Plugins/Runtime/GameplayAbilities/Source/GameplayAbilities/Private/GameplayCueNotify_Actor.cpp
--- a/Engine/Plugins/Runtime/GameplayAbilities/Source/GameplayAbilities/Private/GameplayCueNotify_Actor.cpp
+++ b/Engine/Plugins/Runtime/GameplayAbilities/Source/GameplayAbilities/Private/GameplayCueNotify_Actor.cpp
@@ -19,6 +19,13 @@
 
 #define LOCTEXT_NAMESPACE "GameplayCueNotify_Actor"
 
+static TAutoConsoleVariable<bool> CVarGameplayCueNotifyMultipleRemoveEvents(
+	TEXT("AbilitySystem.GameplayCueNotifyMultipleRemoveEvents"),
+	true,
+	TEXT("Controls how many times a GameplayCueNotify_Actor can execute its OnRemove/OnCeaseRelevant event per instance.")
+	TEXT(" true: each WhileActive/OnBecameRelevant event will permit a following OnRemove/OnCeaseRelevant event.\n")
+	TEXT(" false: only one OnRemove/OnCeaseRelevant event will execute per instance."));
+
 namespace FAbilitySystemTweaks
 {
 	int ClearCueNotifyTimers = 1;
@@ -269,6 +276,10 @@ void AGameplayCueNotify_Actor::HandleGameplayCue(AActor* MyTarget, EGameplayCueE
 		case EGameplayCueEvent::WhileActive:
 			WhileActive(MyTarget, Parameters);
 			bHasHandledWhileActiveEvent = true;
+			if (CVarGameplayCueNotifyMultipleRemoveEvents.GetValueOnGameThread())
+			{
+				bHasHandledOnRemoveEvent = false;
+			}
 			break;
 
 		case EGameplayCueEvent::Executed:
