[GAS] Preventing potential crash in CheckForClearedAbilities.
This was occuring in Lyra when traveling between levels (e.g. playnextgame).

diff --git a/Engine/Plugins/Runtime/GameplayAbilities/Source/GameplayAbilities/Private/AbilitySystemComponent_Abilities.cpp b/Engine/Plugins/Runtime/GameplayAbilities/Source/GameplayAbilities/Private/AbilitySystemComponent_Abilities.cpp
--- a/Engine/Plugins/Runtime/GameplayAbilities/Source/GameplayAbilities/Private/AbilitySystemComponent_Abilities.cpp
+++ b/Engine/Plugins/Runtime/GameplayAbilities/Source/GameplayAbilities/Private/AbilitySystemComponent_Abilities.cpp
@@ -774,10 +774,9 @@ void UAbilitySystemComponent::CheckForClearedAbilities()
 	for (int32 i = 0; i < ReplicatedAbilities.Num(); i++)
 	{
 		UGameplayAbility* Ability = ReplicatedAbilities[i];
-
 		if (!IsValid(Ability))
 		{
-			if (IsUsingRegisteredSubObjectList())
+			if (Ability && IsUsingRegisteredSubObjectList())
 			{
 				RemoveReplicatedSubObject(Ability);
 			}
