[GAS] When we're destroying an Actor that owns an AbilitySystemComponent with an active Gameplay Ability, it's possible that we will miss the FGameplayAbilitySpec::PreReplicatedRemove.  In order to ensure we get a proper EndAbility call, let's force the Cancellation of the Ability using a Replicated RPC.  This RPC is guaranteed to execute before the destruction of the Gameplay Ability Instance.

diff --git a/Engine/Plugins/Runtime/GameplayAbilities/Source/GameplayAbilities/Private/AbilitySystemComponent_Abilities.cpp b/Engine/Plugins/Runtime/GameplayAbilities/Source/GameplayAbilities/Private/AbilitySystemComponent_Abilities.cpp
--- a/Engine/Plugins/Runtime/GameplayAbilities/Source/GameplayAbilities/Private/AbilitySystemComponent_Abilities.cpp
+++ b/Engine/Plugins/Runtime/GameplayAbilities/Source/GameplayAbilities/Private/AbilitySystemComponent_Abilities.cpp
@@ -53,6 +53,7 @@ static TAutoConsoleVariable<bool> CVarGasFixClientSideMontageBlendOutTime(TEXT("
 static TAutoConsoleVariable<bool> CVarUpdateMontageSectionIdToPlay(TEXT("AbilitySystem.UpdateMontageSectionIdToPlay"), true, TEXT("During tick, update the section ID that replicated montages should use"));
 static TAutoConsoleVariable<bool> CVarReplicateMontageNextSectionId(TEXT("AbilitySystem.ReplicateMontageNextSectionId"), true, TEXT("Apply the replicated next section Id to montages when skipping position replication"));
 static TAutoConsoleVariable<bool> CVarEnsureAbilitiesEndGracefully(TEXT("AbilitySystem.EnsureAbilitiesEndGracefully"), true, TEXT("When shutting down (during ClearAllAbilities) we should check if all GameplayAbilities gracefully ended. This should be disabled if you have NonInstanced abilities that are designed for multiple concurrent executions."));
+static TAutoConsoleVariable<bool> CVarReplicateCancelDuringOnRemove(TEXT("AbilitySystem.ReplicateCancelOnRemove"), false, TEXT("When an ability is removed and not yet ended, assume it's cancelled and replicate to ensure the other side receives that before instance destruction. If false, you may miss EndAbility if the Actor is destroyed."));
 
 void UAbilitySystemComponent::InitializeComponent()
 {
@@ -631,6 +632,11 @@ void UAbilitySystemComponent::OnRemoveAbility(FGameplayAbilitySpec& Spec)
 				// End the ability but don't replicate it, OnRemoveAbility gets replicated
 				bool bReplicateEndAbility = false;
 				bool bWasCancelled = false;
+				if (CVarReplicateCancelDuringOnRemove.GetValueOnGameThread())
+				{
+					bReplicateEndAbility = true;
+					bWasCancelled = true;
+				}
 				Instance->EndAbility(Instance->CurrentSpecHandle, Instance->CurrentActorInfo, Instance->CurrentActivationInfo, bReplicateEndAbility, bWasCancelled);
 			}
 			else
