Fix for client crashing in UAbilitySystemComponent::execOnRep_OwningActor
- Added null ptr check for AbilityActorInfo

diff --git a/Engine/Plugins/Runtime/GameplayAbilities/Source/GameplayAbilities/Private/AbilitySystemComponent_Abilities.cpp b/Engine/Plugins/Runtime/GameplayAbilities/Source/GameplayAbilities/Private/AbilitySystemComponent_Abilities.cpp
--- a/Engine/Plugins/Runtime/GameplayAbilities/Source/GameplayAbilities/Private/AbilitySystemComponent_Abilities.cpp
+++ b/Engine/Plugins/Runtime/GameplayAbilities/Source/GameplayAbilities/Private/AbilitySystemComponent_Abilities.cpp
@@ -242,7 +242,10 @@ void UAbilitySystemComponent::ClearActorInfo()
 
 void UAbilitySystemComponent::OnRep_OwningActor()
 {
-	check(AbilityActorInfo.IsValid());
+	if (!ensure(AbilityActorInfo.IsValid()))
+	{
+		return;
+	}
 
 	AActor* LocalOwnerActor = GetOwnerActor();
 	AActor* LocalAvatarActor = GetAvatarActor_Direct();
