Option to defer gameplay tag processing for gameplay cues until the AR finishes and to skip preloading cues. Using a cue in this mode may cause a hitch (potentially a long one while we wait for the AR to finish scanning if it's still going).

#rnx

diff --git a/Engine/Plugins/Runtime/GameplayAbilities/Source/GameplayAbilities/Private/GameplayCueManager.cpp b/Engine/Plugins/Runtime/GameplayAbilities/Source/GameplayAbilities/Private/GameplayCueManager.cpp
--- a/Engine/Plugins/Runtime/GameplayAbilities/Source/GameplayAbilities/Private/GameplayCueManager.cpp
+++ b/Engine/Plugins/Runtime/GameplayAbilities/Source/GameplayAbilities/Private/GameplayCueManager.cpp
@@ -98,7 +98,10 @@ void UGameplayCueManager::OnEngineInitComplete()
 	AssetRegistryModule.Get().OnAssetRenamed().AddUObject(this, &UGameplayCueManager::HandleAssetRenamed);
 	FWorldDelegates::OnPreWorldInitialization.AddUObject(this, &UGameplayCueManager::ReloadObjectLibrary);
 
-	InitializeEditorObjectLibrary();
+	if (GIsEditor)
+	{
+		InitializeEditorObjectLibrary();
+	}
 #endif
 }
 
@@ -141,6 +144,17 @@ void UGameplayCueManager::HandleGameplayCue(AActor* TargetActor, FGameplayTag Ga
 #endif
 
 #if WITH_EDITOR
+	if (RuntimeGameplayCueObjectLibrary.bHasBeenInitialized == false && ShouldDeferScanningRuntimeLibraries())
+	{
+		FAssetRegistryModule& AssetRegistryModule = FModuleManager::LoadModuleChecked<FAssetRegistryModule>(TEXT("AssetRegistry"));
+		if (AssetRegistryModule.Get().IsGathering())
+		{
+			ABILITY_LOG(Warning, TEXT("Waiting for AR scan in HandleGameplayCue..."));
+			// This will trigger the callback to InitializeRuntimeObjectLibrary via OnKnownGathersComplete
+			AssetRegistryModule.Get().WaitForCompletion();
+		}
+	}
+
 	if (GIsEditor && TargetActor == nullptr && UGameplayCueManager::PreviewComponent)
 	{
 		TargetActor = Cast<AActor>(AActor::StaticClass()->GetDefaultObject());
@@ -597,8 +611,28 @@ bool UGameplayCueManager::ShouldAsyncLoadRuntimeObjectLibraries() const
 	return true;
 }
 
+bool UGameplayCueManager::ShouldDeferScanningRuntimeLibraries() const
+{
+	return false;
+}
+
 void UGameplayCueManager::InitializeRuntimeObjectLibrary()
 {
+	FAssetRegistryModule& AssetRegistryModule = FModuleManager::LoadModuleChecked<FAssetRegistryModule>(TEXT("AssetRegistry"));
+	if (AssetRegistryModule.Get().IsGathering() && ShouldDeferScanningRuntimeLibraries())
+	{
+		ABILITY_LOG(Display, TEXT("Asset registry still gathering. Deferring InitializeRuntimeObjectLibrary."));
+		if (!AssetRegistryModule.Get().OnKnownGathersComplete().IsBoundToObject(this))
+		{
+			AssetRegistryModule.Get().OnKnownGathersComplete().AddUObject(this, &UGameplayCueManager::InitializeRuntimeObjectLibrary);
+		}
+		return;
+	}
+
+	ABILITY_LOG(Display, TEXT("Processing InitializeRuntimeObjectLibrary."))
+
+	AssetRegistryModule.Get().OnKnownGathersComplete().RemoveAll(this);
+
 	UE_SCOPED_ENGINE_ACTIVITY(TEXT("Initializing GameplayCueManager Runtime Object Library"));
 
 	RuntimeGameplayCueObjectLibrary.Paths = GetAlwaysLoadedGameplayCuePaths();
diff --git a/Engine/Plugins/Runtime/GameplayAbilities/Source/GameplayAbilities/Public/GameplayCueManager.h b/Engine/Plugins/Runtime/GameplayAbilities/Source/GameplayAbilities/Public/GameplayCueManager.h
--- a/Engine/Plugins/Runtime/GameplayAbilities/Source/GameplayAbilities/Public/GameplayCueManager.h
+++ b/Engine/Plugins/Runtime/GameplayAbilities/Source/GameplayAbilities/Public/GameplayCueManager.h
@@ -251,6 +251,7 @@ protected:
 	virtual bool ShouldSyncScanRuntimeObjectLibraries() const;
 	virtual bool ShouldSyncLoadRuntimeObjectLibraries() const;
 	virtual bool ShouldAsyncLoadRuntimeObjectLibraries() const;
+	virtual bool ShouldDeferScanningRuntimeLibraries() const;
 
 	/** Refreshes the existing, already initialized, object libraries. */
 	void RefreshObjectLibraries();
