Attempting to load invalid attributes provides a warning and reassigning them doesn't crash. #jira UE-254557

diff --git a/Engine/Plugins/Runtime/GameplayAbilities/Source/GameplayAbilities/Private/AttributeSet.cpp b/Engine/Plugins/Runtime/GameplayAbilities/Source/GameplayAbilities/Private/AttributeSet.cpp
--- a/Engine/Plugins/Runtime/GameplayAbilities/Source/GameplayAbilities/Private/AttributeSet.cpp
+++ b/Engine/Plugins/Runtime/GameplayAbilities/Source/GameplayAbilities/Private/AttributeSet.cpp
@@ -278,12 +278,15 @@ void FGameplayAttribute::PostSerialize(const FArchive& Ar)
 			AttributeOwner = Attribute->GetOwnerStruct();
 			Attribute->GetName(AttributeName);
 		}
-		else if (!AttributeName.IsEmpty() && AttributeOwner != nullptr)
+		else if (!AttributeName.IsEmpty())
 		{
-			// Attempt to resolve field path from owner and attribute name
-			Attribute = FindFProperty<FProperty>(AttributeOwner, *AttributeName);
+			if (AttributeOwner != nullptr)
+			{
+				// Attempt to resolve field path from owner and attribute name
+				Attribute = FindFProperty<FProperty>(AttributeOwner, *AttributeName);
+			}
 
-			// Log warning if attribute failed to resolve while name + owner were non-null
+			// Log warning if attribute failed to resolve
 			if (!Attribute.Get())
 			{
 				FUObjectSerializeContext* LoadContext = FUObjectThreadContext::Get().GetSerializeContext();
diff --git a/Engine/Plugins/Runtime/GameplayAbilities/Source/GameplayAbilities/Public/AttributeSet.h b/Engine/Plugins/Runtime/GameplayAbilities/Source/GameplayAbilities/Public/AttributeSet.h
--- a/Engine/Plugins/Runtime/GameplayAbilities/Source/GameplayAbilities/Public/AttributeSet.h
+++ b/Engine/Plugins/Runtime/GameplayAbilities/Source/GameplayAbilities/Public/AttributeSet.h
@@ -68,7 +68,7 @@ struct GAMEPLAYABILITIES_API FGameplayAttribute
 
 	bool IsValid() const
 	{
-		return Attribute != nullptr;
+		return Attribute != nullptr && AttributeOwner != nullptr;
 	}
 
 	/** Set up from a FProperty inside a set */
