Fix double * on dirty Gameplay Ability Blueprints, The * is added via the tab suffix when dirty so we don't need to add it twice when requesting the toolkit name.
PR #11031

diff --git a/Engine/Plugins/Runtime/GameplayAbilities/Source/GameplayAbilitiesEditor/Private/GameplayAbilitiesEditor.cpp b/Engine/Plugins/Runtime/GameplayAbilities/Source/GameplayAbilitiesEditor/Private/GameplayAbilitiesEditor.cpp
--- a/Engine/Plugins/Runtime/GameplayAbilities/Source/GameplayAbilitiesEditor/Private/GameplayAbilitiesEditor.cpp
+++ b/Engine/Plugins/Runtime/GameplayAbilities/Source/GameplayAbilitiesEditor/Private/GameplayAbilitiesEditor.cpp
@@ -81,11 +81,8 @@ FText FGameplayAbilitiesEditor::GetToolkitName() const
 
 	const UObject* EditingObject = EditingObjs[0];
 
-	const bool bDirtyState = EditingObject->GetOutermost()->IsDirty();
-
 	Args.Add(TEXT("ObjectName"), FText::FromString(EditingObject->GetName()));
-	Args.Add(TEXT("DirtyState"), bDirtyState ? FText::FromString(TEXT("*")) : FText::GetEmpty());
-	return FText::Format(LOCTEXT("GameplayAbilitiesToolkitName", "{ObjectName}{DirtyState}"), Args);
+	return FText::Format(LOCTEXT("GameplayAbilitiesToolkitName", "{ObjectName}"), Args);
 }
 
 FText FGameplayAbilitiesEditor::GetToolkitToolTipText() const
