Addresses a race condition involving first-time async loading of GameplayCueNotify classes where deferred cue events (OnActive, WhileActive, OnRemove) could execute out of order. Additionally, GameplayCueManager now avoids doing multiple async load requests for the same GCN class.

Previously, GameplayCues could be executed out of order when the async load of a GCN class has progressed enough that the class is resolvable by path, but the FStreamingDelegate async load callback hasn't fired yet since those are deferred until next frame. During that frame, gameplay cues will execute immediately while prior cues that were deferred haven't executed yet. Cues had a chance of being executed out of order: i.e. calling AddGameplayCue and RemoveGameplayCue could result in the remove being executed before the add. In practice this ends up with cues that remain active while they shouldn't.

This fix caches multiple gameplay cues and their parameters while the GCN class is being async loaded. They are executed in order when the GameplayCueManager receives the callback that the class has finished loading. All cues in the meantime will now consistently be cached, instead of having a chance to execute immediately. In addition, only one async load request will be done per GCN class. Previously each cue request would fire another async load request. When adding GameplayCues, this previously always resulted in two requests (OnActive, WhileActive), now that's one.

diff --git a/Engine/Plugins/Runtime/GameplayAbilities/Source/GameplayAbilities/Private/GameplayCueManager.cpp b/Engine/Plugins/Runtime/GameplayAbilities/Source/GameplayAbilities/Private/GameplayCueManager.cpp
--- a/Engine/Plugins/Runtime/GameplayAbilities/Source/GameplayAbilities/Private/GameplayCueManager.cpp
+++ b/Engine/Plugins/Runtime/GameplayAbilities/Source/GameplayAbilities/Private/GameplayCueManager.cpp
@@ -328,6 +328,11 @@ bool UGameplayCueManager::ShouldAsyncLoadMissingGameplayCues() const
 	return true;
 }
 
+bool UGameplayCueManager::IsAsyncLoadingGameplayCueNotifyClass(const FSoftObjectPath& CueNotifyClass) const
+{
+	return AsyncLoadPendingGameplayCues.Contains(CueNotifyClass);
+}
+
 bool UGameplayCueManager::HandleMissingGameplayCue(UGameplayCueSet* OwningSet, struct FGameplayCueNotifyData& CueData, AActor* TargetActor, EGameplayCueEvent::Type EventType, FGameplayCueParameters& Parameters)
 {
 	if (ShouldSyncLoadMissingGameplayCues())
@@ -346,49 +351,81 @@ bool UGameplayCueManager::HandleMissingGameplayCue(UGameplayCueSet* OwningSet, s
 	}
 	else if (ShouldAsyncLoadMissingGameplayCues())
 	{
-		// Not loaded: start async loading and call when loaded
-		StreamableManager.RequestAsyncLoad(CueData.GameplayCueNotifyObj, FStreamableDelegate::CreateUObject(this, &UGameplayCueManager::OnMissingCueAsyncLoadComplete, 
-			CueData.GameplayCueNotifyObj, TWeakObjectPtr<UGameplayCueSet>(OwningSet), CueData.GameplayCueTag, MakeWeakObjectPtr(TargetActor), EventType, Parameters));
+		const bool bStartAsyncLoad = !IsAsyncLoadingGameplayCueNotifyClass(CueData.GameplayCueNotifyObj);
+		const FAsyncLoadPendingGameplayCue PendingCue(TWeakObjectPtr<UGameplayCueSet>(OwningSet), CueData.GameplayCueTag, MakeWeakObjectPtr(TargetActor), EventType, Parameters);
+		if (bStartAsyncLoad)
+		{
+			ABILITY_LOG(Display, TEXT("Deferred gameplay cue with tag '%s' and event type '%d' until GCN class '%s' is loaded (started async loading)."), *CueData.GameplayCueTag.ToString(), int32(EventType), *CueData.GameplayCueNotifyObj.ToString());
 
-		ABILITY_LOG(Display, TEXT("GameplayCueNotify %s was not loaded when GameplayCue was invoked. Starting async loading."), *CueData.GameplayCueNotifyObj.ToString());
+			// Cache the cue parameters to call it when the GCN class finishes async loading
+			TArray<FAsyncLoadPendingGameplayCue>& PendingEvents = AsyncLoadPendingGameplayCues.Add(CueData.GameplayCueNotifyObj);
+			PendingEvents.Add(PendingCue);
+
+			// Start async loading the GCN class
+			StreamableManager.RequestAsyncLoad(CueData.GameplayCueNotifyObj, FStreamableDelegate::CreateUObject(this, &UGameplayCueManager::OnMissingCueAsyncLoadComplete, CueData.GameplayCueNotifyObj));
+		}
+		else
+		{
+			ABILITY_LOG(Display, TEXT("Deferred gameplay cue with tag '%s' and event type '%d' until GCN class '%s' finishes async loading (already loading)."), *CueData.GameplayCueTag.ToString(), int32(EventType), *CueData.GameplayCueNotifyObj.ToString());
+
+			// Cache the cue parameters to call it when the GCN class finishes async loading
+			AsyncLoadPendingGameplayCues[CueData.GameplayCueNotifyObj].Add(PendingCue);
+		}
 	}
 	return false;
 }
 
-void UGameplayCueManager::OnMissingCueAsyncLoadComplete(FSoftObjectPath LoadedObject, TWeakObjectPtr<UGameplayCueSet> OwningSet, FGameplayTag GameplayCueTag, TWeakObjectPtr<AActor> TargetActor, EGameplayCueEvent::Type EventType, FGameplayCueParameters Parameters)
+void UGameplayCueManager::OnMissingCueAsyncLoadComplete(FSoftObjectPath LoadedNotifyClass)
 {
-	if (!LoadedObject.ResolveObject())
+	if (!LoadedNotifyClass.ResolveObject())
 	{
 		// Load failed
-		ABILITY_LOG(Warning, TEXT("Late load of GameplayCueNotify %s failed!"), *LoadedObject.ToString());
+		ABILITY_LOG(Warning, TEXT("Late load of GameplayCueNotify %s failed!"), *LoadedNotifyClass.ToString());
 		return;
 	}
 
-	if (OwningSet.IsValid())
+	if (ensureMsgf(AsyncLoadPendingGameplayCues.Contains(LoadedNotifyClass), TEXT("Expected an entry for '%s' that was created when async loading for it started"), *LoadedNotifyClass.ToString()))
 	{
-		CurrentWorld = TargetActor.IsValid() ? TargetActor->GetWorld() : nullptr;
-		if (!CurrentWorld)
+		// Keep the array of the pending cue events and their parameters. Remove the GCN class from the map of pending events 
+		// because we're no longer async loading the class: the calls to HandleGameplayCue that we're doing now should execute directly.
+		TArray<FAsyncLoadPendingGameplayCue> PendingEventsCopy = MoveTemp(AsyncLoadPendingGameplayCues[LoadedNotifyClass]);
+		AsyncLoadPendingGameplayCues.Remove(LoadedNotifyClass);
+
+		ABILITY_LOG(Display, TEXT("Finished async loading GameplayCueNotify class '%s'. Executing %d deferred gameplay cues."), *LoadedNotifyClass.ToString(), PendingEventsCopy.Num());
+
+		// Call all pending cues for the GCN class that finished loading
+		for (const FAsyncLoadPendingGameplayCue& PendingEvent : PendingEventsCopy)
 		{
-			// TargetActor has since been destroyed.  Attempt to get the world from the other actors.
-			const AActor* CueInstigator = Parameters.GetInstigator();
-			CurrentWorld = CueInstigator ? CueInstigator->GetWorld() : nullptr;
-			if (!CurrentWorld)
+			// Owning set may have been GCed in the meantime
+			if (PendingEvent.OwningSet.IsValid())
 			{
-				const AActor* EffectCauser = Parameters.GetEffectCauser();
-				CurrentWorld = EffectCauser ? EffectCauser->GetWorld() : nullptr;
-			}
-		}
+				CurrentWorld = PendingEvent.TargetActor.IsValid() ? PendingEvent.TargetActor->GetWorld() : nullptr;
+				if (!CurrentWorld)
+				{
+					// TargetActor has since been destroyed.  Attempt to get the world from the other actors.
+					const AActor* CueInstigator = PendingEvent.Parameters.GetInstigator();
+					CurrentWorld = CueInstigator ? CueInstigator->GetWorld() : nullptr;
+					if (!CurrentWorld)
+					{
+						const AActor* EffectCauser = PendingEvent.Parameters.GetEffectCauser();
+						CurrentWorld = EffectCauser ? EffectCauser->GetWorld() : nullptr;
+					}
+				}
 
-		// Don't handle gameplay cues when world is tearing down
-		if (!GetWorld() || GetWorld()->bIsTearingDown)
-		{
-			return;
-		}
+				// Don't handle gameplay cues for worlds that are tearing down
+				if (!GetWorld() || GetWorld()->bIsTearingDown)
+				{
+					ABILITY_LOG(Verbose, TEXT("Skipping deferred gameplay cue due to invalid world. GCN class: %s | GameplayCueTag: %s | EventType: %d"), *LoadedNotifyClass.ToString(), *PendingEvent.GameplayCueTag.ToString(), int32(PendingEvent.EventType));
+					continue;
+				}
 
-		// Objects are still valid, re-execute cue
-		OwningSet->HandleGameplayCue(TargetActor.Get(), GameplayCueTag, EventType, Parameters);
+				// Objects are still valid, re-execute cue
+				ABILITY_LOG(Verbose, TEXT("Executing deferred gameplay cue. GCN class: %s | GameplayCueTag: %s | EventType: %d"), *LoadedNotifyClass.ToString(), *PendingEvent.GameplayCueTag.ToString(), int32(PendingEvent.EventType));
+				PendingEvent.OwningSet->HandleGameplayCue(PendingEvent.TargetActor.Get(), PendingEvent.GameplayCueTag, PendingEvent.EventType, PendingEvent.Parameters);
 
-		CurrentWorld = nullptr;
+				CurrentWorld = nullptr;
+			}
+		}
 	}
 }
 
diff --git a/Engine/Plugins/Runtime/GameplayAbilities/Source/GameplayAbilities/Private/GameplayCueSet.cpp b/Engine/Plugins/Runtime/GameplayAbilities/Source/GameplayAbilities/Private/GameplayCueSet.cpp
--- a/Engine/Plugins/Runtime/GameplayAbilities/Source/GameplayAbilities/Private/GameplayCueSet.cpp
+++ b/Engine/Plugins/Runtime/GameplayAbilities/Source/GameplayAbilities/Private/GameplayCueSet.cpp
@@ -279,8 +279,18 @@ bool UGameplayCueSet::HandleGameplayCueNotify_Internal(AActor* TargetActor, int3
 		// If object is not loaded yet
 		if (CueData.LoadedGameplayCueClass == nullptr || bDebugFailLoads)
 		{
-			// See if the object is loaded but just not hooked up here
-			CueData.LoadedGameplayCueClass = Cast<UClass>(CueData.GameplayCueNotifyObj.ResolveObject());
+			// See if the object is loaded but just not hooked up here. Only do this just-in-time
+			// detection of GCN class when the CueManager isn't waiting for an async load of the
+			// class. Otherwise, defer to CueManager->HandleMissingGameplayCue so that it can queue 
+			// events and execute them in order on async load completion. This addresses a race 
+			// condition where a UClass is resolvable but the async load callback hasn't fired yet 
+			// since those are deferred to the next frame. Previously this resulted in cues being
+			// executed out of order (remove before add).
+			if (!CueManager->IsAsyncLoadingGameplayCueNotifyClass(CueData.GameplayCueNotifyObj))
+			{
+				CueData.LoadedGameplayCueClass = Cast<UClass>(CueData.GameplayCueNotifyObj.ResolveObject());
+			}
+
 			if (CueData.LoadedGameplayCueClass == nullptr || bDebugFailLoads)
 			{
 				if (!CueManager->HandleMissingGameplayCue(this, CueData, TargetActor, EventType, Parameters))
diff --git a/Engine/Plugins/Runtime/GameplayAbilities/Source/GameplayAbilities/Public/GameplayCueManager.h b/Engine/Plugins/Runtime/GameplayAbilities/Source/GameplayAbilities/Public/GameplayCueManager.h
--- a/Engine/Plugins/Runtime/GameplayAbilities/Source/GameplayAbilities/Public/GameplayCueManager.h
+++ b/Engine/Plugins/Runtime/GameplayAbilities/Source/GameplayAbilities/Public/GameplayCueManager.h
@@ -103,6 +103,25 @@ struct FGameplayCueObjectLibrary
 	bool bHasBeenInitialized;
 };
 
+/** While a GameplayCueNotify class is being loaded async, GameplayCue events will be deferred until load completion. */
+struct FAsyncLoadPendingGameplayCue
+{
+	FAsyncLoadPendingGameplayCue() : 
+		OwningSet(nullptr), GameplayCueTag(), TargetActor(nullptr), EventType(EGameplayCueEvent::Type::OnActive), Parameters() 
+	{
+	}
+
+	FAsyncLoadPendingGameplayCue(TWeakObjectPtr<UGameplayCueSet> OwningSet, FGameplayTag GameplayCueTag, TWeakObjectPtr<AActor> TargetActor, EGameplayCueEvent::Type EventType, const FGameplayCueParameters& Parameters) :
+		OwningSet(OwningSet), GameplayCueTag(GameplayCueTag), TargetActor(TargetActor), EventType(EventType), Parameters(Parameters)
+	{
+	}
+
+	TWeakObjectPtr<UGameplayCueSet> OwningSet;
+	FGameplayTag GameplayCueTag;
+	TWeakObjectPtr<AActor> TargetActor;
+	TEnumAsByte<EGameplayCueEvent::Type> EventType;
+	FGameplayCueParameters Parameters;
+};
 
 /** Singleton manager object that handles dispatching gameplay cues and spawning GameplayCueNotify actors as needed */
 UCLASS()
@@ -362,6 +381,9 @@ public:
 	/** If true, this will asynchronously load missing gameplay cues, and execute cue when the load finishes */
 	virtual bool ShouldAsyncLoadMissingGameplayCues() const;
 
+	/** Whether a GameplayCueNotify class is currently loading (async load requested by this and streaming callback not received yet). */
+	bool IsAsyncLoadingGameplayCueNotifyClass(const FSoftObjectPath& CueNotifyClass) const;
+
 	/** Handles what to do if a missing cue is requested. If return true, this was loaded and execution should continue */
 	virtual bool HandleMissingGameplayCue(UGameplayCueSet* OwningSet, struct FGameplayCueNotifyData& CueData, AActor* TargetActor, EGameplayCueEvent::Type EventType, FGameplayCueParameters& Parameters);
 
@@ -379,7 +401,7 @@ protected:
 
 	void OnGameplayCueNotifyAsyncLoadComplete(TArray<FSoftObjectPath> StringRef);
 
-	void OnMissingCueAsyncLoadComplete(FSoftObjectPath LoadedObject, TWeakObjectPtr<UGameplayCueSet> OwningSet, FGameplayTag GameplayCueTag, TWeakObjectPtr<AActor> TargetActor, EGameplayCueEvent::Type EventType, FGameplayCueParameters Parameters);
+	void OnMissingCueAsyncLoadComplete(FSoftObjectPath LoadedNotifyClass);
 
 	void CheckForPreallocation(UClass* GCClass);
 
@@ -407,6 +429,9 @@ protected:
 	UPROPERTY(transient)
 	TArray<FPreallocationInfo>	PreallocationInfoList_Internal;
 
+	/** While a GCN class was not loaded yet, events are requested and async loading is enabled, cache them until the class finishes async loading. */
+	TMap<FSoftObjectPath, TArray<FAsyncLoadPendingGameplayCue>> AsyncLoadPendingGameplayCues;
+
 	FOnRouteGameplayCue OnRouteGameplayCue;
 
 private:
