GAMEPLAYABILITIES_API added to cue notify structs so they are exposed and can be aggregated for use with custom GameplayCueNotify objects.
Structs are marked final so that they cannot be extended and lead to slicing issues.

diff --git a/Engine/Plugins/Runtime/GameplayAbilities/Source/GameplayAbilities/Private/GameplayCueNotifyTypes.cpp b/Engine/Plugins/Runtime/GameplayAbilities/Source/GameplayAbilities/Private/GameplayCueNotifyTypes.cpp
--- a/Engine/Plugins/Runtime/GameplayAbilities/Source/GameplayAbilities/Private/GameplayCueNotifyTypes.cpp
+++ b/Engine/Plugins/Runtime/GameplayAbilities/Source/GameplayAbilities/Private/GameplayCueNotifyTypes.cpp
@@ -379,6 +379,26 @@ APlayerController* FGameplayCueNotify_SpawnContext::FindLocalPlayerController(EG
 }
 
 
+//////////////////////////////////////////////////////////////////////////
+// FGameplayCueNotify_SpawnResult
+//////////////////////////////////////////////////////////////////////////
+FGameplayCueNotify_SpawnResult::FGameplayCueNotify_SpawnResult()
+{
+	Reset();
+}
+
+void FGameplayCueNotify_SpawnResult::Reset()
+{
+	FxSystemComponents.Reset();
+	AudioComponents.Reset();
+	CameraShakes.Reset();
+	CameraLensEffects.Reset();
+	ForceFeedbackComponent = nullptr;
+	ForceFeedbackTargetPC = nullptr;
+	DecalComponent = nullptr;
+}
+
+
 //////////////////////////////////////////////////////////////////////////
 // FGameplayCueNotify_ParticleInfo
 //////////////////////////////////////////////////////////////////////////
diff --git a/Engine/Plugins/Runtime/GameplayAbilities/Source/GameplayAbilities/Public/GameplayCueNotifyTypes.h b/Engine/Plugins/Runtime/GameplayAbilities/Source/GameplayAbilities/Public/GameplayCueNotifyTypes.h
--- a/Engine/Plugins/Runtime/GameplayAbilities/Source/GameplayAbilities/Public/GameplayCueNotifyTypes.h
+++ b/Engine/Plugins/Runtime/GameplayAbilities/Source/GameplayAbilities/Public/GameplayCueNotifyTypes.h
@@ -33,6 +33,29 @@ static_assert((EPhysicalSurface::SurfaceType_Max <= (sizeof(FGameplayCueNotify_S
 
 DECLARE_LOG_CATEGORY_EXTERN(LogGameplayCueNotify, Log, All);
 
+/*
+ * The Gameplay Cue Notify structs below cannot be inherited, but they are now exposed and can be aggregated.
+ * If you want to extend functionality in your own structs you can include the base struct in your own and use it in
+ * your own Notify Cue objects. EX
+ *
+
+USTRUCT(BlueprintType)
+struct FMyGameplayCueNotify_CustomBurstEffects
+{
+	UPROPERTY(EditDefaultsOnly, BlueprintReadOnly, Category = GameplayCueNotify)
+	FGameplayCueNotify_BurstEffects BaseBurstEffects;
+
+	UPROPERTY(EditDefaultsOnly, BlueprintReadOnly, Category = GameplayCueNotify)
+	float CustomFloat;
+
+	void CustomExecute()
+	{
+		BaseBurstEffects.ExecuteEffects();
+	}
+};
+
+ *
+ **/
 
 /**
  * EGameplayCueNotify_EffectPlaySpace
@@ -91,17 +114,13 @@ enum class EGameplayCueNotify_LocallyControlledPolicy : uint8
  *	Conditions used to determine if the gameplay cue notify should spawn.
  */
 USTRUCT(BlueprintType)
-struct FGameplayCueNotify_SpawnCondition
+struct FGameplayCueNotify_SpawnCondition final
 {
 	GENERATED_BODY()
 
-public:
-
-	FGameplayCueNotify_SpawnCondition();
+	GAMEPLAYABILITIES_API FGameplayCueNotify_SpawnCondition();
 
-	bool ShouldSpawn(const FGameplayCueNotify_SpawnContext& SpawnContext) const;
-
-public:
+	GAMEPLAYABILITIES_API bool ShouldSpawn(const FGameplayCueNotify_SpawnContext& SpawnContext) const;
 
 	// Source actor to use when determining if it is locally controlled.
 	UPROPERTY(EditDefaultsOnly, BlueprintReadOnly, Category = GameplayCueNotify)
@@ -151,20 +170,16 @@ enum class EGameplayCueNotify_AttachPolicy : uint8
  *	Specifies how the gameplay cue notify will be positioned in the world.
  */
 USTRUCT(BlueprintType)
-struct FGameplayCueNotify_PlacementInfo
+struct FGameplayCueNotify_PlacementInfo final
 {
 	GENERATED_BODY()
 
-public:
-
-	FGameplayCueNotify_PlacementInfo();
-
-	bool FindSpawnTransform(const FGameplayCueNotify_SpawnContext& SpawnContext, FTransform& OutTransform) const;
+	GAMEPLAYABILITIES_API FGameplayCueNotify_PlacementInfo();
 
-	void SetComponentTransform(USceneComponent* Component, const FTransform& Transform) const;
-	void TryComponentAttachment(USceneComponent* Component, USceneComponent* AttachComponent) const;
+	GAMEPLAYABILITIES_API bool FindSpawnTransform(const FGameplayCueNotify_SpawnContext& SpawnContext, FTransform& OutTransform) const;
 
-public:
+	GAMEPLAYABILITIES_API void SetComponentTransform(USceneComponent* Component, const FTransform& Transform) const;
+	GAMEPLAYABILITIES_API void TryComponentAttachment(USceneComponent* Component, USceneComponent* AttachComponent) const;
 
 	// Target's socket (or bone) used for location and rotation.  If "None", it uses the target's root.
 	UPROPERTY(EditDefaultsOnly, BlueprintReadOnly, Category = GameplayCueNotify)
@@ -203,14 +218,14 @@ public:
  *
  *	Temporary spawn information collected from input parameters.
  */
-struct FGameplayCueNotify_SpawnContext
+struct FGameplayCueNotify_SpawnContext final
 {
 public:
 
-	FGameplayCueNotify_SpawnContext(UWorld* InWorld, AActor* InTargetActor, const FGameplayCueParameters& InCueParameters);
+	GAMEPLAYABILITIES_API FGameplayCueNotify_SpawnContext(UWorld* InWorld, AActor* InTargetActor, const FGameplayCueParameters& InCueParameters);
 
-	void SetDefaultSpawnCondition(const FGameplayCueNotify_SpawnCondition* InSpawnCondition) {DefaultSpawnCondition = InSpawnCondition;}
-	void SetDefaultPlacementInfo(const FGameplayCueNotify_PlacementInfo* InPlacementInfo) {DefaultPlacementInfo = InPlacementInfo;}
+	GAMEPLAYABILITIES_API void SetDefaultSpawnCondition(const FGameplayCueNotify_SpawnCondition* InSpawnCondition) {DefaultSpawnCondition = InSpawnCondition;}
+	GAMEPLAYABILITIES_API void SetDefaultPlacementInfo(const FGameplayCueNotify_PlacementInfo* InPlacementInfo) {DefaultPlacementInfo = InPlacementInfo;}
 
 	const FGameplayCueNotify_PlacementInfo& GetPlacementInfo(bool bUseOverride, const FGameplayCueNotify_PlacementInfo& PlacementInfoOverride) const
 	{
@@ -222,7 +237,7 @@ public:
 		return (!bUseOverride && DefaultSpawnCondition) ? *DefaultSpawnCondition : SpawnConditionOverride;
 	}
 
-	APlayerController* FindLocalPlayerController(EGameplayCueNotify_LocallyControlledSource Source) const;
+	GAMEPLAYABILITIES_API APlayerController* FindLocalPlayerController(EGameplayCueNotify_LocallyControlledSource Source) const;
 
 protected:
 
@@ -250,26 +265,13 @@ private:
  *	Temporary structure used to track results of spawning components.
  */
 USTRUCT(BlueprintType)
-struct FGameplayCueNotify_SpawnResult
+struct FGameplayCueNotify_SpawnResult final
 {
 	GENERATED_BODY()
 
-public:
+	GAMEPLAYABILITIES_API FGameplayCueNotify_SpawnResult();
 
-	FGameplayCueNotify_SpawnResult() { Reset(); }
-
-	void Reset()
-	{
-		FxSystemComponents.Reset();
-		AudioComponents.Reset();
-		CameraShakes.Reset();
-		CameraLensEffects.Reset();
-		ForceFeedbackComponent = nullptr;
-		ForceFeedbackTargetPC = nullptr;
-		DecalComponent = nullptr;
-	}
-
-public:
+	GAMEPLAYABILITIES_API void Reset();
 
 	// List of FX components spawned.  There may be null pointers here as it matches the defined order.
 	UPROPERTY(BlueprintReadOnly, Transient, Category = GameplayCueNotify)
@@ -307,19 +309,16 @@ public:
  *	Properties that specify how to play a particle effect.
  */
 USTRUCT(BlueprintType)
-struct FGameplayCueNotify_ParticleInfo
+struct FGameplayCueNotify_ParticleInfo final
 {
 	GENERATED_BODY()
+	
+	GAMEPLAYABILITIES_API FGameplayCueNotify_ParticleInfo();
 
-public:
-
-	FGameplayCueNotify_ParticleInfo();
-
-	bool PlayParticleEffect(const FGameplayCueNotify_SpawnContext& SpawnContext, FGameplayCueNotify_SpawnResult& OutSpawnResult) const;
-
-	void ValidateBurstAssets(const UObject* ContainingAsset, const FString& Context, class FDataValidationContext& ValidationContext) const;
+	GAMEPLAYABILITIES_API bool PlayParticleEffect(const FGameplayCueNotify_SpawnContext& SpawnContext, FGameplayCueNotify_SpawnResult& 
+	OutSpawnResult) const;
 
-public:
+	GAMEPLAYABILITIES_API void ValidateBurstAssets(const UObject* ContainingAsset, const FString& Context, class FDataValidationContext& ValidationContext) const;
 
 	// Condition to check before spawning the particle system.
 	UPROPERTY(EditDefaultsOnly, BlueprintReadOnly, Category = GameplayCueNotify, Meta = (EditCondition = "bOverrideSpawnCondition"))
@@ -352,13 +351,13 @@ public:
  *	Properties that specify how to interface with the ISoundParameterControllerInterface
  */
 USTRUCT(BlueprintType)
-struct FGameplayCueNotify_SoundParameterInterfaceInfo
+struct FGameplayCueNotify_SoundParameterInterfaceInfo final
 {
 	GENERATED_BODY()
 
 public:
 
-	FGameplayCueNotify_SoundParameterInterfaceInfo();
+	GAMEPLAYABILITIES_API FGameplayCueNotify_SoundParameterInterfaceInfo();
 
 	// The name of the stop trigger set via the parameter interface
 	UPROPERTY(EditDefaultsOnly, BlueprintReadOnly, Category = GameplayCueNotify)
@@ -371,19 +370,16 @@ public:
  *	Properties that specify how to play a sound effect.
  */
 USTRUCT(BlueprintType)
-struct FGameplayCueNotify_SoundInfo
+struct FGameplayCueNotify_SoundInfo final
 {
 	GENERATED_BODY()
+	
+	GAMEPLAYABILITIES_API FGameplayCueNotify_SoundInfo();
 
-public:
-
-	FGameplayCueNotify_SoundInfo();
-
-	bool PlaySound(const FGameplayCueNotify_SpawnContext& SpawnContext, FGameplayCueNotify_SpawnResult& OutSpawnResult) const;
-
-	void ValidateBurstAssets(const UObject* ContainingAsset, const FString& Context, class FDataValidationContext& ValidationContext) const;
+	GAMEPLAYABILITIES_API bool PlaySound(const FGameplayCueNotify_SpawnContext& SpawnContext, FGameplayCueNotify_SpawnResult& OutSpawnResult) 
+	const;
 
-public:
+	GAMEPLAYABILITIES_API void ValidateBurstAssets(const UObject* ContainingAsset, const FString& Context, class FDataValidationContext& ValidationContext) const;
 
 	// Condition to check before playing the sound.
 	UPROPERTY(EditDefaultsOnly, BlueprintReadOnly, Category = GameplayCueNotify, Meta = (EditCondition = "bOverrideSpawnCondition"))
@@ -435,19 +431,16 @@ public:
  *	Properties that specify how to play a camera shake.
  */
 USTRUCT(BlueprintType)
-struct FGameplayCueNotify_CameraShakeInfo
+struct FGameplayCueNotify_CameraShakeInfo final
 {
 	GENERATED_BODY()
 
-public:
-
-	FGameplayCueNotify_CameraShakeInfo();
+	GAMEPLAYABILITIES_API FGameplayCueNotify_CameraShakeInfo();
 
-	bool PlayCameraShake(const FGameplayCueNotify_SpawnContext& SpawnContext, FGameplayCueNotify_SpawnResult& OutSpawnResult) const;
+	GAMEPLAYABILITIES_API bool PlayCameraShake(const FGameplayCueNotify_SpawnContext& SpawnContext, FGameplayCueNotify_SpawnResult& 
+	OutSpawnResult) const;
 
-	void ValidateBurstAssets(const UObject* ContainingAsset, const FString& Context, class FDataValidationContext& ValidationContext) const;
-
-public:
+	GAMEPLAYABILITIES_API void ValidateBurstAssets(const UObject* ContainingAsset, const FString& Context, class FDataValidationContext& ValidationContext) const;
 
 	// Condition to check before playing the camera shake.
 	UPROPERTY(EditDefaultsOnly, BlueprintReadOnly, Category = GameplayCueNotify, Meta = (EditCondition = "bOverrideSpawnCondition"))
@@ -501,19 +494,16 @@ public:
  *	Properties that specify how to play a camera lens effect.
  */
 USTRUCT(BlueprintType)
-struct FGameplayCueNotify_CameraLensEffectInfo
+struct FGameplayCueNotify_CameraLensEffectInfo final
 {
 	GENERATED_BODY()
 
-public:
-
-	FGameplayCueNotify_CameraLensEffectInfo();
+	GAMEPLAYABILITIES_API FGameplayCueNotify_CameraLensEffectInfo();
 
-	bool PlayCameraLensEffect(const FGameplayCueNotify_SpawnContext& SpawnContext, FGameplayCueNotify_SpawnResult& OutSpawnResult) const;
+	GAMEPLAYABILITIES_API bool PlayCameraLensEffect(const FGameplayCueNotify_SpawnContext& SpawnContext, FGameplayCueNotify_SpawnResult& 
+	OutSpawnResult) const;
 
-	void ValidateBurstAssets(const UObject* ContainingAsset, const FString& Context, class FDataValidationContext& ValidationContext) const;
-
-public:
+	GAMEPLAYABILITIES_API void ValidateBurstAssets(const UObject* ContainingAsset, const FString& Context, class FDataValidationContext& ValidationContext) const;
 
 	// Condition to check before playing the camera lens effect.
 	UPROPERTY(EditDefaultsOnly, BlueprintReadOnly, Category = GameplayCueNotify, Meta = (EditCondition = "bOverrideSpawnCondition"))
@@ -555,19 +545,16 @@ public:
  *	Properties that specify how to play a force feedback effect.
  */
 USTRUCT(BlueprintType)
-struct FGameplayCueNotify_ForceFeedbackInfo
+struct FGameplayCueNotify_ForceFeedbackInfo final
 {
 	GENERATED_BODY()
 
-public:
-
-	FGameplayCueNotify_ForceFeedbackInfo();
+	GAMEPLAYABILITIES_API FGameplayCueNotify_ForceFeedbackInfo();
 
-	bool PlayForceFeedback(const FGameplayCueNotify_SpawnContext& SpawnContext, FGameplayCueNotify_SpawnResult& OutSpawnResult) const;
+	GAMEPLAYABILITIES_API bool PlayForceFeedback(const FGameplayCueNotify_SpawnContext& SpawnContext, FGameplayCueNotify_SpawnResult& 
+	OutSpawnResult) const;
 	
-	void ValidateBurstAssets(const UObject* ContainingAsset, const FString& Context, class FDataValidationContext& ValidationContext) const;
-
-public:
+	GAMEPLAYABILITIES_API void ValidateBurstAssets(const UObject* ContainingAsset, const FString& Context, class FDataValidationContext& ValidationContext) const;
 
 	// Condition to check before playing the force feedback.
 	UPROPERTY(EditDefaultsOnly, BlueprintReadOnly, Category = GameplayCueNotify, Meta = (EditCondition = "bOverrideSpawnCondition"))
@@ -618,15 +605,18 @@ public:
  * Properties that specify how to set input device properties during a gameplay cue notify
  */
 USTRUCT(BlueprintType)
-struct FGameplayCueNotify_InputDevicePropertyInfo
+struct FGameplayCueNotify_InputDevicePropertyInfo final
 {
 	GENERATED_BODY()
+	
+	GAMEPLAYABILITIES_API FGameplayCueNotify_InputDevicePropertyInfo() = default;
 
 	/** Set the device properties on specified on this struct on the Input Device Subsystem. */
-	bool SetDeviceProperties(const FGameplayCueNotify_SpawnContext& SpawnContext, FGameplayCueNotify_SpawnResult& OutSpawnResult) const;
+	GAMEPLAYABILITIES_API bool SetDeviceProperties(const FGameplayCueNotify_SpawnContext& SpawnContext, FGameplayCueNotify_SpawnResult& 
+	OutSpawnResult) const;
 
 	/** Validate that the device properties in this effect are usable */
-	void ValidateBurstAssets(const UObject* ContainingAsset, const FString& Context, class FDataValidationContext& ValidationContext) const;
+	GAMEPLAYABILITIES_API void ValidateBurstAssets(const UObject* ContainingAsset, const FString& Context, class FDataValidationContext& ValidationContext) const;
 	
 	/** Input Device properties to apply */
 	UPROPERTY(EditAnywhere, BlueprintReadOnly, Category = GameplayCueNotify)
@@ -639,17 +629,14 @@ struct FGameplayCueNotify_InputDevicePropertyInfo
  *	Properties that specify how to spawn a decal.
  */
 USTRUCT(BlueprintType)
-struct FGameplayCueNotify_DecalInfo
+struct FGameplayCueNotify_DecalInfo final
 {
 	GENERATED_BODY()
 
-public:
-	
-	FGameplayCueNotify_DecalInfo();
+	GAMEPLAYABILITIES_API FGameplayCueNotify_DecalInfo();
 
-	bool SpawnDecal(const FGameplayCueNotify_SpawnContext& SpawnContext, FGameplayCueNotify_SpawnResult& OutSpawnResult) const;
-
-public:
+	GAMEPLAYABILITIES_API bool SpawnDecal(const FGameplayCueNotify_SpawnContext& SpawnContext, FGameplayCueNotify_SpawnResult& OutSpawnResult) 
+	const;
 
 	// Condition to check before spawning the decal.
 	UPROPERTY(EditDefaultsOnly, BlueprintReadOnly, Category = GameplayCueNotify, Meta = (EditCondition = "bOverrideSpawnCondition"))
@@ -695,18 +682,16 @@ public:
  *	Set of effects to spawn for a single event, used by all gameplay cue notify types.
  */
 USTRUCT(BlueprintType)
-struct FGameplayCueNotify_BurstEffects
+struct FGameplayCueNotify_BurstEffects final
 {
 	GENERATED_BODY()
 
-public:
-
-	FGameplayCueNotify_BurstEffects();
-
-	void ExecuteEffects(const FGameplayCueNotify_SpawnContext& SpawnContext, FGameplayCueNotify_SpawnResult& OutSpawnResult) const;
+	GAMEPLAYABILITIES_API FGameplayCueNotify_BurstEffects();
 
-	void ValidateAssociatedAssets(const UObject* ContainingAsset, const FString& Context, class FDataValidationContext& ValidationContext) const;
+	GAMEPLAYABILITIES_API void ExecuteEffects(const FGameplayCueNotify_SpawnContext& SpawnContext, FGameplayCueNotify_SpawnResult& OutSpawnResult) const;
 
+	GAMEPLAYABILITIES_API void ValidateAssociatedAssets(const UObject* ContainingAsset, const FString& Context, class FDataValidationContext& ValidationContext) const;
+	
 protected:
 
 	// Particle systems to be spawned on gameplay cue execution.  These should never use looping effects!
@@ -745,18 +730,17 @@ protected:
  *	Set of looping effects to spawn for looping gameplay cues.
  */
 USTRUCT(BlueprintType)
-struct FGameplayCueNotify_LoopingEffects
+struct FGameplayCueNotify_LoopingEffects final
 {
 	GENERATED_BODY()
 
-public:
-
-	FGameplayCueNotify_LoopingEffects();
+	GAMEPLAYABILITIES_API FGameplayCueNotify_LoopingEffects();
 
-	void StartEffects(const FGameplayCueNotify_SpawnContext& SpawnContext, FGameplayCueNotify_SpawnResult& OutSpawnResult) const;
-	void StopEffects(FGameplayCueNotify_SpawnResult& SpawnResult) const;
+	GAMEPLAYABILITIES_API void StartEffects(const FGameplayCueNotify_SpawnContext& SpawnContext, FGameplayCueNotify_SpawnResult& 
+	OutSpawnResult) const;
+	GAMEPLAYABILITIES_API void StopEffects(FGameplayCueNotify_SpawnResult& SpawnResult) const;
 	
-	void ValidateAssociatedAssets(const UObject* ContainingAsset, const FString& Context, class FDataValidationContext& ValidationContext) const;
+	GAMEPLAYABILITIES_API void ValidateAssociatedAssets(const UObject* ContainingAsset, const FString& Context, class FDataValidationContext& ValidationContext) const;
 
 protected:
 
