Fix for UGameplayEffect::PostLoad potentially causing an attempt at DataRegistry loads prematurely.

diff --git a/Engine/Plugins/Runtime/GameplayAbilities/Source/GameplayAbilities/Private/GameplayEffect.cpp b/Engine/Plugins/Runtime/GameplayAbilities/Source/GameplayAbilities/Private/GameplayEffect.cpp
--- a/Engine/Plugins/Runtime/GameplayAbilities/Source/GameplayAbilities/Private/GameplayEffect.cpp
+++ b/Engine/Plugins/Runtime/GameplayAbilities/Source/GameplayAbilities/Private/GameplayEffect.cpp
@@ -55,16 +55,6 @@ const float UGameplayEffect::INVALID_LEVEL = FGameplayEffectConstants::INVALID_L
 
 DECLARE_CYCLE_STAT(TEXT("MakeQuery"), STAT_MakeGameplayEffectQuery, STATGROUP_AbilitySystem);
 
-#if WITH_EDITOR
-#define SCALABLEFLOAT_REPORTERROR_WITHPOSTLOAD(Scalable) \
-	if (Scalable.Curve.CurveTable) const_cast<UCurveTable*>(ToRawPtr(Scalable.Curve.CurveTable))->ConditionalPostLoad(); \
-	SCALABLEFLOAT_REPORTERROR(Scalable);
-
-#define SCALABLEFLOAT_REPORTERROR_WITHPATHNAME_WITHPOSTLOAD(Scalable, PathNameString) \
-	if (Scalable.Curve.CurveTable) const_cast<UCurveTable*>(ToRawPtr(Scalable.Curve.CurveTable))->ConditionalPostLoad(); \
-	SCALABLEFLOAT_REPORTERROR_WITHPATHNAME(Scalable, PathNameString);
-#endif // WITH_EDITOR
-
 namespace UE::GameplayEffect
 {
 	enum class EActiveGameplayEffectFix : int32
@@ -202,8 +192,8 @@ void UGameplayEffect::PostInitProperties()
 			}
 		}
 
-		// Catch any GEComponent configuration issues.  We usually want to do this late (after deserialization) for assets in PostLoad, but
-		// we've already determined this is a native class, so let's do that now (we won't get PostLoad).
+		// Catch any GEComponent configuration issues.  We usually want to do this late (after deserialization or BP compilation), but
+		// we've already determined this is a native class, so let's do that now (we won't get PostLoad or BP compilation).
 		FDataValidationContext DataValidationContext;
 		IsDataValid(DataValidationContext);
 
@@ -271,32 +261,33 @@ void UGameplayEffect::GetBlockedAbilityTags(FGameplayTagContainer& OutTagContain
 #if WITH_EDITOR
 EDataValidationResult UGameplayEffect::IsDataValid(FDataValidationContext& Context) const
 {
+	EDataValidationResult ValidationResult = Super::IsDataValid(Context);
+
 	// Report these errors here & now, even though they're not using the DataValidationContext because
 	// these will show-up when saving/presubmit rather than during more extensive testing done at cook time.
-	DurationMagnitude.ReportErrors(GetPathName());
-	for (const FGameplayModifierInfo& CurModInfo : Modifiers)
+	FString AssetPathName = GetPathName();
+
+	if (DurationPolicy == EGameplayEffectDurationType::HasDuration)
 	{
-		CurModInfo.ModifierMagnitude.ReportErrors(GetPathName());
+		ValidationResult = CombineDataValidationResults(ValidationResult, DurationMagnitude.IsDataValid(Context, AssetPathName + TEXT(".DurationMagnitude")));
 	}
 
-	EDataValidationResult ValidationResult = Super::IsDataValid(Context);
+	ValidationResult = CombineDataValidationResults(ValidationResult, Period.IsDataValid(Context, AssetPathName + TEXT(".Period")));
 
-	if (ValidationResult != EDataValidationResult::Invalid)
+	for (const FGameplayModifierInfo& CurModInfo : Modifiers)
 	{
-		for (const UGameplayEffectComponent* GEComponent : GEComponents)
+		ValidationResult = CombineDataValidationResults(ValidationResult, CurModInfo.ModifierMagnitude.IsDataValid(Context, AssetPathName + TEXT(".") + CurModInfo.Attribute.GetName()));
+	}
+
+	for (const UGameplayEffectComponent* GEComponent : GEComponents)
+	{
+		if (GEComponent)
 		{
-			if (GEComponent)
-			{
-				ValidationResult = GEComponent->IsDataValid(Context);
-				if (ValidationResult == EDataValidationResult::Invalid)
-				{
-					break;
-				}
-			}
-			else
-			{
-				Context.AddWarning(LOCTEXT("GEIsNull", "Null entry in GEComponents"));
-			}
+			ValidationResult = CombineDataValidationResults(ValidationResult, GEComponent->IsDataValid(Context));
+		}
+		else
+		{
+			Context.AddWarning(LOCTEXT("GEIsNull", "Null entry in GEComponents"));
 		}
 	}
 
@@ -327,16 +318,10 @@ void UGameplayEffect::PostLoad()
 	OnGameplayEffectChanged();
 
 #if WITH_EDITOR
-	SCALABLEFLOAT_REPORTERROR_WITHPOSTLOAD(Period);
-
 	PRAGMA_DISABLE_DEPRECATION_WARNINGS
 	HasRemoveGameplayEffectsQuery = !RemoveGameplayEffectQuery.IsEmpty();
 	PRAGMA_ENABLE_DEPRECATION_WARNINGS
 
-	// Update the EditorStatusText
-	FDataValidationContext DataValidationContext;
-	IsDataValid(DataValidationContext);
-
 	// We're done loading (and therefore upgrading), boost the version.
 	SetVersion( static_cast<EGameplayEffectVersion>(UE::GameplayEffect::CVarGameplayEffectMaxVersion.GetValueOnGameThread()) );
 #endif // WITH_EDITOR
@@ -1060,6 +1045,17 @@ bool FAttributeBasedFloat::operator!=(const FAttributeBasedFloat& Other) const
 	return !(*this == Other);
 }
 
+#if WITH_EDITOR
+EDataValidationResult FAttributeBasedFloat::IsDataValid(FDataValidationContext& Context, const FString& PathName) const
+{
+	EDataValidationResult Result = Coefficient.IsDataValid(Context, FString::Printf(TEXT("%s.%s"), *PathName, TEXT("Coefficient")));
+	Result = CombineDataValidationResults(Result, PreMultiplyAdditiveValue.IsDataValid(Context, FString::Printf(TEXT("%s.%s"), *PathName, TEXT("PreMultiplyAdditiveValue"))));
+	Result = CombineDataValidationResults(Result, PostMultiplyAdditiveValue.IsDataValid(Context, FString::Printf(TEXT("%s.%s"), *PathName, TEXT("PostMultiplyAdditiveValue"))));
+
+	return Result;
+}
+#endif
+
 // --------------------------------------------------------------------------------------------------------------------------------------------------------
 //
 //	FCustomCalculationBasedFloat
@@ -1108,6 +1104,17 @@ bool FCustomCalculationBasedFloat::operator!=(const FCustomCalculationBasedFloat
 	return !(*this == Other);
 }
 
+#if WITH_EDITOR
+EDataValidationResult FCustomCalculationBasedFloat::IsDataValid(FDataValidationContext& Context, const FString& PathName) const
+{
+	EDataValidationResult Result = Coefficient.IsDataValid(Context, FString::Printf(TEXT("%s.%s"), *PathName, TEXT("Coefficient")));
+	Result = CombineDataValidationResults(Result, PreMultiplyAdditiveValue.IsDataValid(Context, FString::Printf(TEXT("%s.%s"), *PathName, TEXT("PreMultiplyAdditiveValue"))));
+	Result = CombineDataValidationResults(Result, PostMultiplyAdditiveValue.IsDataValid(Context, FString::Printf(TEXT("%s.%s"), *PathName, TEXT("PostMultiplyAdditiveValue"))));
+
+	return Result;
+}
+#endif
+
 // --------------------------------------------------------------------------------------------------------------------------------------------------------
 //
 //	FGameplayEffectMagnitude
@@ -1361,23 +1368,26 @@ FText FGameplayEffectModifierMagnitude::GetValueForEditorDisplay() const
 }
 
 void FGameplayEffectModifierMagnitude::ReportErrors(const FString& PathName) const
+{
+	// Deprecated, use IsDataValid
+}
+
+EDataValidationResult FGameplayEffectModifierMagnitude::IsDataValid(FDataValidationContext& Context, const FString& PathName) const
 {
 	if (MagnitudeCalculationType == EGameplayEffectMagnitudeCalculation::ScalableFloat)
 	{
-		SCALABLEFLOAT_REPORTERROR_WITHPATHNAME_WITHPOSTLOAD(ScalableFloatMagnitude, PathName);
+		return ScalableFloatMagnitude.IsDataValid(Context, FString::Printf(TEXT("%s.%s"), *PathName, TEXT("ScalableFloatMagnitude")));
 	}
 	else if (MagnitudeCalculationType == EGameplayEffectMagnitudeCalculation::AttributeBased)
 	{
-		SCALABLEFLOAT_REPORTERROR_WITHPATHNAME_WITHPOSTLOAD(AttributeBasedMagnitude.Coefficient, PathName);
-		SCALABLEFLOAT_REPORTERROR_WITHPATHNAME_WITHPOSTLOAD(AttributeBasedMagnitude.PreMultiplyAdditiveValue, PathName);
-		SCALABLEFLOAT_REPORTERROR_WITHPATHNAME_WITHPOSTLOAD(AttributeBasedMagnitude.PostMultiplyAdditiveValue, PathName);
+		return AttributeBasedMagnitude.IsDataValid(Context, FString::Printf(TEXT("%s.%s"), *PathName, TEXT("AttributeBasedMagnitude")));
 	}
 	else if (MagnitudeCalculationType == EGameplayEffectMagnitudeCalculation::CustomCalculationClass)
 	{
-		SCALABLEFLOAT_REPORTERROR_WITHPATHNAME_WITHPOSTLOAD(CustomMagnitude.Coefficient, PathName);
-		SCALABLEFLOAT_REPORTERROR_WITHPATHNAME_WITHPOSTLOAD(CustomMagnitude.PreMultiplyAdditiveValue, PathName);
-		SCALABLEFLOAT_REPORTERROR_WITHPATHNAME_WITHPOSTLOAD(CustomMagnitude.PostMultiplyAdditiveValue, PathName);
+		return CustomMagnitude.IsDataValid(Context, FString::Printf(TEXT("%s.%s"), *PathName, TEXT("CustomMagnitude")));
 	}
+
+	return EDataValidationResult::Invalid;
 }
 #endif // WITH_EDITOR
 
diff --git a/Engine/Plugins/Runtime/GameplayAbilities/Source/GameplayAbilities/Private/ScalableFloat.cpp b/Engine/Plugins/Runtime/GameplayAbilities/Source/GameplayAbilities/Private/ScalableFloat.cpp
--- a/Engine/Plugins/Runtime/GameplayAbilities/Source/GameplayAbilities/Private/ScalableFloat.cpp
+++ b/Engine/Plugins/Runtime/GameplayAbilities/Source/GameplayAbilities/Private/ScalableFloat.cpp
@@ -13,8 +13,65 @@
 
 #include UE_INLINE_GENERATED_CPP_BY_NAME(ScalableFloat)
 
+#define LOCTEXT_NAMESPACE "ScalableFloat"
+
 #if WITH_EDITOR
 #include "EditorReimportHandler.h"
+#include "Misc/DataValidation.h"
+
+EDataValidationResult FScalableFloat::IsDataValid(FDataValidationContext& Context, const FString& PathName) const
+{
+	if (Curve.CurveTable)
+	{
+		if (Curve.RowName.IsNone())
+		{
+			Context.AddError(FText::Format(LOCTEXT("RowNameIsNone", "{0}: CurveTable Row is None"), FText::FromString(PathName)));
+			return EDataValidationResult::Invalid;
+		}
+
+		if (!Curve.CurveTable->FindCurve(Curve.RowName, PathName, false))
+		{
+			Context.AddError(FText::Format(LOCTEXT("RowNameDoesNotExist", "{0}: RowName {1} does not exist in CurveTable {2}"),
+				FText::FromString(PathName), FText::FromString(Curve.RowName.ToString()), FText::FromString(Curve.CurveTable.GetName())));
+			return EDataValidationResult::Invalid;
+		}
+	}
+	else if (RegistryType.IsValid())
+	{
+		if (Curve.RowName.IsNone())
+		{
+			Context.AddError(FText::Format(LOCTEXT("RegistryRowNameIsNone", "{0}: Registry CurveTable Row is None"), FText::FromString(PathName)));
+			return EDataValidationResult::Invalid;
+		}
+
+		const UDataRegistrySubsystem* SubSystem = UDataRegistrySubsystem::Get();
+		if (!SubSystem->IsConfigEnabled(false) || !SubSystem->AreRegistriesInitialized())
+		{
+			// If we're given a use case, it's more likely to be a more thorough IsDataValid, so let them know why this failed to validate.
+			if (Context.GetValidationUsecase() != EDataValidationUsecase::None)
+			{
+				Context.AddWarning(FText::Format(LOCTEXT("DataRegistriesNotLoaded", "{0}: DataRegistries were not loaded in time for validation"),
+					FText::FromString(PathName)));
+			}
+
+			// We are not ready to evaluate this until after load
+			return EDataValidationResult::NotValidated;
+		}
+
+		float OutValue = 0.0f;
+		const FRealCurve* OutCurve = nullptr;
+		const FDataRegistryId DataRegistryId(RegistryType, Curve.RowName);
+		const FDataRegistryCacheGetResult GetCurveResult = SubSystem->EvaluateCachedCurve(OutValue, OutCurve, DataRegistryId, OutValue);
+		if (!GetCurveResult)
+		{
+			Context.AddError(FText::Format(LOCTEXT("DataRegistryCurveDoesNotExist", "{0}: Specified DataRegistry Curve '{1}' does not exist"),
+				FText::FromString(PathName), DataRegistryId.ToText()));
+			return EDataValidationResult::Invalid;
+		}
+	}
+
+	return EDataValidationResult::Valid;
+}
 #endif
 
 bool FScalableFloat::EvaluateCurveAtLevel(float& OutValue, const FRealCurve*& OutCurve, float Level, const FString& ContextString, bool bWarnIfInvalid) const
@@ -178,10 +235,7 @@ bool FScalableFloat::IsValid() const
 
 bool FScalableFloat::IsValidWithWarnings(const FString& ContextString) const
 {
-	float OutFloat;
-	const FRealCurve* FoundCurve;
-
-	return EvaluateCurveAtLevel(OutFloat, FoundCurve, 0.f, ContextString, true);
+	return true;
 }
 
 bool FScalableFloat::SerializeFromMismatchedTag(const FPropertyTag& Tag, FStructuredArchive::FSlot Slot)
@@ -485,3 +539,4 @@ FAutoConsoleCommand FindCoefficientScalableFloatsCommand(
 
 #endif
 
+#undef LOCTEXT_NAMESPACE
diff --git a/Engine/Plugins/Runtime/GameplayAbilities/Source/GameplayAbilities/Public/GameplayEffect.h b/Engine/Plugins/Runtime/GameplayAbilities/Source/GameplayAbilities/Public/GameplayEffect.h
--- a/Engine/Plugins/Runtime/GameplayAbilities/Source/GameplayAbilities/Public/GameplayEffect.h
+++ b/Engine/Plugins/Runtime/GameplayAbilities/Source/GameplayAbilities/Public/GameplayEffect.h
@@ -186,6 +186,10 @@ public:
 	/** Equality/Inequality operators */
 	bool operator==(const FAttributeBasedFloat& Other) const;
 	bool operator!=(const FAttributeBasedFloat& Other) const;
+
+#if WITH_EDITOR
+	EDataValidationResult IsDataValid(class FDataValidationContext& Context, const FString& PathName) const;
+#endif
 };
 
 /** Structure to encapsulate magnitudes that are calculated via custom calculation */
@@ -237,6 +241,10 @@ public:
 	/** Equality/Inequality operators */
 	bool operator==(const FCustomCalculationBasedFloat& Other) const;
 	bool operator!=(const FCustomCalculationBasedFloat& Other) const;
+
+#if WITH_EDITOR
+	EDataValidationResult IsDataValid(class FDataValidationContext& Context, const FString& PathName) const;
+#endif
 };
 
 /** Struct for holding SetBytCaller data */
@@ -349,6 +357,9 @@ public:
 
 #if WITH_EDITOR
 	FText GetValueForEditorDisplay() const;
+	EDataValidationResult IsDataValid(class FDataValidationContext& Context, const FString& PathName) const;
+
+	UE_DEPRECATED(5.6, "Use IsDataValid")
 	void ReportErrors(const FString& PathName) const;
 #endif
 
diff --git a/Engine/Plugins/Runtime/GameplayAbilities/Source/GameplayAbilities/Public/ScalableFloat.h b/Engine/Plugins/Runtime/GameplayAbilities/Source/GameplayAbilities/Public/ScalableFloat.h
--- a/Engine/Plugins/Runtime/GameplayAbilities/Source/GameplayAbilities/Public/ScalableFloat.h
+++ b/Engine/Plugins/Runtime/GameplayAbilities/Source/GameplayAbilities/Public/ScalableFloat.h
@@ -90,8 +90,14 @@ public:
 	bool IsValid() const;
 
 	/** Error checking: Checks if float is valid and prints detailed warnings if not valid */
+	UE_DEPRECATED(5.6, "Use IsDataValid and manually output the errors")
 	bool IsValidWithWarnings(const FString& ContextString) const;
 
+#if WITH_EDITOR
+	/** Error checking must be performed by the owning class, then call this to get specific errors */
+	EDataValidationResult IsDataValid(class FDataValidationContext& Context, const FString& PathName) const;
+#endif
+
 	/** Equality/Inequality operators */
 	bool operator==(const FScalableFloat& Other) const;
 	bool operator!=(const FScalableFloat& Other) const;
