Validating index to account for passed in negative values.
#11048

diff --git a/Engine/Plugins/Runtime/GameplayAbilities/Source/GameplayAbilities/Public/Abilities/GameplayAbilityTargetTypes.h b/Engine/Plugins/Runtime/GameplayAbilities/Source/GameplayAbilities/Public/Abilities/GameplayAbilityTargetTypes.h
--- a/Engine/Plugins/Runtime/GameplayAbilities/Source/GameplayAbilities/Public/Abilities/GameplayAbilityTargetTypes.h
+++ b/Engine/Plugins/Runtime/GameplayAbilities/Source/GameplayAbilities/Public/Abilities/GameplayAbilityTargetTypes.h
@@ -224,7 +224,7 @@ struct GAMEPLAYABILITIES_API FGameplayAbilityTargetDataHandle
 	/** Returns true if there are any valid targets */
 	bool IsValid(int32 Index) const
 	{
-		return (Index < Data.Num() && Data[Index].IsValid());
+		return (Data.IsValidIndex(Index) && Data[Index].IsValid());
 	}
 
 	/** Returns data at index, or nullptr if invalid */
