UE: Gameplay Effect MagnitudeCalculationType SetByCaller is now consdered valid data. Any other types that don't have data validation logic are also valid

diff --git a/Engine/Plugins/Runtime/GameplayAbilities/Source/GameplayAbilities/Private/GameplayEffect.cpp b/Engine/Plugins/Runtime/GameplayAbilities/Source/GameplayAbilities/Private/GameplayEffect.cpp
--- a/Engine/Plugins/Runtime/GameplayAbilities/Source/GameplayAbilities/Private/GameplayEffect.cpp
+++ b/Engine/Plugins/Runtime/GameplayAbilities/Source/GameplayAbilities/Private/GameplayEffect.cpp
@@ -1387,7 +1387,7 @@ EDataValidationResult FGameplayEffectModifierMagnitude::IsDataValid(FDataValidat
 		return CustomMagnitude.IsDataValid(Context, FString::Printf(TEXT("%s.%s"), *PathName, TEXT("CustomMagnitude")));
 	}
 
-	return EDataValidationResult::Invalid;
+	return EDataValidationResult::Valid;
 }
 #endif // WITH_EDITOR
 
