UE: Disable scalable float data registry validation for now because not all data registries are added in commandlets, and commandlets can run data validation

diff --git a/Engine/Plugins/Runtime/GameplayAbilities/Source/GameplayAbilities/Private/ScalableFloat.cpp b/Engine/Plugins/Runtime/GameplayAbilities/Source/GameplayAbilities/Private/ScalableFloat.cpp
--- a/Engine/Plugins/Runtime/GameplayAbilities/Source/GameplayAbilities/Private/ScalableFloat.cpp
+++ b/Engine/Plugins/Runtime/GameplayAbilities/Source/GameplayAbilities/Private/ScalableFloat.cpp
@@ -58,6 +58,8 @@ EDataValidationResult FScalableFloat::IsDataValid(FDataValidationContext& Contex
 			return EDataValidationResult::NotValidated;
 		}
 
+		// Disabling data registry validation for now since not all data registries are added in all cases where validation may occur (specifically in commandlets)
+		/*
 		float OutValue = 0.0f;
 		const FRealCurve* OutCurve = nullptr;
 		const FDataRegistryId DataRegistryId(RegistryType, Curve.RowName);
@@ -68,6 +70,7 @@ EDataValidationResult FScalableFloat::IsDataValid(FDataValidationContext& Contex
 				FText::FromString(PathName), DataRegistryId.ToText()));
 			return EDataValidationResult::Invalid;
 		}
+		*/
 	}
 
 	return EDataValidationResult::Valid;
