Moving Modifier Magnitude stack calculation into the GE so that stacking can be changed per effect.
#12117

diff --git a/Engine/Plugins/Runtime/GameplayAbilities/Source/GameplayAbilities/Private/GameplayEffect.cpp b/Engine/Plugins/Runtime/GameplayAbilities/Source/GameplayAbilities/Private/GameplayEffect.cpp
--- a/Engine/Plugins/Runtime/GameplayAbilities/Source/GameplayAbilities/Private/GameplayEffect.cpp
+++ b/Engine/Plugins/Runtime/GameplayAbilities/Source/GameplayAbilities/Private/GameplayEffect.cpp
@@ -146,6 +146,7 @@ UGameplayEffect::UGameplayEffect(const FObjectInitializer& ObjectInitializer)
 	StackDurationRefreshPolicy = EGameplayEffectStackingDurationPolicy::RefreshOnSuccessfulApplication;
 	StackPeriodResetPolicy = EGameplayEffectStackingPeriodPolicy::ResetOnSuccessfulApplication;
 	bRequireModifierSuccessToTriggerCues = true;
+	bFactorInStackCount = true;
 	PRAGMA_ENABLE_DEPRECATION_WARNINGS
 }
 
@@ -1924,6 +1925,21 @@ float FGameplayEffectSpec::GetModifierMagnitude(int32 ModifierIdx, bool bFactorI
 	return ModMagnitude;
 }
 
+float FGameplayEffectSpec::GetModifierMagnitude(int32 ModifierIdx) const
+{
+	check(Modifiers.IsValidIndex(ModifierIdx) && Def && Def->Modifiers.IsValidIndex(ModifierIdx));
+
+	const float SingleEvaluatedMagnitude = Modifiers[ModifierIdx].GetEvaluatedMagnitude();
+
+	float ModMagnitude = SingleEvaluatedMagnitude;
+	if (Def->bFactorInStackCount)
+	{
+		ModMagnitude = GameplayEffectUtilities::ComputeStackedModifierMagnitude(SingleEvaluatedMagnitude, GetStackCount(), Def->Modifiers[ModifierIdx].ModifierOp);
+	}
+
+	return ModMagnitude;
+}
+
 void FGameplayEffectSpec::CalculateModifierMagnitudes()
 {
 	for(int32 ModIdx = 0; ModIdx < Modifiers.Num(); ++ModIdx)
@@ -2958,7 +2974,7 @@ void FActiveGameplayEffectsContainer::PredictivelyExecuteEffectSpec(FGameplayEff
 			}
 		}
 
-		FGameplayModifierEvaluatedData EvalData(ModDef.Attribute, ModDef.ModifierOp, SpecToUse.GetModifierMagnitude(ModIdx, true));
+		FGameplayModifierEvaluatedData EvalData(ModDef.Attribute, ModDef.ModifierOp, SpecToUse.GetModifierMagnitude(ModIdx));
 		ModifierSuccessfullyExecuted |= InternalExecuteMod(SpecToUse, EvalData);
 	}
 
@@ -3089,7 +3105,7 @@ void FActiveGameplayEffectsContainer::ExecuteActiveEffectsFrom(FGameplayEffectSp
 			}
 		}
 		
-		FGameplayModifierEvaluatedData EvalData(ModDef.Attribute, ModDef.ModifierOp, SpecToUse.GetModifierMagnitude(ModIdx, true));
+		FGameplayModifierEvaluatedData EvalData(ModDef.Attribute, ModDef.ModifierOp, SpecToUse.GetModifierMagnitude(ModIdx));
 		ModifierSuccessfullyExecuted |= InternalExecuteMod(SpecToUse, EvalData);
 	}
 
@@ -4321,7 +4337,7 @@ void FActiveGameplayEffectsContainer::AddActiveGameplayEffectGrantedTagsAndModif
 			// Note we assume the EvaluatedMagnitude is up to do. There is no case currently where we should recalculate magnitude based on
 			// Ongoing tags being met. We either calculate magnitude one time, or its done via OnDirty calls (or potentially a frequency timer one day)
 
-			float EvaluatedMagnitude = Effect.Spec.GetModifierMagnitude(ModIdx, true);	// Note this could cause an attribute aggregator to be created, so must do this before calling/caching the Aggregator below!
+			float EvaluatedMagnitude = Effect.Spec.GetModifierMagnitude(ModIdx);	// Note this could cause an attribute aggregator to be created, so must do this before calling/caching the Aggregator below!
 
 			FAggregator* Aggregator = FindOrCreateAttributeAggregator(Effect.Spec.Def->Modifiers[ModIdx].Attribute).Get();
 			if (ensure(Aggregator))
diff --git a/Engine/Plugins/Runtime/GameplayAbilities/Source/GameplayAbilities/Private/GameplayEffectAggregator.cpp b/Engine/Plugins/Runtime/GameplayAbilities/Source/GameplayAbilities/Private/GameplayEffectAggregator.cpp
--- a/Engine/Plugins/Runtime/GameplayAbilities/Source/GameplayAbilities/Private/GameplayEffectAggregator.cpp
+++ b/Engine/Plugins/Runtime/GameplayAbilities/Source/GameplayAbilities/Private/GameplayEffectAggregator.cpp
@@ -512,7 +512,7 @@ void FAggregator::UpdateAggregatorMod(FActiveGameplayEffectHandle ActiveHandle,
 		if (ModDef.Attribute == Attribute)
 		{
 			FAggregatorModChannel& ModChannel = ModChannels.FindOrAddModChannel(ModDef.EvaluationChannelSettings.GetEvaluationChannel());
-			ModChannel.AddMod(Spec.GetModifierMagnitude(ModIdx, true), ModDef.ModifierOp, &ModDef.SourceTags, &ModDef.TargetTags, bWasLocallyGenerated, InHandle);
+			ModChannel.AddMod(Spec.GetModifierMagnitude(ModIdx), ModDef.ModifierOp, &ModDef.SourceTags, &ModDef.TargetTags, bWasLocallyGenerated, InHandle);
 		}
 	}
 
diff --git a/Engine/Plugins/Runtime/GameplayAbilities/Source/GameplayAbilities/Public/GameplayEffect.h b/Engine/Plugins/Runtime/GameplayAbilities/Source/GameplayAbilities/Public/GameplayEffect.h
--- a/Engine/Plugins/Runtime/GameplayAbilities/Source/GameplayAbilities/Public/GameplayEffect.h
+++ b/Engine/Plugins/Runtime/GameplayAbilities/Source/GameplayAbilities/Public/GameplayEffect.h
@@ -1109,13 +1109,23 @@ struct GAMEPLAYABILITIES_API FGameplayEffectSpec
 	/**
 	 * Get the computed magnitude of the modifier on the spec with the specified index
 	 * 
-	 * @param ModifierIndx			Modifier to get
+	 * @param ModifierIdx			Modifier to get
 	 * @param bFactorInStackCount	If true, the calculation will include the stack count
 	 * 
 	 * @return Computed magnitude
 	 */
+	UE_DEPRECATED(5.6, "Versions using bFactorInStackCount are deprecated, please use GetModifierMagnitude(int32 ModifierIdx) instead.")
 	float GetModifierMagnitude(int32 ModifierIdx, bool bFactorInStackCount) const;
 
+	/**
+	 * Get the computed magnitude of the modifier on the spec with the specified index
+	 * 
+	 * @param ModifierIdx			Modifier to get
+	 * 
+	 * @return Computed magnitude
+	 */
+	float GetModifierMagnitude(int32 ModifierIdx) const;
+
 	/** Fills out the modifier magnitudes inside the Modifier Specs */
 	void CalculateModifierMagnitudes();
 
@@ -2375,6 +2385,10 @@ public:
 	UPROPERTY(EditDefaultsOnly, BlueprintReadOnly, Category = Stacking, meta = (EditConditionHides, EditCondition = "StackingType != EGameplayEffectStackingType::None"))
 	EGameplayEffectStackingExpirationPolicy StackExpirationPolicy;
 
+	/** If true, the calculation will include the stack count for Modifier Magnitudes */
+	UPROPERTY(EditDefaultsOnly, BlueprintReadOnly, Category = Stacking, meta = (EditConditionHides, EditCondition = "StackingType != EGameplayEffectStackingType::None"))
+	bool bFactorInStackCount;
+
 	// ----------------------------------------------------------------------
 	//	Granted abilities
 	// ----------------------------------------------------------------------
