Fix crash from ClientActivateAbilitySucceedWithEventData
- when new owner of ASC is dead, ASC will auto clear out OwnerActor and AbilityActorInfo; while the client RPC may trigger after that and crashed with OwnerActor always valid assumption

diff --git a/Engine/Plugins/Runtime/GameplayAbilities/Source/GameplayAbilities/Private/AbilitySystemComponent_Abilities.cpp b/Engine/Plugins/Runtime/GameplayAbilities/Source/GameplayAbilities/Private/AbilitySystemComponent_Abilities.cpp
--- a/Engine/Plugins/Runtime/GameplayAbilities/Source/GameplayAbilities/Private/AbilitySystemComponent_Abilities.cpp
+++ b/Engine/Plugins/Runtime/GameplayAbilities/Source/GameplayAbilities/Private/AbilitySystemComponent_Abilities.cpp
@@ -2357,6 +2357,13 @@ void UAbilitySystemComponent::ClientActivateAbilitySucceedWithEventData_Implemen
 
 	check(AbilityToActivate);
 	ensure(AbilityActorInfo.IsValid());
+	
+	if (!AbilityActorInfo->OwnerActor.IsValid())
+	{
+		// When the owner is explicitly switched and then gets destroyed, the OwnerActor could become invalid
+		UE_LOG(LogAbilitySystem, Verbose, TEXT("%s: server confirmed [%s] %s (Prediction key %s), but failed with invalid OwnerActor"), *GetNameSafe(GetOwner()), *Handle.ToString(), *GetNameSafe(AbilityToActivate), *PredictionKey.ToString());
+		return;
+	}
 
 	UE_LOG(LogAbilitySystem, Verbose, TEXT("%s: Server Confirmed [%s] %s. PredictionKey: %s"), *GetNameSafe(GetOwner()), *Handle.ToString(), *GetNameSafe(AbilityToActivate), *PredictionKey.ToString());
 	UE_VLOG(GetOwner(), VLogAbilitySystem, Verbose, TEXT("Server Confirmed [%s] %s. PredictionKey: %s"), *Handle.ToString(), *GetNameSafe(AbilityToActivate), *PredictionKey.ToString());
diff --git a/Engine/Plugins/Runtime/GameplayAbilities/Source/GameplayAbilities/Private/GameplayAbilityTypes.cpp b/Engine/Plugins/Runtime/GameplayAbilities/Source/GameplayAbilities/Private/GameplayAbilityTypes.cpp
--- a/Engine/Plugins/Runtime/GameplayAbilities/Source/GameplayAbilities/Private/GameplayAbilityTypes.cpp
+++ b/Engine/Plugins/Runtime/GameplayAbilities/Source/GameplayAbilities/Private/GameplayAbilityTypes.cpp
@@ -155,7 +155,7 @@ FGameplayAbilityActivationInfo::FGameplayAbilityActivationInfo(AActor* InActor)
 {
 	// On Init, we are either Authority or NonAuthority. We haven't been given a PredictionKey and we haven't been confirmed.
 	// NonAuthority essentially means 'I'm not sure what how I'm going to do this yet'.
-	ActivationMode = (InActor->GetLocalRole() == ROLE_Authority ? EGameplayAbilityActivationMode::Authority : EGameplayAbilityActivationMode::NonAuthority);
+	ActivationMode = ((InActor && InActor->GetLocalRole() == ROLE_Authority) ? EGameplayAbilityActivationMode::Authority : EGameplayAbilityActivationMode::NonAuthority);
 }
 
 void FGameplayAbilityActivationInfo::SetPredicting(FPredictionKey PredictionKey)
