Adding editor only mutable functions for Gameplay Ability Asset tags.

diff --git a/Engine/Plugins/Runtime/GameplayAbilities/Source/GameplayAbilities/Private/Abilities/GameplayAbility.cpp b/Engine/Plugins/Runtime/GameplayAbilities/Source/GameplayAbilities/Private/Abilities/GameplayAbility.cpp
--- a/Engine/Plugins/Runtime/GameplayAbilities/Source/GameplayAbilities/Private/Abilities/GameplayAbility.cpp
+++ b/Engine/Plugins/Runtime/GameplayAbilities/Source/GameplayAbilities/Private/Abilities/GameplayAbility.cpp
@@ -1684,8 +1684,15 @@ void UGameplayAbility::AddAbilityTaskDebugMessage(UGameplayTask* AbilityTask, FS
 	Msg.Message = FString::Printf(TEXT("{%s} %s"), AbilityTask ? *AbilityTask->GetDebugString() : TEXT(""), *DebugMessage);
 }
 
+#if WITH_EDITOR
+FGameplayTagContainer& UGameplayAbility::EditorGetAssetTags()
+{
+	return AbilityTags;
+}
+#endif //WITH_EDITOR
+
 /**
- *	Helper methods for adding GaAmeplayCues without having to go through GameplayEffects.
+ *	Helper methods for adding GameplayCues without having to go through GameplayEffects.
  *	For now, none of these will happen predictively. We can eventually build this out more to 
  *	work with the PredictionKey system.
  */
diff --git a/Engine/Plugins/Runtime/GameplayAbilities/Source/GameplayAbilities/Public/Abilities/GameplayAbility.h b/Engine/Plugins/Runtime/GameplayAbilities/Source/GameplayAbilities/Public/Abilities/GameplayAbility.h
--- a/Engine/Plugins/Runtime/GameplayAbilities/Source/GameplayAbilities/Public/Abilities/GameplayAbility.h
+++ b/Engine/Plugins/Runtime/GameplayAbilities/Source/GameplayAbilities/Public/Abilities/GameplayAbility.h
@@ -482,6 +482,11 @@ public:
 	/** Adds a debug message to display to the user */
 	void AddAbilityTaskDebugMessage(UGameplayTask* AbilityTask, FString DebugMessage);
 
+#if WITH_EDITOR
+	/** Allow modification of the AssetTags (AbilityTags) while in editor */
+	FGameplayTagContainer& EditorGetAssetTags();
+#endif // WITH_EDITOR
+
 	// --------------------------------------
 	//	Public variables, exposed for backwards compatibility
 	// --------------------------------------
