Checking the attribute Get is valid instead of the AttributeOwner so that changes work immediately, not only on load.

diff --git a/Engine/Plugins/Runtime/GameplayAbilities/Source/GameplayAbilities/Public/AttributeSet.h b/Engine/Plugins/Runtime/GameplayAbilities/Source/GameplayAbilities/Public/AttributeSet.h
--- a/Engine/Plugins/Runtime/GameplayAbilities/Source/GameplayAbilities/Public/AttributeSet.h
+++ b/Engine/Plugins/Runtime/GameplayAbilities/Source/GameplayAbilities/Public/AttributeSet.h
@@ -66,9 +66,10 @@ struct GAMEPLAYABILITIES_API FGameplayAttribute
 
 	FGameplayAttribute(FProperty *NewProperty);
 
+	/** Returns true if our Attribute has a valid resolved Field */
 	bool IsValid() const
 	{
-		return Attribute != nullptr && AttributeOwner != nullptr;
+		return Attribute.Get() != nullptr;
 	}
 
 	/** Set up from a FProperty inside a set */
