Force GetPeriod() to return 0 if the policy is Instant.

diff --git a/Engine/Plugins/Runtime/GameplayAbilities/Source/GameplayAbilities/Private/GameplayEffect.cpp b/Engine/Plugins/Runtime/GameplayAbilities/Source/GameplayAbilities/Private/GameplayEffect.cpp
--- a/Engine/Plugins/Runtime/GameplayAbilities/Source/GameplayAbilities/Private/GameplayEffect.cpp
+++ b/Engine/Plugins/Runtime/GameplayAbilities/Source/GameplayAbilities/Private/GameplayEffect.cpp
@@ -1894,6 +1894,11 @@ const FGameplayTagContainer& FGameplayEffectSpec::GetDynamicAssetTags() const
 
 float FGameplayEffectSpec::GetPeriod() const
 {
+	if (Def && Def->DurationPolicy == EGameplayEffectDurationType::Instant)
+	{
+		return UGameplayEffect::NO_PERIOD;
+	}
+	
 	return Period;
 }
 
diff --git a/Engine/Plugins/Runtime/GameplayAbilities/Source/GameplayAbilities/Public/GameplayEffect.h b/Engine/Plugins/Runtime/GameplayAbilities/Source/GameplayAbilities/Public/GameplayEffect.h
--- a/Engine/Plugins/Runtime/GameplayAbilities/Source/GameplayAbilities/Public/GameplayEffect.h
+++ b/Engine/Plugins/Runtime/GameplayAbilities/Source/GameplayAbilities/Public/GameplayEffect.h
@@ -1046,6 +1046,8 @@ struct GAMEPLAYABILITIES_API FGameplayEffectSpec
 	void SetDuration(float NewDuration, bool bLockDuration);
 
 	float GetDuration() const;
+
+	/** Returns the Period for the effect. If DurationPolicy is Instant this will forcibly return NO_PERIOD */
 	float GetPeriod() const;
 	
 	UE_DEPRECATED(5.3, "This no longer applies.  Use UChanceToApplyGameplayEffectComponent instead")
