[GAS] Reintroducing some validation for FScalableFloat

#rnx

diff --git a/Engine/Plugins/Runtime/GameplayAbilities/Source/GameplayAbilities/Private/ScalableFloat.cpp b/Engine/Plugins/Runtime/GameplayAbilities/Source/GameplayAbilities/Private/ScalableFloat.cpp
--- a/Engine/Plugins/Runtime/GameplayAbilities/Source/GameplayAbilities/Private/ScalableFloat.cpp
+++ b/Engine/Plugins/Runtime/GameplayAbilities/Source/GameplayAbilities/Private/ScalableFloat.cpp
@@ -58,19 +58,30 @@ EDataValidationResult FScalableFloat::IsDataValid(FDataValidationContext& Contex
 			return EDataValidationResult::NotValidated;
 		}
 
-		// Disabling data registry validation for now since not all data registries are added in all cases where validation may occur (specifically in commandlets)
-		/*
+		// This case is a bit tough, we are not 100% sure all DataRegistries are loaded.  We should give some indication an error can be occurring,
+		// but we also don't want to fail any automated checks that skip DataRegistry loading.
 		float OutValue = 0.0f;
 		const FRealCurve* OutCurve = nullptr;
 		const FDataRegistryId DataRegistryId(RegistryType, Curve.RowName);
 		const FDataRegistryCacheGetResult GetCurveResult = SubSystem->EvaluateCachedCurve(OutValue, OutCurve, DataRegistryId, OutValue);
 		if (!GetCurveResult)
 		{
-			Context.AddError(FText::Format(LOCTEXT("DataRegistryCurveDoesNotExist", "{0}: Specified DataRegistry Curve '{1}' does not exist"),
-				FText::FromString(PathName), DataRegistryId.ToText()));
-			return EDataValidationResult::Invalid;
+			FText WarningText = FText::Format(LOCTEXT("DataRegistryCurveDoesNotExist", "{0}: Specified DataRegistry Curve '{1}' does not exist (or was not loaded - you may need to load it to fix this message)"),
+				FText::FromString(PathName), DataRegistryId.ToText());
+
+			switch (Context.GetValidationUsecase())
+			{
+				// If we're manually testing, we can enforce the user loads the DataRegistry to confirm it's right 
+				case EDataValidationUsecase::Manual:
+					Context.AddWarning(WarningText);
+					return EDataValidationResult::NotValidated;
+
+				// For all other cases, let's downgrade to an info message so we can inform a user without tripping any automation errors
+				default:
+					Context.AddMessage(EMessageSeverity::Info, WarningText);
+					return EDataValidationResult::NotValidated;
+			}
 		}
-		*/
 	}
 
 	return EDataValidationResult::Valid;
