Making parts of AbilityAsync_WaitGameplayTagQuery public so it can be created in C++ code

diff --git a/Engine/Plugins/Runtime/GameplayAbilities/Source/GameplayAbilities/Public/Abilities/Async/AbilityAsync_WaitGameplayTagQuery.h b/Engine/Plugins/Runtime/GameplayAbilities/Source/GameplayAbilities/Public/Abilities/Async/AbilityAsync_WaitGameplayTagQuery.h
--- a/Engine/Plugins/Runtime/GameplayAbilities/Source/GameplayAbilities/Public/Abilities/Async/AbilityAsync_WaitGameplayTagQuery.h
+++ b/Engine/Plugins/Runtime/GameplayAbilities/Source/GameplayAbilities/Public/Abilities/Async/AbilityAsync_WaitGameplayTagQuery.h
@@ -15,11 +15,28 @@ UCLASS()
 class GAMEPLAYABILITIES_API UAbilityAsync_WaitGameplayTagQuery : public UAbilityAsync
 {
 	GENERATED_BODY()
-protected:
+
+public:
+	/** This delegate will be triggered when the trigger condition has been met. */
+	UPROPERTY(BlueprintAssignable)
+	FAsyncWaitGameplayTagQueryDelegate Triggered;
+
+	/**
+	 * 	Wait until the given gameplay tag query has become true or false, based on TriggerCondition, looking at the target actors ASC.
+	 *  If the the tag query already satisfies the TriggerCondition when this task is started, it will immediately broadcast the Triggered
+	 *  event. It will keep listening as long as bOnlyTriggerOnce = false.
+	 *  If used in an ability graph, this async action will wait even after activation ends. It's recommended to use WaitGameplayTagQuery instead.
+	 */
+ 	UFUNCTION(BlueprintCallable, Category = "Ability|Tasks", meta = (DefaultToSelf = "TargetActor", BlueprintInternalUseOnly = "TRUE"))
+ 	static UAbilityAsync_WaitGameplayTagQuery* WaitGameplayTagQueryOnActor(AActor* TargetActor, 
+																		   const FGameplayTagQuery TagQuery, 
+																		   const EWaitGameplayTagQueryTriggerCondition TriggerCondition = EWaitGameplayTagQueryTriggerCondition::WhenTrue,
+																		   const bool bOnlyTriggerOnce=false);
 
 	/** Activates this AbilityAsync node. */
 	virtual void Activate() override;
 
+protected:
 	/** Ends this AbilityAsync node, unregistering all tag changed delagates. */
 	virtual void EndAction() override;
 
@@ -53,20 +70,4 @@ protected:
 
 	/** This is the handles to the tag changed delegate for each gameplay tag in the TagQuery. */
 	TMap<FGameplayTag, FDelegateHandle> TagHandleMap;
-
-	/** This delegate will be triggered when the trigger condition has been met. */
-	UPROPERTY(BlueprintAssignable)
-	FAsyncWaitGameplayTagQueryDelegate Triggered;
-
-	/**
-	 * 	Wait until the given gameplay tag query has become true or false, based on TriggerCondition, looking at the target actors ASC.
-	 *  If the the tag query already satisfies the TriggerCondition when this task is started, it will immediately broadcast the Triggered
-	 *  event. It will keep listening as long as bOnlyTriggerOnce = false.
-	 *  If used in an ability graph, this async action will wait even after activation ends. It's recommended to use WaitGameplayTagQuery instead.
-	 */
- 	UFUNCTION(BlueprintCallable, Category = "Ability|Tasks", meta = (DefaultToSelf = "TargetActor", BlueprintInternalUseOnly = "TRUE"))
- 	static UAbilityAsync_WaitGameplayTagQuery* WaitGameplayTagQueryOnActor(AActor* TargetActor, 
-																		   const FGameplayTagQuery TagQuery, 
-																		   const EWaitGameplayTagQueryTriggerCondition TriggerCondition = EWaitGameplayTagQueryTriggerCondition::WhenTrue,
-																		   const bool bOnlyTriggerOnce=false);
 };
\ No newline at end of file
