Do not try to lookup scalable floats during engine shutdown. This can happen during the shutdown streaming flush if assets try to resolve floats during PostLoad.
This behavior is controlled by s.FlushStreamingOnExit, but it flushes by default.

diff --git a/Engine/Plugins/Runtime/GameplayAbilities/Source/GameplayAbilities/Private/ScalableFloat.cpp b/Engine/Plugins/Runtime/GameplayAbilities/Source/GameplayAbilities/Private/ScalableFloat.cpp
--- a/Engine/Plugins/Runtime/GameplayAbilities/Source/GameplayAbilities/Private/ScalableFloat.cpp
+++ b/Engine/Plugins/Runtime/GameplayAbilities/Source/GameplayAbilities/Private/ScalableFloat.cpp
@@ -109,7 +109,8 @@ bool FScalableFloat::EvaluateCurveAtLevel(float& OutValue, const FRealCurve*& Ou
 				CachedCurve = OutCurve = Curve.GetCurve(ContextString, bWarnIfInvalid);
 				LocalCachedCurveID = GlobalCachedCurveID;
 			}
-			else if (RegistryType.IsValid())
+			// Don't lookup during shutdown as the subsystem may have been deleted
+			else if (RegistryType.IsValid() && !IsEngineExitRequested())
 			{
 				UDataRegistrySubsystem* SubSystem = UDataRegistrySubsystem::Get();
 				if (ensure(SubSystem))
