Moved EndTask() inside the !bIsSimulating check so that clients don't kill the JumpForce Ability Task before the server.

diff --git a/Engine/Plugins/Runtime/GameplayAbilities/Source/GameplayAbilities/Private/Abilities/Tasks/AbilityTask_ApplyRootMotionJumpForce.cpp b/Engine/Plugins/Runtime/GameplayAbilities/Source/GameplayAbilities/Private/Abilities/Tasks/AbilityTask_ApplyRootMotionJumpForce.cpp
--- a/Engine/Plugins/Runtime/GameplayAbilities/Source/GameplayAbilities/Private/Abilities/Tasks/AbilityTask_ApplyRootMotionJumpForce.cpp
+++ b/Engine/Plugins/Runtime/GameplayAbilities/Source/GameplayAbilities/Private/Abilities/Tasks/AbilityTask_ApplyRootMotionJumpForce.cpp
@@ -140,9 +140,8 @@ void UAbilityTask_ApplyRootMotionJumpForce::Finish()
 				OnFinish.Broadcast();
 			}
 		}
+		EndTask();
 	}
-
-	EndTask();
 }
 
 void UAbilityTask_ApplyRootMotionJumpForce::TickTask(float DeltaTime)
@@ -175,7 +174,8 @@ void UAbilityTask_ApplyRootMotionJumpForce::TickTask(float DeltaTime)
 	}
 	else
 	{
-		Finish();
+		bIsFinished = true;
+		EndTask();
 	}
 }
 
