Added helpful gameplay tag binding and unbinding functions to the UAbilitySystemBlueprintLibrary that leverages dynamic delegates.

-The new methods allow for:
-listening to gameplay tag changes as a boolean applied/removed, or as a count
-binding to a group of tags or a single tag
-unbinding by handle, optionally by specific tag
-optionally executing the dynamic delegate immediately if the tag is already applied
-they aren't async task nodes so can be used in functions in blueprint

-The purpose of these new binding fns is to give designers more options when it comes to listening for gameplay tag changes, adding another tool to the GAS designer's toolbox.

-Also useful for native/code but since there are already non-dynamic delegates exposed for code to use via UAbilitySystemComponent::RegisterGameplayTagEvent, the benefit is less. Code would want to use these in one of the following cases:
-They want their callback to immediately execute if we already have the tag (this is a pretty big one & a super common pattern in code, so definitely precedent for this as a time-saving measure)
-They want to bind to an explicit list of tags with one binding

-Note that the word 'EventWrapper' is used throughout the change; this is to indicate that these tag binding events are just a wrapper for the non-dynamic multicast delegates that the ASC maintains for tag changes.

#tests Tested all new blueprint callable methods added to the UAbilitySystemBlueprintLibrary by listening for gameplay tags that are applied and removed. Tested listening for multiple tags with one binding call (BindEventWrapperToAnyOfGameplayTagsChanged and BindEventWrapperToAnyOfGameplayTagContainerChanged). Tested listening to multiple tags, unbinding some but not all of those, then listening to the tags we remain bounded to. Tested all binding and unbinding methods in conjunction with eachother. Tested that the same object binding to the same tag with multiple unique bind and unbind calls all works as expected. Tested that the bExecuteImmediatelyIfTagApplied param on the bind fns works as expected. Tested correct tag count coming from new event bindings

diff --git a/Engine/Plugins/Runtime/GameplayAbilities/Source/GameplayAbilities/Private/AbilitySystemBlueprintLibrary.cpp b/Engine/Plugins/Runtime/GameplayAbilities/Source/GameplayAbilities/Private/AbilitySystemBlueprintLibrary.cpp
--- a/Engine/Plugins/Runtime/GameplayAbilities/Source/GameplayAbilities/Private/AbilitySystemBlueprintLibrary.cpp
+++ b/Engine/Plugins/Runtime/GameplayAbilities/Source/GameplayAbilities/Private/AbilitySystemBlueprintLibrary.cpp
@@ -14,6 +14,61 @@
 
 #include UE_INLINE_GENERATED_CPP_BY_NAME(AbilitySystemBlueprintLibrary)
 
+FGameplayTagChangedEventWrapperSpecHandle::FGameplayTagChangedEventWrapperSpecHandle()
+	: Data(nullptr)
+{
+}
+
+FGameplayTagChangedEventWrapperSpecHandle::FGameplayTagChangedEventWrapperSpecHandle(FGameplayTagChangedEventWrapperSpec* DataPtr)
+	: Data(DataPtr)
+{
+}
+
+bool FGameplayTagChangedEventWrapperSpecHandle::operator==(FGameplayTagChangedEventWrapperSpecHandle const& Other) const
+{
+	const bool bBothValid = Data.IsValid() && Other.Data.IsValid();
+	const bool bBothInvalid = !Data.IsValid() && !Other.Data.IsValid();
+	return (bBothInvalid || (bBothValid && (Data.Get() == Other.Data.Get())));
+}
+
+bool FGameplayTagChangedEventWrapperSpecHandle::operator!=(FGameplayTagChangedEventWrapperSpecHandle const& Other) const
+{
+	return !(FGameplayTagChangedEventWrapperSpecHandle::operator==(Other));
+}
+
+FGameplayTagChangedEventWrapperSpec::FGameplayTagChangedEventWrapperSpec(
+	UAbilitySystemComponent* AbilitySystemComponent, 
+	FOnGameplayTagChangedEventWrapperSignature InGameplayTagChangedEventWrapperDelegate, 
+	EGameplayTagEventType::Type InTagListeningPolicy)
+	: AbilitySystemComponentWk(AbilitySystemComponent)
+	, GameplayTagChangedEventWrapperDelegate(InGameplayTagChangedEventWrapperDelegate)
+	, TagListeningPolicy(InTagListeningPolicy)
+	, DelegateBindings{}
+{
+}
+
+FGameplayTagChangedEventWrapperSpec::~FGameplayTagChangedEventWrapperSpec()
+{
+	const int32 RemainingDelegateBindingsCount = DelegateBindings.Num();
+	if (RemainingDelegateBindingsCount > 0)
+	{
+		// We still have delegates bound to the ASC - we need to warn the user!
+		// We expect the user to unbind delegates they bound.
+
+		// The exception is if the ASC itself is not valid which indicates things are tearing down - in that case, we'll give them a pass since it's a moot point that
+		//  we are still bound if the ASC isn't around anymore.
+		UAbilitySystemComponent* AbilitySystemComponent = AbilitySystemComponentWk.Get();
+		if (IsValid(AbilitySystemComponent))
+		{
+			ABILITY_LOG(
+				Error, 
+				TEXT("~FGameplayTagChangedEventWrapperSpec: our bound spec is being destroyed but we still have %d delegate bindings bound to the ASC on '%s'! Please cache off the Bound delegate handle and unbind it when finished."),
+				RemainingDelegateBindingsCount,
+				*GetNameSafe(AbilitySystemComponent->GetOwner()));
+		}
+	}
+}
+
 UAbilitySystemBlueprintLibrary::UAbilitySystemBlueprintLibrary(const FObjectInitializer& ObjectInitializer)
 : Super(ObjectInitializer)
 {
@@ -49,6 +104,159 @@ void UAbilitySystemBlueprintLibrary::SendGameplayEventToActor(AActor* Actor, FGa
 	}
 }
 
+FGameplayTagChangedEventWrapperSpecHandle UAbilitySystemBlueprintLibrary::BindEventWrapperToGameplayTagChanged(
+	UAbilitySystemComponent* AbilitySystemComponent,
+	FGameplayTag Tag, 
+	FOnGameplayTagChangedEventWrapperSignature GameplayTagChangedEventWrapperDelegate, 
+	bool bExecuteImmediatelyIfTagApplied /*= true*/, 
+	EGameplayTagEventType::Type TagListeningPolicy /*= EGameplayTagEventType::NewOrRemoved*/)
+{
+	if (!::IsValid(AbilitySystemComponent))
+	{
+		return FGameplayTagChangedEventWrapperSpecHandle();
+	}
+
+	FGameplayTagChangedEventWrapperSpec* TagBindingSpec = new FGameplayTagChangedEventWrapperSpec(AbilitySystemComponent, GameplayTagChangedEventWrapperDelegate, TagListeningPolicy);
+	FGameplayTagChangedEventWrapperSpecHandle TagBindingHandle(TagBindingSpec);
+
+	// Bind to the ASC's tag change listening delegate (which is not a 'dynamic' delegate and thereby can't be used in BP).
+	const FDelegateHandle TagChangedDelegateHandle = AbilitySystemComponent->RegisterGameplayTagEvent(Tag, TagListeningPolicy).AddLambda(
+		[GameplayTagChangedEventWrapperDelegate]
+		(const FGameplayTag GameplayTag, int32 GameplayTagCount)
+	{
+		UAbilitySystemBlueprintLibrary::ProcessGameplayTagChangedEventWrapper(GameplayTag, GameplayTagCount, GameplayTagChangedEventWrapperDelegate);
+	});
+
+	TagBindingSpec->DelegateBindings.Add(Tag, TagChangedDelegateHandle);
+
+	if (bExecuteImmediatelyIfTagApplied)
+	{
+		const int32 GameplayTagCount = AbilitySystemComponent->GetGameplayTagCount(Tag);
+		if (GameplayTagCount > 0)
+		{
+			GameplayTagChangedEventWrapperDelegate.Execute(Tag, GameplayTagCount);
+		}
+	}
+	
+	return TagBindingHandle;
+}
+
+FGameplayTagChangedEventWrapperSpecHandle UAbilitySystemBlueprintLibrary::BindEventWrapperToAnyOfGameplayTagsChanged(
+	UAbilitySystemComponent* AbilitySystemComponent,
+	const TArray<FGameplayTag>& Tags, 
+	FOnGameplayTagChangedEventWrapperSignature GameplayTagChangedEventWrapperDelegate, 
+	bool bExecuteImmediatelyIfTagApplied/* = true*/,
+	EGameplayTagEventType::Type TagListeningPolicy/* = EGameplayTagEventType::NewOrRemoved*/)
+{
+	if (!::IsValid(AbilitySystemComponent))
+	{
+		return FGameplayTagChangedEventWrapperSpecHandle();
+	}
+
+	FGameplayTagChangedEventWrapperSpec* TagBindingSpec = new FGameplayTagChangedEventWrapperSpec(AbilitySystemComponent, GameplayTagChangedEventWrapperDelegate, TagListeningPolicy);
+	FGameplayTagChangedEventWrapperSpecHandle TagBindingHandle(TagBindingSpec);
+
+	TagBindingSpec->DelegateBindings.Reserve(Tags.Num());
+
+	// Bind each tag and add to the DelegateBindings container.
+	for (const FGameplayTag& Tag : Tags)
+	{
+		// Bind to the ASC's tag change listening delegate (which is not a 'dynamic' delegate and thereby can't be used in BP).
+		const FDelegateHandle TagChangedDelegateHandle = AbilitySystemComponent->RegisterGameplayTagEvent(Tag, TagListeningPolicy).AddLambda(
+			[GameplayTagChangedEventWrapperDelegate]
+			(const FGameplayTag GameplayTag, int32 GameplayTagCount)
+		{
+			UAbilitySystemBlueprintLibrary::ProcessGameplayTagChangedEventWrapper(GameplayTag, GameplayTagCount, GameplayTagChangedEventWrapperDelegate);
+		});
+
+		TagBindingSpec->DelegateBindings.Add(Tag, TagChangedDelegateHandle);
+	}
+
+	if (bExecuteImmediatelyIfTagApplied)
+	{
+		for (const FGameplayTag& Tag : Tags)
+		{
+			const int32 GameplayTagCount = AbilitySystemComponent->GetGameplayTagCount(Tag);
+			if (GameplayTagCount > 0)
+			{
+				GameplayTagChangedEventWrapperDelegate.Execute(Tag, GameplayTagCount);
+			}
+		}
+	}
+
+	return TagBindingHandle;
+}
+
+FGameplayTagChangedEventWrapperSpecHandle UAbilitySystemBlueprintLibrary::BindEventWrapperToAnyOfGameplayTagContainerChanged(
+	UAbilitySystemComponent* AbilitySystemComponent,
+	const FGameplayTagContainer TagContainer, 
+	FOnGameplayTagChangedEventWrapperSignature GameplayTagChangedEventWrapperDelegate, 
+	bool bExecuteImmediatelyIfTagApplied/* = true*/,
+	EGameplayTagEventType::Type TagListeningPolicy/* = EGameplayTagEventType::NewOrRemoved*/)
+{
+	TArray<FGameplayTag> Tags;
+	TagContainer.GetGameplayTagArray(Tags);
+	
+	return BindEventWrapperToAnyOfGameplayTagsChanged(AbilitySystemComponent, Tags, GameplayTagChangedEventWrapperDelegate, bExecuteImmediatelyIfTagApplied, TagListeningPolicy);
+}
+
+void UAbilitySystemBlueprintLibrary::UnbindAllGameplayTagChangedEventWrappersForHandle(FGameplayTagChangedEventWrapperSpecHandle Handle)
+{
+	FGameplayTagChangedEventWrapperSpec* GameplayTagChangedEventDataPtr = Handle.Data.Get();
+	if (GameplayTagChangedEventDataPtr == nullptr)
+	{
+		return;
+	}
+
+	UAbilitySystemComponent* AbilitySystemComponent = GameplayTagChangedEventDataPtr->AbilitySystemComponentWk.Get();
+	if (AbilitySystemComponent == nullptr)
+	{
+		return;
+	}
+
+	for (auto Iterator = GameplayTagChangedEventDataPtr->DelegateBindings.CreateConstIterator(); Iterator; ++Iterator)
+	{
+		AbilitySystemComponent->UnregisterGameplayTagEvent(Iterator->Value, Iterator->Key, GameplayTagChangedEventDataPtr->TagListeningPolicy);
+	}
+
+	GameplayTagChangedEventDataPtr->DelegateBindings.Reset();
+}
+
+void UAbilitySystemBlueprintLibrary::UnbindGameplayTagChangedEventWrapperForHandle(FGameplayTag Tag, FGameplayTagChangedEventWrapperSpecHandle Handle)
+{
+	FGameplayTagChangedEventWrapperSpec* GameplayTagChangedEventDataPtr = Handle.Data.Get();
+	if (GameplayTagChangedEventDataPtr == nullptr)
+	{
+		return;
+	}
+
+	UAbilitySystemComponent* AbilitySystemComponent = GameplayTagChangedEventDataPtr->AbilitySystemComponentWk.Get();
+	if (AbilitySystemComponent == nullptr)
+	{
+		return;
+	}
+
+	for (auto DelegateBindingIterator = GameplayTagChangedEventDataPtr->DelegateBindings.CreateIterator(); DelegateBindingIterator; ++DelegateBindingIterator)
+	{
+		const FGameplayTag& BoundTag = DelegateBindingIterator->Key;
+		if (!BoundTag.MatchesTagExact(Tag))
+		{
+			continue;
+		}
+
+		AbilitySystemComponent->UnregisterGameplayTagEvent(DelegateBindingIterator->Value, BoundTag, GameplayTagChangedEventDataPtr->TagListeningPolicy);
+		DelegateBindingIterator.RemoveCurrent();
+	}
+}
+
+void UAbilitySystemBlueprintLibrary::ProcessGameplayTagChangedEventWrapper(
+	const FGameplayTag GameplayTag, 
+	int32 GameplayTagCount,
+	FOnGameplayTagChangedEventWrapperSignature GameplayTagChangedEventWrapperDelegate)
+{
+	GameplayTagChangedEventWrapperDelegate.Execute(GameplayTag, GameplayTagCount);
+}
+
 bool UAbilitySystemBlueprintLibrary::IsValid(FGameplayAttribute Attribute)
 {
 	return Attribute.IsValid();
diff --git a/Engine/Plugins/Runtime/GameplayAbilities/Source/GameplayAbilities/Public/AbilitySystemBlueprintLibrary.h b/Engine/Plugins/Runtime/GameplayAbilities/Source/GameplayAbilities/Public/AbilitySystemBlueprintLibrary.h
--- a/Engine/Plugins/Runtime/GameplayAbilities/Source/GameplayAbilities/Public/AbilitySystemBlueprintLibrary.h
+++ b/Engine/Plugins/Runtime/GameplayAbilities/Source/GameplayAbilities/Public/AbilitySystemBlueprintLibrary.h
@@ -21,6 +21,47 @@ class UAbilitySystemComponent;
 class UGameplayEffect;
 class UGameplayEffectUIData;
 
+/** Called when a gameplay tag bound to a event wrapper via one of the BindEventWrapper<to gameplay tag(s)> methods on the AbilitySystemLibrary changes. */
+DECLARE_DYNAMIC_DELEGATE_TwoParams(FOnGameplayTagChangedEventWrapperSignature, const FGameplayTag&, Tag, int32, TagCount);
+
+/** Holds tracking data for gameplay tag change event wrappers that have been bound. */
+struct FGameplayTagChangedEventWrapperSpec
+{
+	FGameplayTagChangedEventWrapperSpec(
+		UAbilitySystemComponent* AbilitySystemComponent,
+		FOnGameplayTagChangedEventWrapperSignature InGameplayTagChangedEventWrapperDelegate, 
+		EGameplayTagEventType::Type InTagListeningPolicy);
+	~FGameplayTagChangedEventWrapperSpec();
+
+	/** The AbilitySystemComponent this spec is bound to. */
+	TWeakObjectPtr<UAbilitySystemComponent> AbilitySystemComponentWk;
+
+	/** The event wrapper delegate cached off, to be executed when the gameplaytag(s) we care about are changed. */
+	FOnGameplayTagChangedEventWrapperSignature GameplayTagChangedEventWrapperDelegate;
+
+	/** Specifies what kinds of gameplay tag changes we will execute for. */
+	EGameplayTagEventType::Type TagListeningPolicy;
+
+	/** Map of the respective gameplay tags to the delegate handle the ASC gave us to use for unbinding later. */
+	TMap<FGameplayTag, FDelegateHandle> DelegateBindings;
+};
+
+/** Handle to a event wrapper listening for gameplay tag change(s) via one of the BindEventWrapper<to gameplay tag(s)> methods on the AbilitySystemLibrary */
+USTRUCT(BlueprintType)
+struct FGameplayTagChangedEventWrapperSpecHandle
+{
+	GENERATED_BODY()
+
+	FGameplayTagChangedEventWrapperSpecHandle();
+	FGameplayTagChangedEventWrapperSpecHandle(FGameplayTagChangedEventWrapperSpec* DataPtr);
+
+	/** Internal pointer to binding spec */
+	TSharedPtr<FGameplayTagChangedEventWrapperSpec>	Data;
+
+	bool operator==(FGameplayTagChangedEventWrapperSpecHandle const& Other) const;
+	bool operator!=(FGameplayTagChangedEventWrapperSpecHandle const& Other) const;
+};
+
 /** Blueprint library for ability system. Many of these functions are useful to call from native as well */
 UCLASS(meta=(ScriptName="AbilitySystemLibrary"))
 class GAMEPLAYABILITIES_API UAbilitySystemBlueprintLibrary : public UBlueprintFunctionLibrary
@@ -39,6 +80,89 @@ class GAMEPLAYABILITIES_API UAbilitySystemBlueprintLibrary : public UBlueprintFu
 	UFUNCTION(BlueprintCallable, Category = Ability, Meta = (Tooltip = "This function can be used to trigger an ability on the actor in question with useful payload data."))
 	static void SendGameplayEventToActor(AActor* Actor, UPARAM(meta=(GameplayTagFilter="GameplayEventTagsCategory")) FGameplayTag EventTag, FGameplayEventData Payload);
 	
+	// -------------------------------------------------------------------------------
+	//		Gameplay Tag
+	// -------------------------------------------------------------------------------
+
+	/**
+	 * Binds to changes in the given Tag on the given ASC's owned tags.
+	 * Cache off the returned handle and call one of the 'Unbind'...'EventWrapper' fns and pass in the handle when you are finished with the binding.
+	 * @param Tag                               FGameplayTag to listen for changes on
+	 * @param BindingObject                     UObject owning the Event (usually self)
+	 * @param GameplayTagChangedEventWrapperDelegate Tag Changed Event to trigger
+	 * @param bExecuteImmediatelyIfTagApplied   If true, the bound event will immediately execute if we already have the tag.
+	 * @param TagListeningPolicy                Whether we are listening for any tag change or just applied/removed.
+	 * @return                                  FGameplayTagChangedEventWrapperSpecHandle Handle by which this binding request can be unbound.
+	 */
+	UFUNCTION(BlueprintCallable, Category = "Gameplay Tags", meta = (AdvancedDisplay = "bExecuteImmediatelyIfTagApplied, TagListeningPolicy"))
+	static FGameplayTagChangedEventWrapperSpecHandle BindEventWrapperToGameplayTagChanged(
+		UAbilitySystemComponent* AbilitySystemComponent,
+		FGameplayTag Tag, 
+		FOnGameplayTagChangedEventWrapperSignature GameplayTagChangedEventWrapperDelegate, 
+		bool bExecuteImmediatelyIfTagApplied = true,
+		EGameplayTagEventType::Type TagListeningPolicy = EGameplayTagEventType::NewOrRemoved);
+
+	/**
+	 * Binds to changes in the given Tags on the given ASC's owned tags.
+	 * Cache off the returned handle and call one of the 'Unbind'...'EventWrapper' fns and pass in the handle when you are finished with the binding.
+	 * @param Tags                              TArray of FGameplayTag to listen for changes on
+	 * @param BindingObject                     UObject owning the Event (usually self)
+	 * @param GameplayTagChangedEventWrapperDelegate Tag Changed Event to trigger
+	 * @param bExecuteImmediatelyIfTagApplied   If true, the bound event will immediately execute if we already have the tag.
+	 * @param TagListeningPolicy                Whether we are listening for any tag change or just applied/removed.
+	 * @return                                  FGameplayTagChangedEventWrapperSpecHandle Handle by which this binding request can be unbound.
+	 */
+	UFUNCTION(BlueprintCallable, Category = "Gameplay Tags", meta = (AdvancedDisplay = "bExecuteImmediatelyIfTagApplied, TagListeningPolicy"))
+	static FGameplayTagChangedEventWrapperSpecHandle BindEventWrapperToAnyOfGameplayTagsChanged(
+		UAbilitySystemComponent* AbilitySystemComponent,
+		const TArray<FGameplayTag>& Tags,
+		FOnGameplayTagChangedEventWrapperSignature GameplayTagChangedEventWrapperDelegate, 
+		bool bExecuteImmediatelyIfTagApplied = true,
+		EGameplayTagEventType::Type TagListeningPolicy = EGameplayTagEventType::NewOrRemoved);
+
+	/**
+	 * Binds to changes in the given TagContainer on the given ASC's owned tags.
+	 * Cache off the returned handle and call one of the 'Unbind'...'EventWrapper' fns and pass in the handle when you are finished with the binding.
+	 * @param TagContainer                      FGameplayTagContainer of tags to listen for changes on
+	 * @param BindingObject                     UObject owning the Event (usually self)
+	 * @param GameplayTagChangedEventWrapperDelegate Tag Changed Event to trigger
+	 * @param bExecuteImmediatelyIfTagApplied   If true, the bound event will immediately execute if we already have the tag.
+	 * @param TagListeningPolicy                Whether we are listening for any tag change or just applied/removed.
+	 * @return                                  FGameplayTagChangedEventWrapperSpecHandle Handle by which this binding request can be unbound.
+	 */
+	UFUNCTION(BlueprintCallable, Category = "Gameplay Tags", meta = (AdvancedDisplay = "bExecuteImmediatelyIfTagApplied, TagListeningPolicy"))
+	static FGameplayTagChangedEventWrapperSpecHandle BindEventWrapperToAnyOfGameplayTagContainerChanged(
+		UAbilitySystemComponent* AbilitySystemComponent,
+		const FGameplayTagContainer TagContainer, 
+		FOnGameplayTagChangedEventWrapperSignature GameplayTagChangedEventWrapperDelegate, 
+		bool bExecuteImmediatelyIfTagApplied = true,
+		EGameplayTagEventType::Type TagListeningPolicy = EGameplayTagEventType::NewOrRemoved);
+
+	/**
+	 * Unbinds the event wrapper tag change event bound via a BindEventWrapper<to gameplay tag(s)> method that is tied to the given Handle.
+	 * (expected to unbind 1 or none)
+	 * @param Handle     FGameplayTagChangedEventWrapperSpecHandle Handle provided when binding to a delegate
+	 */
+	UFUNCTION(BlueprintCallable, Category = "Gameplay Tags")
+	static void UnbindAllGameplayTagChangedEventWrappersForHandle(FGameplayTagChangedEventWrapperSpecHandle Handle);
+
+	/**
+	 * Unbinds the event wrapper tag change event bound via a BindEventWrapper<to gameplay tag(s)> method that is tied to the given Handle, for the specific tag.
+	 *  and were bound to the given Tag.
+	 * (expected to unbind 1 or none, only makes sense to call if the original binding was for listening to multiple tags.)
+	 * @param Tag        FGameplayTag to unbind from
+	 * @param Handle     int Handle provided when binding to a delegate
+	 */
+	UFUNCTION(BlueprintCallable, Category = "Gameplay Tags")
+	static void UnbindGameplayTagChangedEventWrapperForHandle(FGameplayTag Tag, FGameplayTagChangedEventWrapperSpecHandle Handle);
+
+protected:
+
+	// Helper fn to process gameplay tag changed event wrappers.
+	static void ProcessGameplayTagChangedEventWrapper(const FGameplayTag GameplayTag, int32 GameplayTagCount, FOnGameplayTagChangedEventWrapperSignature GameplayTagChangedEventWrapperDelegate);
+
+public:
+
 	// -------------------------------------------------------------------------------
 	//		Attribute
 	// -------------------------------------------------------------------------------
