[GameplayAbilities]
* Moved dll storage to methods/variables instead of on types. This is needed to get merged modules working

diff --git a/Engine/Plugins/Runtime/GameplayAbilities/Source/GameplayAbilities/Public/Abilities/Async/AbilityAsync.h b/Engine/Plugins/Runtime/GameplayAbilities/Source/GameplayAbilities/Public/Abilities/Async/AbilityAsync.h
--- a/Engine/Plugins/Runtime/GameplayAbilities/Source/GameplayAbilities/Public/Abilities/Async/AbilityAsync.h
+++ b/Engine/Plugins/Runtime/GameplayAbilities/Source/GameplayAbilities/Public/Abilities/Async/AbilityAsync.h
@@ -8,6 +8,8 @@
 #include "Engine/CancellableAsyncAction.h"
 #include "AbilityAsync.generated.h"
 
+#define UE_API GAMEPLAYABILITIES_API
+
 class UAbilitySystemComponent;
 
 /**
@@ -16,31 +18,33 @@ class UAbilitySystemComponent;
  * By default these actions are only kept alive by the blueprint graph that spawns them and will eventually be destroyed if the graph instance is deleted or spawns a replacement.
  * EndAction should be called when a one-time action has succeeded or failed, but for longer-lived actions with multiple triggers it can be called from blueprints.
  */
-UCLASS(Abstract, meta = (ExposedAsyncProxy = AsyncAction))
-class GAMEPLAYABILITIES_API UAbilityAsync : public UCancellableAsyncAction
+UCLASS(Abstract, meta = (ExposedAsyncProxy = AsyncAction), MinimalAPI)
+class UAbilityAsync : public UCancellableAsyncAction
 {
 	GENERATED_BODY()
 public:
 
-	virtual void Cancel() override;
+	UE_API virtual void Cancel() override;
 
 	/** Explicitly end the action, will disable any callbacks and allow action to be destroyed */
 	UFUNCTION(BlueprintCallable, Category = "Ability|Async")
-	virtual void EndAction();
+	UE_API virtual void EndAction();
 
 	/** This should be called prior to broadcasting delegates back into the event graph, this ensures the action and ability system component are still valid */
-	virtual bool ShouldBroadcastDelegates() const;
+	UE_API virtual bool ShouldBroadcastDelegates() const;
 
 	/** Returns the ability system component this action is bound to */
-	virtual UAbilitySystemComponent* GetAbilitySystemComponent() const;
+	UE_API virtual UAbilitySystemComponent* GetAbilitySystemComponent() const;
 
 	/** Sets the bound component */
-	virtual void SetAbilitySystemComponent(UAbilitySystemComponent* InAbilitySystemComponent);
+	UE_API virtual void SetAbilitySystemComponent(UAbilitySystemComponent* InAbilitySystemComponent);
 
 	/** Sets the bound component by searching actor for one */
-	virtual void SetAbilityActor(AActor* InActor);
+	UE_API virtual void SetAbilityActor(AActor* InActor);
 
 private:
 	TWeakObjectPtr<UAbilitySystemComponent> AbilitySystemComponent;
 
 };
+
+#undef UE_API
diff --git a/Engine/Plugins/Runtime/GameplayAbilities/Source/GameplayAbilities/Public/Abilities/Async/AbilityAsync_WaitAttributeChanged.h b/Engine/Plugins/Runtime/GameplayAbilities/Source/GameplayAbilities/Public/Abilities/Async/AbilityAsync_WaitAttributeChanged.h
--- a/Engine/Plugins/Runtime/GameplayAbilities/Source/GameplayAbilities/Public/Abilities/Async/AbilityAsync_WaitAttributeChanged.h
+++ b/Engine/Plugins/Runtime/GameplayAbilities/Source/GameplayAbilities/Public/Abilities/Async/AbilityAsync_WaitAttributeChanged.h
@@ -7,10 +7,12 @@
 #include "Abilities/GameplayAbilityTypes.h"
 #include "AbilityAsync_WaitAttributeChanged.generated.h"
 
+#define UE_API GAMEPLAYABILITIES_API
+
 class UAbilitySystemComponent;
 
-UCLASS()
-class GAMEPLAYABILITIES_API UAbilityAsync_WaitAttributeChanged : public UAbilityAsync
+UCLASS(MinimalAPI)
+class UAbilityAsync_WaitAttributeChanged : public UAbilityAsync
 {
 	GENERATED_BODY()
 
@@ -22,7 +24,7 @@ public:
 	 * If used in an ability graph, this async action will wait even after activation ends. It's recommended to use WaitForAttributeChange instead.
 	 */
 	UFUNCTION(BlueprintCallable, Category = "Ability|Tasks", meta = (DefaultToSelf = "TargetActor", BlueprintInternalUseOnly = "TRUE"))
-	static UAbilityAsync_WaitAttributeChanged* WaitForAttributeChanged(AActor* TargetActor, FGameplayAttribute Attribute, bool OnlyTriggerOnce = false);
+	static UE_API UAbilityAsync_WaitAttributeChanged* WaitForAttributeChanged(AActor* TargetActor, FGameplayAttribute Attribute, bool OnlyTriggerOnce = false);
 
 	DECLARE_DYNAMIC_MULTICAST_DELEGATE_ThreeParams(FAsyncWaitAttributeChangedDelegate, FGameplayAttribute, Attribute, float, NewValue, float, OldValue);
 	UPROPERTY(BlueprintAssignable)
@@ -30,13 +32,15 @@ public:
 
 protected:
 
-	virtual void Activate() override;
-	virtual void EndAction() override;
+	UE_API virtual void Activate() override;
+	UE_API virtual void EndAction() override;
 
-	void OnAttributeChanged(const FOnAttributeChangeData& ChangeData);
+	UE_API void OnAttributeChanged(const FOnAttributeChangeData& ChangeData);
 
 	FGameplayAttribute Attribute;
 	bool OnlyTriggerOnce = false;
 
 	FDelegateHandle MyHandle;
 };
+
+#undef UE_API
diff --git a/Engine/Plugins/Runtime/GameplayAbilities/Source/GameplayAbilities/Public/Abilities/Async/AbilityAsync_WaitGameplayEffectApplied.h b/Engine/Plugins/Runtime/GameplayAbilities/Source/GameplayAbilities/Public/Abilities/Async/AbilityAsync_WaitGameplayEffectApplied.h
--- a/Engine/Plugins/Runtime/GameplayAbilities/Source/GameplayAbilities/Public/Abilities/Async/AbilityAsync_WaitGameplayEffectApplied.h
+++ b/Engine/Plugins/Runtime/GameplayAbilities/Source/GameplayAbilities/Public/Abilities/Async/AbilityAsync_WaitGameplayEffectApplied.h
@@ -8,6 +8,8 @@
 #include "Delegates/IDelegateInstance.h"
 #include "AbilityAsync_WaitGameplayEffectApplied.generated.h"
 
+#define UE_API GAMEPLAYABILITIES_API
+
 class UAbilitySystemComponent;
 
 /**
@@ -16,8 +18,8 @@ class UAbilitySystemComponent;
  * change, and the gameplay cues they emit are replicated.
  * This will only listen for local server or predicted client effects.
  */
-UCLASS()
-class GAMEPLAYABILITIES_API UAbilityAsync_WaitGameplayEffectApplied : public UAbilityAsync
+UCLASS(MinimalAPI)
+class UAbilityAsync_WaitGameplayEffectApplied : public UAbilityAsync
 {
 	GENERATED_BODY()
 
@@ -28,17 +30,17 @@ public:
 	 * If used in an ability graph, this async action will wait even after activation ends. It's recommended to use WaitGameplayEffectApplied instead.
 	 */
 	UFUNCTION(BlueprintCallable, Category = "Ability|Async", meta = (DefaultToSelf = "TargetActor", BlueprintInternalUseOnly = "TRUE"))
-	static UAbilityAsync_WaitGameplayEffectApplied* WaitGameplayEffectAppliedToActor(AActor* TargetActor, const FGameplayTargetDataFilterHandle SourceFilter, FGameplayTagRequirements SourceTagRequirements, FGameplayTagRequirements TargetTagRequirements, bool TriggerOnce = false, bool ListenForPeriodicEffect = false);
+	static UE_API UAbilityAsync_WaitGameplayEffectApplied* WaitGameplayEffectAppliedToActor(AActor* TargetActor, const FGameplayTargetDataFilterHandle SourceFilter, FGameplayTagRequirements SourceTagRequirements, FGameplayTagRequirements TargetTagRequirements, bool TriggerOnce = false, bool ListenForPeriodicEffect = false);
 
 	DECLARE_DYNAMIC_MULTICAST_DELEGATE_ThreeParams(FOnAppliedDelegate, AActor*, Source, FGameplayEffectSpecHandle, SpecHandle, FActiveGameplayEffectHandle, ActiveHandle);
 	UPROPERTY(BlueprintAssignable)
 	FOnAppliedDelegate OnApplied;
 
 protected:
-	virtual void Activate() override;
-	virtual void EndAction() override;
+	UE_API virtual void Activate() override;
+	UE_API virtual void EndAction() override;
 
-	void OnApplyGameplayEffectCallback(UAbilitySystemComponent* Target, const FGameplayEffectSpec& SpecApplied, FActiveGameplayEffectHandle ActiveHandle);
+	UE_API void OnApplyGameplayEffectCallback(UAbilitySystemComponent* Target, const FGameplayEffectSpec& SpecApplied, FActiveGameplayEffectHandle ActiveHandle);
 
 	FGameplayTargetDataFilterHandle Filter;
 	FGameplayTagRequirements SourceTagRequirements;
@@ -51,3 +53,5 @@ protected:
 	FDelegateHandle OnPeriodicGameplayEffectExecuteCallbackDelegateHandle;
 	bool bLocked = false;
 };
+
+#undef UE_API
diff --git a/Engine/Plugins/Runtime/GameplayAbilities/Source/GameplayAbilities/Public/Abilities/Async/AbilityAsync_WaitGameplayEvent.h b/Engine/Plugins/Runtime/GameplayAbilities/Source/GameplayAbilities/Public/Abilities/Async/AbilityAsync_WaitGameplayEvent.h
--- a/Engine/Plugins/Runtime/GameplayAbilities/Source/GameplayAbilities/Public/Abilities/Async/AbilityAsync_WaitGameplayEvent.h
+++ b/Engine/Plugins/Runtime/GameplayAbilities/Source/GameplayAbilities/Public/Abilities/Async/AbilityAsync_WaitGameplayEvent.h
@@ -7,10 +7,12 @@
 #include "Abilities/GameplayAbilityTypes.h"
 #include "AbilityAsync_WaitGameplayEvent.generated.h"
 
+#define UE_API GAMEPLAYABILITIES_API
+
 class UAbilitySystemComponent;
 
-UCLASS()
-class GAMEPLAYABILITIES_API UAbilityAsync_WaitGameplayEvent : public UAbilityAsync
+UCLASS(MinimalAPI)
+class UAbilityAsync_WaitGameplayEvent : public UAbilityAsync
 {
 	GENERATED_BODY()
 
@@ -22,7 +24,7 @@ public:
 	 * If used in an ability graph, this async action will wait even after activation ends. It's recommended to use WaitGameplayEvent instead.
 	 */
 	UFUNCTION(BlueprintCallable, Category = "Ability|Async", meta = (DefaultToSelf = "TargetActor", BlueprintInternalUseOnly = "TRUE"))
-	static UAbilityAsync_WaitGameplayEvent* WaitGameplayEventToActor(AActor* TargetActor, UPARAM(meta=(GameplayTagFilter="GameplayEventTagsCategory")) FGameplayTag EventTag, bool OnlyTriggerOnce = false, bool OnlyMatchExact = true);
+	static UE_API UAbilityAsync_WaitGameplayEvent* WaitGameplayEventToActor(AActor* TargetActor, UPARAM(meta=(GameplayTagFilter="GameplayEventTagsCategory")) FGameplayTag EventTag, bool OnlyTriggerOnce = false, bool OnlyMatchExact = true);
 
 	DECLARE_DYNAMIC_MULTICAST_DELEGATE_OneParam(FEventReceivedDelegate, FGameplayEventData, Payload);
 
@@ -30,11 +32,11 @@ public:
 	FEventReceivedDelegate EventReceived;
 
 protected:
-	virtual void Activate() override;
-	virtual void EndAction() override;
+	UE_API virtual void Activate() override;
+	UE_API virtual void EndAction() override;
 
-	virtual void GameplayEventCallback(const FGameplayEventData* Payload);
-	virtual void GameplayEventContainerCallback(FGameplayTag MatchingTag, const FGameplayEventData* Payload);
+	UE_API virtual void GameplayEventCallback(const FGameplayEventData* Payload);
+	UE_API virtual void GameplayEventContainerCallback(FGameplayTag MatchingTag, const FGameplayEventData* Payload);
 
 	FGameplayTag Tag;
 	bool OnlyTriggerOnce = false;
@@ -42,3 +44,5 @@ protected:
 
 	FDelegateHandle MyHandle;
 };
+
+#undef UE_API
diff --git a/Engine/Plugins/Runtime/GameplayAbilities/Source/GameplayAbilities/Public/Abilities/Async/AbilityAsync_WaitGameplayTag.h b/Engine/Plugins/Runtime/GameplayAbilities/Source/GameplayAbilities/Public/Abilities/Async/AbilityAsync_WaitGameplayTag.h
--- a/Engine/Plugins/Runtime/GameplayAbilities/Source/GameplayAbilities/Public/Abilities/Async/AbilityAsync_WaitGameplayTag.h
+++ b/Engine/Plugins/Runtime/GameplayAbilities/Source/GameplayAbilities/Public/Abilities/Async/AbilityAsync_WaitGameplayTag.h
@@ -7,21 +7,23 @@
 #include "Abilities/GameplayAbilityTypes.h"
 #include "AbilityAsync_WaitGameplayTag.generated.h"
 
+#define UE_API GAMEPLAYABILITIES_API
+
 class UAbilitySystemComponent;
 
-UCLASS(Abstract)
-class GAMEPLAYABILITIES_API UAbilityAsync_WaitGameplayTag : public UAbilityAsync
+UCLASS(Abstract, MinimalAPI)
+class UAbilityAsync_WaitGameplayTag : public UAbilityAsync
 {
 	GENERATED_BODY()
 protected:
 
 	DECLARE_DYNAMIC_MULTICAST_DELEGATE(FAsyncWaitGameplayTagDelegate);
 
-	virtual void Activate() override;
-	virtual void EndAction() override;
+	UE_API virtual void Activate() override;
+	UE_API virtual void EndAction() override;
 
-	virtual void GameplayTagCallback(const FGameplayTag Tag, int32 NewCount);
-	virtual void BroadcastDelegate();
+	UE_API virtual void GameplayTagCallback(const FGameplayTag Tag, int32 NewCount);
+	UE_API virtual void BroadcastDelegate();
 
 	int32 TargetCount = -1;
 	FGameplayTag Tag;
@@ -30,8 +32,8 @@ protected:
 	FDelegateHandle MyHandle;
 };
 
-UCLASS()
-class GAMEPLAYABILITIES_API UAbilityAsync_WaitGameplayTagAdded : public UAbilityAsync_WaitGameplayTag
+UCLASS(MinimalAPI)
+class UAbilityAsync_WaitGameplayTagAdded : public UAbilityAsync_WaitGameplayTag
 {
 	GENERATED_BODY()
 public:
@@ -45,13 +47,13 @@ public:
 	 * If used in an ability graph, this async action will wait even after activation ends. It's recommended to use WaitGameplayTagAdd instead.
 	 */
 	UFUNCTION(BlueprintCallable, Category = "Ability|Tasks", meta = (DefaultToSelf = "TargetActor", BlueprintInternalUseOnly = "TRUE"))
-	static UAbilityAsync_WaitGameplayTagAdded* WaitGameplayTagAddToActor(AActor* TargetActor, FGameplayTag Tag, bool OnlyTriggerOnce=false);
+	static UE_API UAbilityAsync_WaitGameplayTagAdded* WaitGameplayTagAddToActor(AActor* TargetActor, FGameplayTag Tag, bool OnlyTriggerOnce=false);
 
-	virtual void BroadcastDelegate() override;
+	UE_API virtual void BroadcastDelegate() override;
 };
 
-UCLASS()
-class GAMEPLAYABILITIES_API UAbilityAsync_WaitGameplayTagRemoved : public UAbilityAsync_WaitGameplayTag
+UCLASS(MinimalAPI)
+class UAbilityAsync_WaitGameplayTagRemoved : public UAbilityAsync_WaitGameplayTag
 {
 	GENERATED_BODY()
 public:
@@ -65,7 +67,9 @@ public:
 	 * If used in an ability graph, this async action will wait even after activation ends. It's recommended to use WaitGameplayTagRemove instead.
 	 */
 	UFUNCTION(BlueprintCallable, Category = "Ability|Tasks", meta = (DefaultToSelf = "TargetActor", BlueprintInternalUseOnly = "TRUE"))
-	static UAbilityAsync_WaitGameplayTagRemoved* WaitGameplayTagRemoveFromActor(AActor* TargetActor, FGameplayTag Tag, bool OnlyTriggerOnce=false);
+	static UE_API UAbilityAsync_WaitGameplayTagRemoved* WaitGameplayTagRemoveFromActor(AActor* TargetActor, FGameplayTag Tag, bool OnlyTriggerOnce=false);
+
+	UE_API virtual void BroadcastDelegate() override;
+};
 
-	virtual void BroadcastDelegate() override;
-};
\ No newline at end of file
+#undef UE_API
diff --git a/Engine/Plugins/Runtime/GameplayAbilities/Source/GameplayAbilities/Public/Abilities/Async/AbilityAsync_WaitGameplayTagCountChanged.h b/Engine/Plugins/Runtime/GameplayAbilities/Source/GameplayAbilities/Public/Abilities/Async/AbilityAsync_WaitGameplayTagCountChanged.h
--- a/Engine/Plugins/Runtime/GameplayAbilities/Source/GameplayAbilities/Public/Abilities/Async/AbilityAsync_WaitGameplayTagCountChanged.h
+++ b/Engine/Plugins/Runtime/GameplayAbilities/Source/GameplayAbilities/Public/Abilities/Async/AbilityAsync_WaitGameplayTagCountChanged.h
@@ -4,18 +4,20 @@
 #include "AbilityAsync.h"
 #include "AbilityAsync_WaitGameplayTagCountChanged.generated.h"
 
-UCLASS()
-class GAMEPLAYABILITIES_API UAbilityAsync_WaitGameplayTagCountChanged : public UAbilityAsync
+#define UE_API GAMEPLAYABILITIES_API
+
+UCLASS(MinimalAPI)
+class UAbilityAsync_WaitGameplayTagCountChanged : public UAbilityAsync
 {
 	GENERATED_BODY()
 protected:
 
 	DECLARE_DYNAMIC_MULTICAST_DELEGATE_OneParam(FAsyncWaitGameplayTagCountDelegate, int32, TagCount);
 
-	virtual void Activate() override;
-	virtual void EndAction() override;
+	UE_API virtual void Activate() override;
+	UE_API virtual void EndAction() override;
 
-	virtual void GameplayTagCallback(const FGameplayTag Tag, int32 NewCount);
+	UE_API virtual void GameplayTagCallback(const FGameplayTag Tag, int32 NewCount);
 
 	FGameplayTag Tag;
 	FDelegateHandle GameplayTagCountChangedHandle;
@@ -29,5 +31,7 @@ public:
 	 * If used in an ability graph, this async action will wait even after activation ends. It's recommended to use WaitGameplayTagCountChange instead.
 	 */
 	UFUNCTION(BlueprintCallable, Category = "Ability|Tasks", meta = (DefaultToSelf = "TargetActor", BlueprintInternalUseOnly = "TRUE"))
-	static UAbilityAsync_WaitGameplayTagCountChanged* WaitGameplayTagCountChangedOnActor(AActor* TargetActor, FGameplayTag Tag);
-};
\ No newline at end of file
+	static UE_API UAbilityAsync_WaitGameplayTagCountChanged* WaitGameplayTagCountChangedOnActor(AActor* TargetActor, FGameplayTag Tag);
+};
+
+#undef UE_API
diff --git a/Engine/Plugins/Runtime/GameplayAbilities/Source/GameplayAbilities/Public/Abilities/Async/AbilityAsync_WaitGameplayTagQuery.h b/Engine/Plugins/Runtime/GameplayAbilities/Source/GameplayAbilities/Public/Abilities/Async/AbilityAsync_WaitGameplayTagQuery.h
--- a/Engine/Plugins/Runtime/GameplayAbilities/Source/GameplayAbilities/Public/Abilities/Async/AbilityAsync_WaitGameplayTagQuery.h
+++ b/Engine/Plugins/Runtime/GameplayAbilities/Source/GameplayAbilities/Public/Abilities/Async/AbilityAsync_WaitGameplayTagQuery.h
@@ -6,13 +6,15 @@
 
 #include "AbilityAsync_WaitGameplayTagQuery.generated.h"
 
+#define UE_API GAMEPLAYABILITIES_API
+
 class UAbilitySystemComponent;
 
 DECLARE_DYNAMIC_MULTICAST_DELEGATE(FAsyncWaitGameplayTagQueryDelegate);
 
 /** This class defines an async node to wait on a gameplay tag query. */
-UCLASS()
-class GAMEPLAYABILITIES_API UAbilityAsync_WaitGameplayTagQuery : public UAbilityAsync
+UCLASS(MinimalAPI)
+class UAbilityAsync_WaitGameplayTagQuery : public UAbilityAsync
 {
 	GENERATED_BODY()
 
@@ -28,26 +30,26 @@ public:
 	 *  If used in an ability graph, this async action will wait even after activation ends. It's recommended to use WaitGameplayTagQuery instead.
 	 */
  	UFUNCTION(BlueprintCallable, Category = "Ability|Tasks", meta = (DefaultToSelf = "TargetActor", BlueprintInternalUseOnly = "TRUE"))
- 	static UAbilityAsync_WaitGameplayTagQuery* WaitGameplayTagQueryOnActor(AActor* TargetActor, 
+ 	static UE_API UAbilityAsync_WaitGameplayTagQuery* WaitGameplayTagQueryOnActor(AActor* TargetActor, 
 																		   const FGameplayTagQuery TagQuery, 
 																		   const EWaitGameplayTagQueryTriggerCondition TriggerCondition = EWaitGameplayTagQueryTriggerCondition::WhenTrue,
 																		   const bool bOnlyTriggerOnce=false);
 
 	/** Activates this AbilityAsync node. */
-	virtual void Activate() override;
+	UE_API virtual void Activate() override;
 
 protected:
 	/** Ends this AbilityAsync node, unregistering all tag changed delagates. */
-	virtual void EndAction() override;
+	UE_API virtual void EndAction() override;
 
 	/** Returns if this AbilityAsync node should currently broadcast it's delegates. */
-	virtual bool ShouldBroadcastDelegates() const;
+	UE_API virtual bool ShouldBroadcastDelegates() const;
 
 	/** This will update the tags in the TargetTags container to reflect that tags that are on the target ASC. */
-	void UpdateTargetTags(const FGameplayTag Tag, int32 NewCount);
+	UE_API void UpdateTargetTags(const FGameplayTag Tag, int32 NewCount);
 
 	/** This will evaluate the TargetTags using the given TagQuery, executing the Trigger delegate if needed. */
-	void EvaluateTagQuery();
+	UE_API void EvaluateTagQuery();
 
 	/** This is the tag query to evaluate for triggering this node. */
 	FGameplayTagQuery TagQuery;
@@ -70,4 +72,6 @@ protected:
 
 	/** This is the handles to the tag changed delegate for each gameplay tag in the TagQuery. */
 	TMap<FGameplayTag, FDelegateHandle> TagHandleMap;
-};
\ No newline at end of file
+};
+
+#undef UE_API
diff --git a/Engine/Plugins/Runtime/GameplayAbilities/Source/GameplayAbilities/Public/Abilities/GameplayAbility.h b/Engine/Plugins/Runtime/GameplayAbilities/Source/GameplayAbilities/Public/Abilities/GameplayAbility.h
--- a/Engine/Plugins/Runtime/GameplayAbilities/Source/GameplayAbilities/Public/Abilities/GameplayAbility.h
+++ b/Engine/Plugins/Runtime/GameplayAbilities/Source/GameplayAbilities/Public/Abilities/GameplayAbility.h
@@ -16,6 +16,8 @@
 #include "Net/Core/PushModel/PushModelMacros.h"
 #include "GameplayAbility.generated.h"
 
+#define UE_API GAMEPLAYABILITIES_API
+
 class UAbilitySystemComponent;
 class UAnimMontage;
 class UGameplayAbility;
@@ -104,8 +106,8 @@ struct FAbilityTriggerData
 };
 
 /** Abilities define custom gameplay logic that can be activated by players or external game logic */
-UCLASS(Blueprintable)
-class GAMEPLAYABILITIES_API UGameplayAbility : public UObject, public IGameplayTaskOwnerInterface
+UCLASS(Blueprintable, MinimalAPI)
+class UGameplayAbility : public UObject, public IGameplayTaskOwnerInterface
 {
 	GENERATED_UCLASS_BODY()
 	REPLICATED_BASE_CLASS(UGameplayAbility)
@@ -142,7 +144,7 @@ public:
 	// --------------------------------------
 
 	/** Returns how the ability is instanced when executed. This limits what an ability can do in its implementation. */
-	EGameplayAbilityInstancingPolicy::Type GetInstancingPolicy() const;
+	UE_API EGameplayAbilityInstancingPolicy::Type GetInstancingPolicy() const;
 
 	/** How an ability replicates state/events to everyone on the network */
 	EGameplayAbilityReplicationPolicy::Type GetReplicationPolicy() const
@@ -164,36 +166,36 @@ public:
 
 	/** Returns the actor info associated with this ability, has cached pointers to useful objects */
 	UFUNCTION(BlueprintCallable, Category=Ability)
-	FGameplayAbilityActorInfo GetActorInfo() const;
+	UE_API FGameplayAbilityActorInfo GetActorInfo() const;
 
 	/** Returns the actor that owns this ability, which may not have a physical location */
 	UFUNCTION(BlueprintCallable, Category = Ability)
-	AActor* GetOwningActorFromActorInfo() const;
+	UE_API AActor* GetOwningActorFromActorInfo() const;
 
 	/** Returns the physical actor that is executing this ability. May be null */
 	UFUNCTION(BlueprintCallable, Category = Ability)
-	AActor* GetAvatarActorFromActorInfo() const;
+	UE_API AActor* GetAvatarActorFromActorInfo() const;
 
 	/** Convenience method for abilities to get skeletal mesh component - useful for aiming abilities */
 	UFUNCTION(BlueprintCallable, DisplayName = "GetSkeletalMeshComponentFromActorInfo", Category = Ability)
-	USkeletalMeshComponent* GetOwningComponentFromActorInfo() const;
+	UE_API USkeletalMeshComponent* GetOwningComponentFromActorInfo() const;
 
 	/** Returns the AbilitySystemComponent that is activating this ability */
 	UFUNCTION(BlueprintCallable, Category = Ability)
-	UAbilitySystemComponent* GetAbilitySystemComponentFromActorInfo() const;
+	UE_API UAbilitySystemComponent* GetAbilitySystemComponentFromActorInfo() const;
 
 	UE_DEPRECATED(5.5, "Use GetAbilitySystemComponentFromActorInfo_Ensured")
-	UAbilitySystemComponent* GetAbilitySystemComponentFromActorInfo_Checked() const;
-	UAbilitySystemComponent* GetAbilitySystemComponentFromActorInfo_Ensured() const;
+	UE_API UAbilitySystemComponent* GetAbilitySystemComponentFromActorInfo_Checked() const;
+	UE_API UAbilitySystemComponent* GetAbilitySystemComponentFromActorInfo_Ensured() const;
 
 	/** The ability is considered to have these tags. */
-	const FGameplayTagContainer& GetAssetTags() const;
+	UE_API const FGameplayTagContainer& GetAssetTags() const;
 
 	/** Gets the current actor info bound to this ability - can only be called on instanced abilities. */
-	const FGameplayAbilityActorInfo* GetCurrentActorInfo() const;
+	UE_API const FGameplayAbilityActorInfo* GetCurrentActorInfo() const;
 
 	/** Gets the current activation info bound to this ability - can only be called on instanced abilities. */
-	FGameplayAbilityActivationInfo GetCurrentActivationInfo() const;
+	UE_API FGameplayAbilityActivationInfo GetCurrentActivationInfo() const;
 
 	/** Gets the current activation info bound to this ability - can only be called on instanced abilities. */
 	FGameplayAbilityActivationInfo& GetCurrentActivationInfoRef()
@@ -203,59 +205,59 @@ public:
 	}
 
 	/** Gets the current AbilitySpecHandle- can only be called on instanced abilities. */
-	FGameplayAbilitySpecHandle GetCurrentAbilitySpecHandle() const;
+	UE_API FGameplayAbilitySpecHandle GetCurrentAbilitySpecHandle() const;
 
 	/** Retrieves the actual AbilitySpec for this ability. Can only be called on instanced abilities. */
-	FGameplayAbilitySpec* GetCurrentAbilitySpec() const;
+	UE_API FGameplayAbilitySpec* GetCurrentAbilitySpec() const;
 
 	/** Retrieves the EffectContext of the GameplayEffect that granted this ability. Can only be called on instanced abilities. */
 	UFUNCTION(BlueprintCallable, Category = Ability)
-	FGameplayEffectContextHandle GetGrantedByEffectContext() const;
+	UE_API FGameplayEffectContextHandle GetGrantedByEffectContext() const;
 
 	/** Generates a GameplayEffectContextHandle from our owner and an optional TargetData.*/
 	UFUNCTION(BlueprintCallable, Category = Ability)
-	virtual FGameplayEffectContextHandle GetContextFromOwner(FGameplayAbilityTargetDataHandle OptionalTargetData) const;
+	UE_API virtual FGameplayEffectContextHandle GetContextFromOwner(FGameplayAbilityTargetDataHandle OptionalTargetData) const;
 
 	/** Returns an effect context, given a specified actor info */
-	virtual FGameplayEffectContextHandle MakeEffectContext(const FGameplayAbilitySpecHandle Handle, const FGameplayAbilityActorInfo *ActorInfo) const;
+	UE_API virtual FGameplayEffectContextHandle MakeEffectContext(const FGameplayAbilitySpecHandle Handle, const FGameplayAbilityActorInfo *ActorInfo) const;
 
 	/** Convenience method for abilities to get outgoing gameplay effect specs (for example, to pass on to projectiles to apply to whoever they hit) */
 	UFUNCTION(BlueprintCallable, Category=Ability)
-	FGameplayEffectSpecHandle MakeOutgoingGameplayEffectSpec(TSubclassOf<UGameplayEffect> GameplayEffectClass, float Level=1.f) const;
+	UE_API FGameplayEffectSpecHandle MakeOutgoingGameplayEffectSpec(TSubclassOf<UGameplayEffect> GameplayEffectClass, float Level=1.f) const;
 
 	/** Native version of above function */
-	virtual FGameplayEffectSpecHandle MakeOutgoingGameplayEffectSpec(const FGameplayAbilitySpecHandle Handle, const FGameplayAbilityActorInfo* ActorInfo, const FGameplayAbilityActivationInfo ActivationInfo, TSubclassOf<UGameplayEffect> GameplayEffectClass, float Level = 1.f) const;
+	UE_API virtual FGameplayEffectSpecHandle MakeOutgoingGameplayEffectSpec(const FGameplayAbilitySpecHandle Handle, const FGameplayAbilityActorInfo* ActorInfo, const FGameplayAbilityActivationInfo ActivationInfo, TSubclassOf<UGameplayEffect> GameplayEffectClass, float Level = 1.f) const;
 
 	/** Add the Ability's tags to the given GameplayEffectSpec. This is likely to be overridden per project. */
-	virtual void ApplyAbilityTagsToGameplayEffectSpec(FGameplayEffectSpec& Spec, FGameplayAbilitySpec* AbilitySpec) const;
+	UE_API virtual void ApplyAbilityTagsToGameplayEffectSpec(FGameplayEffectSpec& Spec, FGameplayAbilitySpec* AbilitySpec) const;
 
 	/** Returns true if the ability is currently active */
-	bool IsActive() const;
+	UE_API bool IsActive() const;
 
 	/** Is this ability triggered from TriggerData (or is it triggered explicitly through input/game code) */
-	bool IsTriggered() const;
+	UE_API bool IsTriggered() const;
 
 	/** Is this ability running on a a predicting client, this is false in single player */
-	bool IsPredictingClient() const;
+	UE_API bool IsPredictingClient() const;
 
 	/** True if this is on the server and is being executed for a non-local player, false in single player */
-	bool IsForRemoteClient() const;
+	UE_API bool IsForRemoteClient() const;
 
 	/** True if the owning actor is locally controlled, true in single player */
 	UFUNCTION(BlueprintCallable, BlueprintPure = false, Category = Ability, Meta = (ExpandBoolAsExecs = "ReturnValue"))
-	bool IsLocallyControlled() const;
+	UE_API bool IsLocallyControlled() const;
 
 	/** True if this is the server or single player */
-	bool HasAuthority(const FGameplayAbilityActivationInfo* ActivationInfo) const;
+	UE_API bool HasAuthority(const FGameplayAbilityActivationInfo* ActivationInfo) const;
 
 	UFUNCTION(BlueprintCallable, BlueprintPure = false, Category = Ability, DisplayName = "HasAuthority", Meta = (ExpandBoolAsExecs = "ReturnValue"))
-	bool K2_HasAuthority() const;
+	UE_API bool K2_HasAuthority() const;
 
 	/** True if we are authority or we have a valid prediciton key that is expected to work */
-	bool HasAuthorityOrPredictionKey(const FGameplayAbilityActorInfo* ActorInfo, const FGameplayAbilityActivationInfo* ActivationInfo) const;
+	UE_API bool HasAuthorityOrPredictionKey(const FGameplayAbilityActorInfo* ActorInfo, const FGameplayAbilityActivationInfo* ActivationInfo) const;
 
 	/** True if this has been instanced, always true for blueprints */
-	bool IsInstantiated() const;
+	UE_API bool IsInstantiated() const;
 
 	/** Notification that the ability has ended.  Set using TryActivateAbility. */
 	FOnGameplayAbilityEnded OnGameplayAbilityEnded;
@@ -277,54 +279,54 @@ public:
 	// --------------------------------------
 
 	/** Returns true if this ability can be activated right now. Has no side effects */
-	virtual bool CanActivateAbility(const FGameplayAbilitySpecHandle Handle, const FGameplayAbilityActorInfo* ActorInfo, const FGameplayTagContainer* SourceTags = nullptr, const FGameplayTagContainer* TargetTags = nullptr, OUT FGameplayTagContainer* OptionalRelevantTags = nullptr) const;
+	UE_API virtual bool CanActivateAbility(const FGameplayAbilitySpecHandle Handle, const FGameplayAbilityActorInfo* ActorInfo, const FGameplayTagContainer* SourceTags = nullptr, const FGameplayTagContainer* TargetTags = nullptr, OUT FGameplayTagContainer* OptionalRelevantTags = nullptr) const;
 
 	/** Returns true if this ability can be triggered right now. Has no side effects */
-	virtual bool ShouldAbilityRespondToEvent(const FGameplayAbilityActorInfo* ActorInfo, const FGameplayEventData* Payload) const;
+	UE_API virtual bool ShouldAbilityRespondToEvent(const FGameplayAbilityActorInfo* ActorInfo, const FGameplayEventData* Payload) const;
 
 	/** Returns true if an an ability should be activated */
-	virtual bool ShouldActivateAbility(ENetRole Role) const;
+	UE_API virtual bool ShouldActivateAbility(ENetRole Role) const;
 
 	/** Returns the time in seconds remaining on the currently active cooldown. */
 	UFUNCTION(BlueprintCallable, Category = Ability)
-	float GetCooldownTimeRemaining() const;
+	UE_API float GetCooldownTimeRemaining() const;
 
 	/** Returns the time in seconds remaining on the currently active cooldown. */
-	virtual float GetCooldownTimeRemaining(const FGameplayAbilityActorInfo* ActorInfo) const;
+	UE_API virtual float GetCooldownTimeRemaining(const FGameplayAbilityActorInfo* ActorInfo) const;
 
 	/** Returns the time in seconds remaining on the currently active cooldown and the original duration for this cooldown. */
-	virtual void GetCooldownTimeRemainingAndDuration(FGameplayAbilitySpecHandle Handle, const FGameplayAbilityActorInfo* ActorInfo, float& TimeRemaining, float& CooldownDuration) const;
+	UE_API virtual void GetCooldownTimeRemainingAndDuration(FGameplayAbilitySpecHandle Handle, const FGameplayAbilityActorInfo* ActorInfo, float& TimeRemaining, float& CooldownDuration) const;
 
 	/** Returns all tags that can put this ability into cooldown */
-	virtual const FGameplayTagContainer* GetCooldownTags() const;
+	UE_API virtual const FGameplayTagContainer* GetCooldownTags() const;
 	
 	/** Returns true if none of the ability's tags are blocked and if it doesn't have a "Blocking" tag and has all "Required" tags. */
-	virtual bool DoesAbilitySatisfyTagRequirements(const UAbilitySystemComponent& AbilitySystemComponent, const FGameplayTagContainer* SourceTags = nullptr, const FGameplayTagContainer* TargetTags = nullptr, OUT FGameplayTagContainer* OptionalRelevantTags = nullptr) const;
+	UE_API virtual bool DoesAbilitySatisfyTagRequirements(const UAbilitySystemComponent& AbilitySystemComponent, const FGameplayTagContainer* SourceTags = nullptr, const FGameplayTagContainer* TargetTags = nullptr, OUT FGameplayTagContainer* OptionalRelevantTags = nullptr) const;
 
 	/** Returns true if this ability is blocking other abilities */
-	virtual bool IsBlockingOtherAbilities() const;
+	UE_API virtual bool IsBlockingOtherAbilities() const;
 
 	/** Sets rather ability block flags are enabled or disabled. Only valid on instanced abilities */
 	UFUNCTION(BlueprintCallable, Category = Ability)
-	virtual void SetShouldBlockOtherAbilities(bool bShouldBlockAbilities);
+	UE_API virtual void SetShouldBlockOtherAbilities(bool bShouldBlockAbilities);
 
 	// --------------------------------------
 	//	CancelAbility
 	// --------------------------------------
 
 	/** Destroys instanced-per-execution abilities. Instance-per-actor abilities should 'reset'. Any active ability state tasks receive the 'OnAbilityStateInterrupted' event. Non instance abilities - what can we do? */
-	virtual void CancelAbility(const FGameplayAbilitySpecHandle Handle, const FGameplayAbilityActorInfo* ActorInfo, const FGameplayAbilityActivationInfo ActivationInfo, bool bReplicateCancelAbility);
+	UE_API virtual void CancelAbility(const FGameplayAbilitySpecHandle Handle, const FGameplayAbilityActorInfo* ActorInfo, const FGameplayAbilityActivationInfo ActivationInfo, bool bReplicateCancelAbility);
 
 	/** Call from Blueprint to cancel the ability naturally */
 	UFUNCTION(BlueprintCallable, Category = Ability, DisplayName = "CancelAbility", meta=(ScriptName = "CancelAbility"))
-	void K2_CancelAbility();
+	UE_API void K2_CancelAbility();
 
 	/** Returns true if this ability can be canceled */
-	virtual bool CanBeCanceled() const;
+	UE_API virtual bool CanBeCanceled() const;
 
 	/** Sets whether the ability should ignore cancel requests. Only valid on instanced abilities */
 	UFUNCTION(BlueprintCallable, Category=Ability)
-	virtual void SetCanBeCanceled(bool bCanBeCanceled);
+	UE_API virtual void SetCanBeCanceled(bool bCanBeCanceled);
 
 	// --------------------------------------
 	//	CommitAbility
@@ -332,55 +334,55 @@ public:
 
 	/** Attempts to commit the ability (spend resources, etc). This our last chance to fail. Child classes that override ActivateAbility must call this themselves! */
 	UFUNCTION(BlueprintCallable, Category = Ability, DisplayName = "CommitAbility", meta=(ScriptName = "CommitAbility"))
-	virtual bool K2_CommitAbility();
+	UE_API virtual bool K2_CommitAbility();
 
 	/** Attempts to commit the ability's cooldown only. If BroadcastCommitEvent is true, it will broadcast the commit event that tasks like WaitAbilityCommit are listening for. */
 	UFUNCTION(BlueprintCallable, Category = Ability, DisplayName = "CommitAbilityCooldown", meta=(ScriptName = "CommitAbilityCooldown"))
-	virtual bool K2_CommitAbilityCooldown(bool BroadcastCommitEvent=false, bool ForceCooldown=false);
+	UE_API virtual bool K2_CommitAbilityCooldown(bool BroadcastCommitEvent=false, bool ForceCooldown=false);
 
 	/** Attempts to commit the ability's cost only. If BroadcastCommitEvent is true, it will broadcast the commit event that tasks like WaitAbilityCommit are listening for. */
 	UFUNCTION(BlueprintCallable, Category = Ability, DisplayName = "CommitAbilityCost", meta=(ScriptName = "CommitAbilityCost"))
-	virtual bool K2_CommitAbilityCost(bool BroadcastCommitEvent=false);
+	UE_API virtual bool K2_CommitAbilityCost(bool BroadcastCommitEvent=false);
 
 	/** Checks the ability's cooldown, but does not apply it.*/
 	UFUNCTION(BlueprintCallable, Category = Ability, DisplayName = "CheckAbilityCooldown", meta=(ScriptName = "CheckAbilityCooldown"))
-	virtual bool K2_CheckAbilityCooldown();
+	UE_API virtual bool K2_CheckAbilityCooldown();
 
 	/** Checks the ability's cost, but does not apply it. */
 	UFUNCTION(BlueprintCallable, Category = Ability, DisplayName = "CheckAbilityCost", meta=(ScriptName = "CheckAbilityCost"))
-	virtual bool K2_CheckAbilityCost();
+	UE_API virtual bool K2_CheckAbilityCost();
 
-	virtual bool CommitAbility(const FGameplayAbilitySpecHandle Handle, const FGameplayAbilityActorInfo* ActorInfo, const FGameplayAbilityActivationInfo ActivationInfo, OUT FGameplayTagContainer* OptionalRelevantTags = nullptr);
-	virtual bool CommitAbilityCooldown(const FGameplayAbilitySpecHandle Handle, const FGameplayAbilityActorInfo* ActorInfo, const FGameplayAbilityActivationInfo ActivationInfo, const bool ForceCooldown, OUT FGameplayTagContainer* OptionalRelevantTags = nullptr);
-	virtual bool CommitAbilityCost(const FGameplayAbilitySpecHandle Handle, const FGameplayAbilityActorInfo* ActorInfo, const FGameplayAbilityActivationInfo ActivationInfo, OUT FGameplayTagContainer* OptionalRelevantTags = nullptr);
+	UE_API virtual bool CommitAbility(const FGameplayAbilitySpecHandle Handle, const FGameplayAbilityActorInfo* ActorInfo, const FGameplayAbilityActivationInfo ActivationInfo, OUT FGameplayTagContainer* OptionalRelevantTags = nullptr);
+	UE_API virtual bool CommitAbilityCooldown(const FGameplayAbilitySpecHandle Handle, const FGameplayAbilityActorInfo* ActorInfo, const FGameplayAbilityActivationInfo ActivationInfo, const bool ForceCooldown, OUT FGameplayTagContainer* OptionalRelevantTags = nullptr);
+	UE_API virtual bool CommitAbilityCost(const FGameplayAbilitySpecHandle Handle, const FGameplayAbilityActorInfo* ActorInfo, const FGameplayAbilityActivationInfo ActivationInfo, OUT FGameplayTagContainer* OptionalRelevantTags = nullptr);
 
 	/** The last chance to fail before committing, this will usually be the same as CanActivateAbility. Some abilities may need to do extra checks here if they are consuming extra stuff in CommitExecute */
-	virtual bool CommitCheck(const FGameplayAbilitySpecHandle Handle, const FGameplayAbilityActorInfo* ActorInfo, const FGameplayAbilityActivationInfo ActivationInfo, OUT FGameplayTagContainer* OptionalRelevantTags = nullptr);
+	UE_API virtual bool CommitCheck(const FGameplayAbilitySpecHandle Handle, const FGameplayAbilityActorInfo* ActorInfo, const FGameplayAbilityActivationInfo ActivationInfo, OUT FGameplayTagContainer* OptionalRelevantTags = nullptr);
 
 	/** BP event called from CommitAbility */
 	UFUNCTION(BlueprintImplementableEvent, Category = Ability, DisplayName = "CommitExecute", meta = (ScriptName = "CommitExecute"))
-	void K2_CommitExecute();
+	UE_API void K2_CommitExecute();
 
 	/** Does the commit atomically (consume resources, do cooldowns, etc) */
-	virtual void CommitExecute(const FGameplayAbilitySpecHandle Handle, const FGameplayAbilityActorInfo* ActorInfo, const FGameplayAbilityActivationInfo ActivationInfo);
+	UE_API virtual void CommitExecute(const FGameplayAbilitySpecHandle Handle, const FGameplayAbilityActorInfo* ActorInfo, const FGameplayAbilityActivationInfo ActivationInfo);
 
 	/** Returns the gameplay effect used to determine cooldown */
-	virtual UGameplayEffect* GetCooldownGameplayEffect() const;
+	UE_API virtual UGameplayEffect* GetCooldownGameplayEffect() const;
 
 	/** Returns the gameplay effect used to apply cost */
-	virtual UGameplayEffect* GetCostGameplayEffect() const;
+	UE_API virtual UGameplayEffect* GetCostGameplayEffect() const;
 
 	/** Checks cooldown. returns true if we can be used again. False if not */
-	virtual bool CheckCooldown(const FGameplayAbilitySpecHandle Handle, const FGameplayAbilityActorInfo* ActorInfo, OUT FGameplayTagContainer* OptionalRelevantTags = nullptr) const;
+	UE_API virtual bool CheckCooldown(const FGameplayAbilitySpecHandle Handle, const FGameplayAbilityActorInfo* ActorInfo, OUT FGameplayTagContainer* OptionalRelevantTags = nullptr) const;
 
 	/** Applies CooldownGameplayEffect to the target */
-	virtual void ApplyCooldown(const FGameplayAbilitySpecHandle Handle, const FGameplayAbilityActorInfo* ActorInfo, const FGameplayAbilityActivationInfo ActivationInfo) const;
+	UE_API virtual void ApplyCooldown(const FGameplayAbilitySpecHandle Handle, const FGameplayAbilityActorInfo* ActorInfo, const FGameplayAbilityActivationInfo ActivationInfo) const;
 
 	/** Checks cost. returns true if we can pay for the ability. False if not */
-	virtual bool CheckCost(const FGameplayAbilitySpecHandle Handle, const FGameplayAbilityActorInfo* ActorInfo, OUT FGameplayTagContainer* OptionalRelevantTags = nullptr) const;
+	UE_API virtual bool CheckCost(const FGameplayAbilitySpecHandle Handle, const FGameplayAbilityActorInfo* ActorInfo, OUT FGameplayTagContainer* OptionalRelevantTags = nullptr) const;
 
 	/** Applies the ability's cost to the target */
-	virtual void ApplyCost(const FGameplayAbilitySpecHandle Handle, const FGameplayAbilityActorInfo* ActorInfo, const FGameplayAbilityActivationInfo ActivationInfo) const;
+	UE_API virtual void ApplyCost(const FGameplayAbilitySpecHandle Handle, const FGameplayAbilityActorInfo* ActorInfo, const FGameplayAbilityActivationInfo ActivationInfo) const;
 
 	// --------------------------------------
 	//	Input
@@ -402,14 +404,14 @@ public:
 
 	/** Returns the currently playing montage for this ability, if any */
 	UFUNCTION(BlueprintCallable, Category = Animation)
-	UAnimMontage* GetCurrentMontage() const;
+	UE_API UAnimMontage* GetCurrentMontage() const;
 
 	/** Call to set/get the current montage from a montage task. Set to allow hooking up montage events to ability events */
-	virtual void SetCurrentMontage(class UAnimMontage* InCurrentMontage);
+	UE_API virtual void SetCurrentMontage(class UAnimMontage* InCurrentMontage);
 
 	/** Movement Sync */
 	UE_DEPRECATED(5.3, "This serves no purpose and will be removed in future engine versions")
-	virtual void SetMovementSyncPoint(FName SyncName);
+	UE_API virtual void SetMovementSyncPoint(FName SyncName);
 
 	// ----------------------------------------------------------------------------------------------------------------
 	//	Ability Levels and source objects 
@@ -417,50 +419,50 @@ public:
 	
 	/** Returns current level of the Ability */
 	UFUNCTION(BlueprintCallable, Category = Ability)
-	int32 GetAbilityLevel() const;
+	UE_API int32 GetAbilityLevel() const;
 
 	/** Returns current ability level for non instanced abilities. You must call this version in these contexts! */
-	int32 GetAbilityLevel(FGameplayAbilitySpecHandle Handle, const FGameplayAbilityActorInfo* ActorInfo) const;
+	UE_API int32 GetAbilityLevel(FGameplayAbilitySpecHandle Handle, const FGameplayAbilityActorInfo* ActorInfo) const;
 
 	/** Returns current ability level for non instanced abilities. You must call this version in these contexts! */
 	UFUNCTION(BlueprintCallable, Category = Ability, meta = (DisplayName = "GetAbilityLevelNonInstanced", ReturnDisplayName = "AbilityLevel"))
-	int32 GetAbilityLevel_BP(FGameplayAbilitySpecHandle Handle, const FGameplayAbilityActorInfo& ActorInfo) const;
+	UE_API int32 GetAbilityLevel_BP(FGameplayAbilitySpecHandle Handle, const FGameplayAbilityActorInfo& ActorInfo) const;
 
 	/** Retrieves the SourceObject associated with this ability. Can only be called on instanced abilities. */
 	UFUNCTION(BlueprintCallable, Category = Ability)
-	UObject* GetCurrentSourceObject() const;
+	UE_API UObject* GetCurrentSourceObject() const;
 
 	/** Retrieves the SourceObject associated with this ability. Callable on non instanced */
-	UObject* GetSourceObject(FGameplayAbilitySpecHandle Handle, const FGameplayAbilityActorInfo* ActorInfo) const;
+	UE_API UObject* GetSourceObject(FGameplayAbilitySpecHandle Handle, const FGameplayAbilityActorInfo* ActorInfo) const;
 
 	/** Retrieves the SourceObject associated with this ability. Callable on non instanced */
 	UFUNCTION(BlueprintCallable, Category = Ability, meta = (DisplayName = "GetSourceObjectNonInstanced", ReturnDisplayName = "SourceObject"))
-	UObject* GetSourceObject_BP(FGameplayAbilitySpecHandle Handle, const FGameplayAbilityActorInfo& ActorInfo) const;
+	UE_API UObject* GetSourceObject_BP(FGameplayAbilitySpecHandle Handle, const FGameplayAbilityActorInfo& ActorInfo) const;
 
 	// --------------------------------------
 	//	Interaction with ability system component
 	// --------------------------------------
 
 	/** Called by ability system component to inform this ability instance the remote instance was ended */
-	virtual void SetRemoteInstanceHasEnded();
+	UE_API virtual void SetRemoteInstanceHasEnded();
 
 	/** Called to inform the ability that the AvatarActor has been replaced. If the ability is dependent on avatar state, it may want to end itself. */
-	virtual void NotifyAvatarDestroyed();
+	UE_API virtual void NotifyAvatarDestroyed();
 
 	/** Called to inform the ability that a task is waiting for the player to do something */
-	virtual void NotifyAbilityTaskWaitingOnPlayerData(class UAbilityTask* AbilityTask);
+	UE_API virtual void NotifyAbilityTaskWaitingOnPlayerData(class UAbilityTask* AbilityTask);
 
 	/** Called to inform the ability that a task is waiting for the player's avatar to do something in world */
-	virtual void NotifyAbilityTaskWaitingOnAvatar(class UAbilityTask* AbilityTask);
+	UE_API virtual void NotifyAbilityTaskWaitingOnAvatar(class UAbilityTask* AbilityTask);
 
 	/** Called when the ability is given to an AbilitySystemComponent */
-	virtual void OnGiveAbility(const FGameplayAbilityActorInfo* ActorInfo, const FGameplayAbilitySpec& Spec);
+	UE_API virtual void OnGiveAbility(const FGameplayAbilityActorInfo* ActorInfo, const FGameplayAbilitySpec& Spec);
 
 	/** Called when the ability is removed from an AbilitySystemComponent */
 	virtual void OnRemoveAbility(const FGameplayAbilityActorInfo* ActorInfo, const FGameplayAbilitySpec& Spec) {}
 
 	/** Called when the avatar actor is set/changes */
-	virtual void OnAvatarSet(const FGameplayAbilityActorInfo* ActorInfo, const FGameplayAbilitySpec& Spec);
+	UE_API virtual void OnAvatarSet(const FGameplayAbilityActorInfo* ActorInfo, const FGameplayAbilitySpec& Spec);
 
 	/** Takes in the ability spec and checks if we should allow replication on the ability spec, this will NOT stop replication of the ability UObject just the spec inside the UAbilitySystemComponenet ActivatableAbilities for this ability */
 	virtual bool ShouldReplicateAbilitySpec(const FGameplayAbilitySpec& AbilitySpec) const
@@ -473,18 +475,18 @@ public:
 	 * This returns instantly and has no other side effects other than clearing the current prediction key.
 	 */ 
 	UFUNCTION(BlueprintCallable, Category = Ability)
-	void InvalidateClientPredictionKey() const;
+	UE_API void InvalidateClientPredictionKey() const;
 
 	/** Removes the GameplayEffect that granted this ability. Can only be called on instanced abilities. */
 	UFUNCTION(BlueprintCallable, Category = Ability)
-	virtual void RemoveGrantedByEffect();
+	UE_API virtual void RemoveGrantedByEffect();
 
 	/** Adds a debug message to display to the user */
-	void AddAbilityTaskDebugMessage(UGameplayTask* AbilityTask, FString DebugMessage);
+	UE_API void AddAbilityTaskDebugMessage(UGameplayTask* AbilityTask, FString DebugMessage);
 
 #if WITH_EDITOR
 	/** Allow modification of the AssetTags (AbilityTags) while in editor */
-	FGameplayTagContainer& EditorGetAssetTags();
+	UE_API FGameplayTagContainer& EditorGetAssetTags();
 #endif // WITH_EDITOR
 
 	// --------------------------------------
@@ -507,32 +509,32 @@ public:
 	// --------------------------------------
 	//	UObject overrides
 	// --------------------------------------	
-	virtual UWorld* GetWorld() const override;
-	virtual int32 GetFunctionCallspace(UFunction* Function, FFrame* Stack) override;
-	virtual bool CallRemoteFunction(UFunction* Function, void* Parameters, FOutParmRec* OutParms, FFrame* Stack) override;
-	virtual bool IsSupportedForNetworking() const override;
+	UE_API virtual UWorld* GetWorld() const override;
+	UE_API virtual int32 GetFunctionCallspace(UFunction* Function, FFrame* Stack) override;
+	UE_API virtual bool CallRemoteFunction(UFunction* Function, void* Parameters, FOutParmRec* OutParms, FFrame* Stack) override;
+	UE_API virtual bool IsSupportedForNetworking() const override;
 
 #if WITH_EDITOR
-	virtual EDataValidationResult IsDataValid(FDataValidationContext& Context) const override;
+	UE_API virtual EDataValidationResult IsDataValid(FDataValidationContext& Context) const override;
 #endif
 
 	/** Overridden to allow Blueprint replicated properties to work */
-	virtual void GetLifetimeReplicatedProps(TArray< class FLifetimeProperty >& OutLifetimeProps) const;
+	UE_API virtual void GetLifetimeReplicatedProps(TArray< class FLifetimeProperty >& OutLifetimeProps) const;
 
 #if UE_WITH_IRIS
 	/** Register all replication fragments */
-	virtual void RegisterReplicationFragments(UE::Net::FFragmentRegistrationContext& Context, UE::Net::EFragmentRegistrationFlags RegistrationFlags) override;
+	UE_API virtual void RegisterReplicationFragments(UE::Net::FFragmentRegistrationContext& Context, UE::Net::EFragmentRegistrationFlags RegistrationFlags) override;
 #endif // UE_WITH_IRIS
 
 	// --------------------------------------
 	//	IGameplayTaskOwnerInterface
 	// --------------------------------------	
-	virtual UGameplayTasksComponent* GetGameplayTasksComponent(const UGameplayTask& Task) const override;
-	virtual AActor* GetGameplayTaskOwner(const UGameplayTask* Task) const override;
-	virtual AActor* GetGameplayTaskAvatar(const UGameplayTask* Task) const override;
-	virtual void OnGameplayTaskInitialized(UGameplayTask& Task) override;
-	virtual void OnGameplayTaskActivated(UGameplayTask& Task) override;
-	virtual void OnGameplayTaskDeactivated(UGameplayTask& Task) override;
+	UE_API virtual UGameplayTasksComponent* GetGameplayTasksComponent(const UGameplayTask& Task) const override;
+	UE_API virtual AActor* GetGameplayTaskOwner(const UGameplayTask* Task) const override;
+	UE_API virtual AActor* GetGameplayTaskAvatar(const UGameplayTask* Task) const override;
+	UE_API virtual void OnGameplayTaskInitialized(UGameplayTask& Task) override;
+	UE_API virtual void OnGameplayTaskActivated(UGameplayTask& Task) override;
+	UE_API virtual void OnGameplayTaskDeactivated(UGameplayTask& Task) override;
 
 protected:
 
@@ -542,7 +544,7 @@ protected:
 	 * At runtime, the AbilitySpec is queried through a Gameplay Ability's CDO for its AbilityTags which can be a combination of these
 	 * Asset Tags and specifically granted DynamicAbilityTags (all instances generated from an AbilitySpec are expected to share the same AbilityTags).
 	 */
-	void SetAssetTags(const FGameplayTagContainer& InAbilityTags);
+	UE_API void SetAssetTags(const FGameplayTagContainer& InAbilityTags);
 
 	// --------------------------------------
 	//	ShouldAbilityRespondToEvent
@@ -550,13 +552,13 @@ protected:
 
 	/** Returns true if this ability can be activated right now. Has no side effects */
 	UFUNCTION(BlueprintImplementableEvent, Category = Ability, DisplayName = "ShouldAbilityRespondToEvent", meta=(ScriptName = "ShouldAbilityRespondToEvent"))
-	bool K2_ShouldAbilityRespondToEvent(FGameplayAbilityActorInfo ActorInfo, FGameplayEventData Payload) const;
+	UE_API bool K2_ShouldAbilityRespondToEvent(FGameplayAbilityActorInfo ActorInfo, FGameplayEventData Payload) const;
 
 	bool bHasBlueprintShouldAbilityRespondToEvent;
 
 	/** Sends a gameplay event, also creates a prediction window */
 	UFUNCTION(BlueprintCallable, Category = Ability)
-	virtual void SendGameplayEvent(UPARAM(meta=(GameplayTagFilter="GameplayEventTagsCategory")) FGameplayTag EventTag, FGameplayEventData Payload);
+	UE_API virtual void SendGameplayEvent(UPARAM(meta=(GameplayTagFilter="GameplayEventTagsCategory")) FGameplayTag EventTag, FGameplayEventData Payload);
 
 	// --------------------------------------
 	//	CanActivate
@@ -564,7 +566,7 @@ protected:
 	
 	/** Returns true if this ability can be activated right now. Has no side effects */
 	UFUNCTION(BlueprintImplementableEvent, Category = Ability, DisplayName="CanActivateAbility", meta=(ScriptName="CanActivateAbility"))
-	bool K2_CanActivateAbility(FGameplayAbilityActorInfo ActorInfo, const FGameplayAbilitySpecHandle Handle, FGameplayTagContainer& RelevantTags) const;
+	UE_API bool K2_CanActivateAbility(FGameplayAbilityActorInfo ActorInfo, const FGameplayAbilitySpecHandle Handle, FGameplayTagContainer& RelevantTags) const;
 
 	bool bHasBlueprintCanUse;
 
@@ -584,46 +586,46 @@ protected:
 	 *  
 	 */
 	UFUNCTION(BlueprintImplementableEvent, Category = Ability, DisplayName = "ActivateAbility", meta=(ScriptName = "ActivateAbility"))
-	void K2_ActivateAbility();
+	UE_API void K2_ActivateAbility();
 
 	UFUNCTION(BlueprintImplementableEvent, Category = Ability, DisplayName = "ActivateAbilityFromEvent", meta=(ScriptName = "ActivateAbilityFromEvent"))
-	void K2_ActivateAbilityFromEvent(const FGameplayEventData& EventData);
+	UE_API void K2_ActivateAbilityFromEvent(const FGameplayEventData& EventData);
 
 	bool bHasBlueprintActivate;
 	bool bHasBlueprintActivateFromEvent;
 
 	/** Actually activate ability, do not call this directly */
-	virtual void ActivateAbility(const FGameplayAbilitySpecHandle Handle, const FGameplayAbilityActorInfo* ActorInfo, const FGameplayAbilityActivationInfo ActivationInfo, const FGameplayEventData* TriggerEventData);
+	UE_API virtual void ActivateAbility(const FGameplayAbilitySpecHandle Handle, const FGameplayAbilityActorInfo* ActorInfo, const FGameplayAbilityActivationInfo ActivationInfo, const FGameplayEventData* TriggerEventData);
 
 	/** Do boilerplate init stuff and then call ActivateAbility */
-	virtual void PreActivate(const FGameplayAbilitySpecHandle Handle, const FGameplayAbilityActorInfo* ActorInfo, const FGameplayAbilityActivationInfo ActivationInfo, FOnGameplayAbilityEnded::FDelegate* OnGameplayAbilityEndedDelegate, const FGameplayEventData* TriggerEventData = nullptr);
+	UE_API virtual void PreActivate(const FGameplayAbilitySpecHandle Handle, const FGameplayAbilityActorInfo* ActorInfo, const FGameplayAbilityActivationInfo ActivationInfo, FOnGameplayAbilityEnded::FDelegate* OnGameplayAbilityEndedDelegate, const FGameplayEventData* TriggerEventData = nullptr);
 
 	/** Executes PreActivate and ActivateAbility */
-	void CallActivateAbility(const FGameplayAbilitySpecHandle Handle, const FGameplayAbilityActorInfo* ActorInfo, const FGameplayAbilityActivationInfo ActivationInfo, FOnGameplayAbilityEnded::FDelegate* OnGameplayAbilityEndedDelegate = nullptr, const FGameplayEventData* TriggerEventData = nullptr);
+	UE_API void CallActivateAbility(const FGameplayAbilitySpecHandle Handle, const FGameplayAbilityActorInfo* ActorInfo, const FGameplayAbilityActivationInfo ActivationInfo, FOnGameplayAbilityEnded::FDelegate* OnGameplayAbilityEndedDelegate = nullptr, const FGameplayEventData* TriggerEventData = nullptr);
 
 	/** Called on a predictive ability when the server confirms its execution */
-	virtual void ConfirmActivateSucceed();
+	UE_API virtual void ConfirmActivateSucceed();
 
 	// -------------------------------------
 	//	EndAbility
 	// -------------------------------------
 	/** Call from blueprints to forcibly end the ability without canceling it. This will replicate the end ability to the client or server which can interrupt tasks */
 	UFUNCTION(BlueprintCallable, Category = Ability, DisplayName="End Ability", meta=(ScriptName = "EndAbility"))
-	virtual void K2_EndAbility();
+	UE_API virtual void K2_EndAbility();
 
 	/** Call from blueprints to end the ability naturally. This will only end predicted abilities locally, allowing it end naturally on the client or server */
 	UFUNCTION(BlueprintCallable, Category = Ability, DisplayName = "End Ability Locally", meta = (ScriptName = "EndAbilityLocally"))
-	virtual void K2_EndAbilityLocally();
+	UE_API virtual void K2_EndAbilityLocally();
 
 	/** Blueprint event, will be called if an ability ends normally or abnormally */
 	UFUNCTION(BlueprintImplementableEvent, Category = Ability, DisplayName = "OnEndAbility", meta=(ScriptName = "OnEndAbility"))
-	void K2_OnEndAbility(bool bWasCancelled);
+	UE_API void K2_OnEndAbility(bool bWasCancelled);
 
 	/** Check if the ability can be ended */
-	bool IsEndAbilityValid(const FGameplayAbilitySpecHandle Handle, const FGameplayAbilityActorInfo* ActorInfo) const;
+	UE_API bool IsEndAbilityValid(const FGameplayAbilitySpecHandle Handle, const FGameplayAbilityActorInfo* ActorInfo) const;
 
 	/** Native function, called if an ability ends normally or abnormally. If bReplicate is set to true, try to replicate the ending to the client/server */
-	virtual void EndAbility(const FGameplayAbilitySpecHandle Handle, const FGameplayAbilityActorInfo* ActorInfo, const FGameplayAbilityActivationInfo ActivationInfo, bool bReplicateEndAbility, bool bWasCancelled);
+	UE_API virtual void EndAbility(const FGameplayAbilitySpecHandle Handle, const FGameplayAbilityActorInfo* ActorInfo, const FGameplayAbilityActivationInfo ActivationInfo, bool bReplicateEndAbility, bool bWasCancelled);
 
 	// -------------------------------------
 	//  Apply Gameplay effects to Self
@@ -631,16 +633,16 @@ protected:
 
 	/** Apply a gameplay effect to the owner of this ability */
 	UFUNCTION(BlueprintCallable, Category = Ability, DisplayName="ApplyGameplayEffectToOwner", meta=(ScriptName="ApplyGameplayEffectToOwner"))
-	FActiveGameplayEffectHandle BP_ApplyGameplayEffectToOwner(TSubclassOf<UGameplayEffect> GameplayEffectClass, int32 GameplayEffectLevel = 1, int32 Stacks = 1);
+	UE_API FActiveGameplayEffectHandle BP_ApplyGameplayEffectToOwner(TSubclassOf<UGameplayEffect> GameplayEffectClass, int32 GameplayEffectLevel = 1, int32 Stacks = 1);
 
 	/** Non blueprintcallable, safe to call on CDO/NonInstance abilities */
-	FActiveGameplayEffectHandle ApplyGameplayEffectToOwner(const FGameplayAbilitySpecHandle Handle, const FGameplayAbilityActorInfo* ActorInfo, const FGameplayAbilityActivationInfo ActivationInfo, const UGameplayEffect* GameplayEffect, float GameplayEffectLevel, int32 Stacks = 1) const;
+	UE_API FActiveGameplayEffectHandle ApplyGameplayEffectToOwner(const FGameplayAbilitySpecHandle Handle, const FGameplayAbilityActorInfo* ActorInfo, const FGameplayAbilityActivationInfo ActivationInfo, const UGameplayEffect* GameplayEffect, float GameplayEffectLevel, int32 Stacks = 1) const;
 
 	/** Apply a previously created gameplay effect spec to the owner of this ability */
 	UFUNCTION(BlueprintCallable, Category = Ability, DisplayName = "ApplyGameplayEffectSpecToOwner", meta=(ScriptName = "ApplyGameplayEffectSpecToOwner"))
-	FActiveGameplayEffectHandle K2_ApplyGameplayEffectSpecToOwner(const FGameplayEffectSpecHandle EffectSpecHandle);
+	UE_API FActiveGameplayEffectHandle K2_ApplyGameplayEffectSpecToOwner(const FGameplayEffectSpecHandle EffectSpecHandle);
 
-	FActiveGameplayEffectHandle ApplyGameplayEffectSpecToOwner(const FGameplayAbilitySpecHandle AbilityHandle, const FGameplayAbilityActorInfo* ActorInfo, const FGameplayAbilityActivationInfo ActivationInfo, const FGameplayEffectSpecHandle SpecHandle) const;
+	UE_API FActiveGameplayEffectHandle ApplyGameplayEffectSpecToOwner(const FGameplayAbilitySpecHandle AbilityHandle, const FGameplayAbilityActorInfo* ActorInfo, const FGameplayAbilityActivationInfo ActivationInfo, const FGameplayEffectSpecHandle SpecHandle) const;
 
 	// -------------------------------------
 	//  Apply Gameplay effects to Target
@@ -648,16 +650,16 @@ protected:
 
 	/** Apply a gameplay effect to a Target */
 	UFUNCTION(BlueprintCallable, Category = Ability, DisplayName = "ApplyGameplayEffectToTarget", meta=(ScriptName = "ApplyGameplayEffectToTarget"))
-	TArray<FActiveGameplayEffectHandle> BP_ApplyGameplayEffectToTarget(FGameplayAbilityTargetDataHandle TargetData, TSubclassOf<UGameplayEffect> GameplayEffectClass, int32 GameplayEffectLevel = 1, int32 Stacks = 1);
+	UE_API TArray<FActiveGameplayEffectHandle> BP_ApplyGameplayEffectToTarget(FGameplayAbilityTargetDataHandle TargetData, TSubclassOf<UGameplayEffect> GameplayEffectClass, int32 GameplayEffectLevel = 1, int32 Stacks = 1);
 
 	/** Non blueprintcallable, safe to call on CDO/NonInstance abilities */
-	TArray<FActiveGameplayEffectHandle> ApplyGameplayEffectToTarget(const FGameplayAbilitySpecHandle Handle, const FGameplayAbilityActorInfo* ActorInfo, const FGameplayAbilityActivationInfo ActivationInfo, const FGameplayAbilityTargetDataHandle& Target, TSubclassOf<UGameplayEffect> GameplayEffectClass, float GameplayEffectLevel, int32 Stacks = 1) const;
+	UE_API TArray<FActiveGameplayEffectHandle> ApplyGameplayEffectToTarget(const FGameplayAbilitySpecHandle Handle, const FGameplayAbilityActorInfo* ActorInfo, const FGameplayAbilityActivationInfo ActivationInfo, const FGameplayAbilityTargetDataHandle& Target, TSubclassOf<UGameplayEffect> GameplayEffectClass, float GameplayEffectLevel, int32 Stacks = 1) const;
 
 	/** Apply a previously created gameplay effect spec to a target */
 	UFUNCTION(BlueprintCallable, Category = Ability, DisplayName = "ApplyGameplayEffectSpecToTarget", meta=(ScriptName = "ApplyGameplayEffectSpecToTarget"))
-	TArray<FActiveGameplayEffectHandle> K2_ApplyGameplayEffectSpecToTarget(const FGameplayEffectSpecHandle EffectSpecHandle, FGameplayAbilityTargetDataHandle TargetData);
+	UE_API TArray<FActiveGameplayEffectHandle> K2_ApplyGameplayEffectSpecToTarget(const FGameplayEffectSpecHandle EffectSpecHandle, FGameplayAbilityTargetDataHandle TargetData);
 
-	TArray<FActiveGameplayEffectHandle> ApplyGameplayEffectSpecToTarget(const FGameplayAbilitySpecHandle AbilityHandle, const FGameplayAbilityActorInfo* ActorInfo, const FGameplayAbilityActivationInfo ActivationInfo, const FGameplayEffectSpecHandle SpecHandle, const FGameplayAbilityTargetDataHandle& TargetData) const;
+	UE_API TArray<FActiveGameplayEffectHandle> ApplyGameplayEffectSpecToTarget(const FGameplayAbilitySpecHandle AbilityHandle, const FGameplayAbilityActorInfo* ActorInfo, const FGameplayAbilityActivationInfo ActivationInfo, const FGameplayEffectSpecHandle SpecHandle, const FGameplayAbilityTargetDataHandle& TargetData) const;
 
 	// -------------------------------------
 	//  Remove Gameplay effects from Self
@@ -665,15 +667,15 @@ protected:
 	
 	/** Removes GameplayEffects from owner which match the given asset level tags */
 	UFUNCTION(BlueprintCallable, Category = Ability, DisplayName="RemoveGameplayEffectFromOwnerWithAssetTags", meta=(ScriptName="RemoveGameplayEffectFromOwnerWithAssetTags"))
-	void BP_RemoveGameplayEffectFromOwnerWithAssetTags(FGameplayTagContainer WithAssetTags, int32 StacksToRemove = -1);
+	UE_API void BP_RemoveGameplayEffectFromOwnerWithAssetTags(FGameplayTagContainer WithAssetTags, int32 StacksToRemove = -1);
 
 	/** Removes GameplayEffects from owner which grant the given tags */
 	UFUNCTION(BlueprintCallable, Category = Ability, DisplayName="RemoveGameplayEffectFromOwnerWithGrantedTags", meta=(ScriptName="RemoveGameplayEffectFromOwnerWithGrantedTags"))
-	void BP_RemoveGameplayEffectFromOwnerWithGrantedTags(FGameplayTagContainer WithGrantedTags, int32 StacksToRemove = -1);
+	UE_API void BP_RemoveGameplayEffectFromOwnerWithGrantedTags(FGameplayTagContainer WithGrantedTags, int32 StacksToRemove = -1);
 
 	/** Removes GameplayEffect from owner that match the given handle */
 	UFUNCTION(BlueprintCallable, Category = Ability, DisplayName = "RemoveGameplayEffectFromOwnerWithHandle", meta=(ScriptName = "RemoveGameplayEffectFromOwnerWithHandle"))
-	void BP_RemoveGameplayEffectFromOwnerWithHandle(FActiveGameplayEffectHandle Handle, int32 StacksToRemove = -1);
+	UE_API void BP_RemoveGameplayEffectFromOwnerWithHandle(FActiveGameplayEffectHandle Handle, int32 StacksToRemove = -1);
 
 	// -------------------------------------
 	//	GameplayCue
@@ -682,23 +684,23 @@ protected:
 
 	/** Invoke a gameplay cue on the ability owner */
 	UFUNCTION(BlueprintCallable, Category = Ability, meta=(GameplayTagFilter="GameplayCue"), DisplayName="Execute GameplayCue On Owner", meta=(ScriptName="ExecuteGameplayCue"))
-	virtual void K2_ExecuteGameplayCue(FGameplayTag GameplayCueTag, FGameplayEffectContextHandle Context);
+	UE_API virtual void K2_ExecuteGameplayCue(FGameplayTag GameplayCueTag, FGameplayEffectContextHandle Context);
 
 	/** Invoke a gameplay cue on the ability owner, with extra parameters */
 	UFUNCTION(BlueprintCallable, Category = Ability, meta = (GameplayTagFilter = "GameplayCue"), DisplayName = "Execute GameplayCueWithParams On Owner", meta=(ScriptName = "ExecuteGameplayCueWithParams"))
-	virtual void K2_ExecuteGameplayCueWithParams(FGameplayTag GameplayCueTag, const FGameplayCueParameters& GameplayCueParameters);
+	UE_API virtual void K2_ExecuteGameplayCueWithParams(FGameplayTag GameplayCueTag, const FGameplayCueParameters& GameplayCueParameters);
 
 	/** Adds a persistent gameplay cue to the ability owner. Optionally will remove if ability ends */
 	UFUNCTION(BlueprintCallable, Category = Ability, meta=(GameplayTagFilter="GameplayCue"), DisplayName="Add GameplayCue To Owner", meta=(ScriptName="AddGameplayCue"))
-	virtual void K2_AddGameplayCue(FGameplayTag GameplayCueTag, FGameplayEffectContextHandle Context, bool bRemoveOnAbilityEnd = true);
+	UE_API virtual void K2_AddGameplayCue(FGameplayTag GameplayCueTag, FGameplayEffectContextHandle Context, bool bRemoveOnAbilityEnd = true);
 
 	/** Adds a persistent gameplay cue to the ability owner. Optionally will remove if ability ends */
 	UFUNCTION(BlueprintCallable, Category = Ability, meta = (GameplayTagFilter = "GameplayCue"), DisplayName = "Add GameplayCueWithParams To Owner", meta=(ScriptName = "AddGameplayCueWithParams"))
-	virtual void K2_AddGameplayCueWithParams(FGameplayTag GameplayCueTag, const FGameplayCueParameters& GameplayCueParameter, bool bRemoveOnAbilityEnd = true);
+	UE_API virtual void K2_AddGameplayCueWithParams(FGameplayTag GameplayCueTag, const FGameplayCueParameters& GameplayCueParameter, bool bRemoveOnAbilityEnd = true);
 
 	/** Removes a persistent gameplay cue from the ability owner */
 	UFUNCTION(BlueprintCallable, Category = Ability, meta=(GameplayTagFilter="GameplayCue"), DisplayName="Remove GameplayCue From Owner", meta=(ScriptName="RemoveGameplayCue"))
-	virtual void K2_RemoveGameplayCue(FGameplayTag GameplayCueTag);
+	UE_API virtual void K2_RemoveGameplayCue(FGameplayTag GameplayCueTag);
 
 	// -------------------------------------
 	//	Protected properties
@@ -794,24 +796,24 @@ protected:
 
 	/** Finds all currently active tasks named InstanceName and confirms them. What this means depends on the individual task. By default, this does nothing other than ending if bEndTask is true. */
 	UFUNCTION(BlueprintCallable, Category = Ability)
-	void ConfirmTaskByInstanceName(FName InstanceName, bool bEndTask);
+	UE_API void ConfirmTaskByInstanceName(FName InstanceName, bool bEndTask);
 
 	/** Internal function, cancels all the tasks we asked to cancel last frame (by instance name). */
-	void EndOrCancelTasksByInstanceName();
+	UE_API void EndOrCancelTasksByInstanceName();
 	TArray<FName> CancelTaskInstanceNames;
 
 	/** Add any task with this instance name to a list to be ended (not canceled) next frame.  See also CancelTaskByInstanceName. */
 	UFUNCTION(BlueprintCallable, Category = Ability)
-	void EndTaskByInstanceName(FName InstanceName);
+	UE_API void EndTaskByInstanceName(FName InstanceName);
 	TArray<FName> EndTaskInstanceNames;
 
 	/** Add any task with this instance name to a list to be canceled (not ended) next frame.  See also EndTaskByInstanceName. */
 	UFUNCTION(BlueprintCallable, Category = Ability)
-	void CancelTaskByInstanceName(FName InstanceName);
+	UE_API void CancelTaskByInstanceName(FName InstanceName);
 
 	/** Ends any active ability state task with the given name. If name is 'None' all active states will be ended (in an arbitrary order). */
 	UFUNCTION(BlueprintCallable, Category = Ability)
-	void EndAbilityState(FName OptionalStateNameToEnd);
+	UE_API void EndAbilityState(FName OptionalStateNameToEnd);
 
 	/** List of currently active tasks, do not modify directly */
 	UPROPERTY()
@@ -826,11 +828,11 @@ protected:
 
 	/** Immediately jumps the active montage to a section */
 	UFUNCTION(BlueprintCallable, Category="Ability|Animation")
-	void MontageJumpToSection(FName SectionName);
+	UE_API void MontageJumpToSection(FName SectionName);
 
 	/** Sets pending section on active montage */
 	UFUNCTION(BlueprintCallable, Category = "Ability|Animation")
-	void MontageSetNextSectionName(FName FromSectionName, FName ToSectionName);
+	UE_API void MontageSetNextSectionName(FName FromSectionName, FName ToSectionName);
 
 	/**
 	 * Stops the current animation montage.
@@ -838,7 +840,7 @@ protected:
 	 * @param OverrideBlendTime If >= 0, will override the BlendOutTime parameter on the AnimMontage instance
 	 */
 	UFUNCTION(BlueprintCallable, Category="Ability|Animation", Meta = (AdvancedDisplay = "OverrideBlendOutTime"))
-	void MontageStop(float OverrideBlendOutTime = -1.0f);
+	UE_API void MontageStop(float OverrideBlendOutTime = -1.0f);
 
 	/** Active montage being played by this ability */
 	UPROPERTY()
@@ -850,27 +852,27 @@ protected:
 
 	/** Creates a target location from where the owner avatar is */
 	UFUNCTION(BlueprintPure, Category = Ability)
-	FGameplayAbilityTargetingLocationInfo MakeTargetLocationInfoFromOwnerActor();
+	UE_API FGameplayAbilityTargetingLocationInfo MakeTargetLocationInfoFromOwnerActor();
 
 	/** Creates a target location from a socket on the owner avatar's skeletal mesh */
 	UFUNCTION(BlueprintPure, Category = Ability)
-	FGameplayAbilityTargetingLocationInfo MakeTargetLocationInfoFromOwnerSkeletalMeshComponent(FName SocketName);
+	UE_API FGameplayAbilityTargetingLocationInfo MakeTargetLocationInfoFromOwnerSkeletalMeshComponent(FName SocketName);
 
 	// ----------------------------------------------------------------------------------------------------------------
 	//	Setters for temporary execution data
 	// ----------------------------------------------------------------------------------------------------------------
 
 	/** Called to initialize after being created due to replication */
-	virtual void PostNetInit();
+	UE_API virtual void PostNetInit();
 
 	/** Modifies actor info, only safe on instanced abilities */
-	virtual void SetCurrentActorInfo(const FGameplayAbilitySpecHandle Handle, const FGameplayAbilityActorInfo* ActorInfo) const;
+	UE_API virtual void SetCurrentActorInfo(const FGameplayAbilitySpecHandle Handle, const FGameplayAbilityActorInfo* ActorInfo) const;
 	
 	/** Modifies activation info, only safe on instanced abilities */
-	virtual void SetCurrentActivationInfo(const FGameplayAbilityActivationInfo ActivationInfo);
+	UE_API virtual void SetCurrentActivationInfo(const FGameplayAbilityActivationInfo ActivationInfo);
 	
 	/** Sets both actor and activation info */
-	void SetCurrentInfo(const FGameplayAbilitySpecHandle Handle, const FGameplayAbilityActorInfo* ActorInfo, const FGameplayAbilityActivationInfo ActivationInfo);
+	UE_API void SetCurrentInfo(const FGameplayAbilitySpecHandle Handle, const FGameplayAbilityActorInfo* ActorInfo, const FGameplayAbilityActivationInfo ActivationInfo);
 	
 	/** 
 	 *  This is shared, cached information about the thing using us
@@ -910,10 +912,10 @@ protected:
 	mutable TArray<FPostLockDelegate> WaitingToExecute;
 
 	/** Increases the scope lock count. */
-	void IncrementListLock() const;
+	UE_API void IncrementListLock() const;
 
 	/** Decreases the scope lock count. Runs the waiting to execute delegates if the count drops to zero. */
-	void DecrementListLock() const;
+	UE_API void DecrementListLock() const;
 
 public:
 	UE_DEPRECATED(5.4, "This is unsafe and unnecessary.  It is ignored.")
@@ -929,3 +931,5 @@ protected:
 	UPROPERTY(BlueprintReadOnly, Category = Ability, meta=(DeprecatedProperty, DeprecationMessage="This is unsafe. Do not use."))
 	bool bMarkPendingKillOnAbilityEnd;
 };
+
+#undef UE_API
diff --git a/Engine/Plugins/Runtime/GameplayAbilities/Source/GameplayAbilities/Public/Abilities/GameplayAbilityRepAnimMontage.h b/Engine/Plugins/Runtime/GameplayAbilities/Source/GameplayAbilities/Public/Abilities/GameplayAbilityRepAnimMontage.h
--- a/Engine/Plugins/Runtime/GameplayAbilities/Source/GameplayAbilities/Public/Abilities/GameplayAbilityRepAnimMontage.h
+++ b/Engine/Plugins/Runtime/GameplayAbilities/Source/GameplayAbilities/Public/Abilities/GameplayAbilityRepAnimMontage.h
@@ -5,6 +5,8 @@
 #include "GameplayPrediction.h"
 #include "GameplayAbilityRepAnimMontage.generated.h"
 
+#define UE_API GAMEPLAYABILITIES_API
+
 class UAnimMontage;
 class UAnimSequenceBase;
 
@@ -18,7 +20,7 @@ enum class ERepAnimPositionMethod
 
 /** Data about montages that is replicated to simulated clients */
 USTRUCT()
-struct GAMEPLAYABILITIES_API FGameplayAbilityRepAnimMontage
+struct FGameplayAbilityRepAnimMontage
 {
 	GENERATED_USTRUCT_BODY()
 
@@ -99,11 +101,11 @@ struct GAMEPLAYABILITIES_API FGameplayAbilityRepAnimMontage
 	{
 	}
 
-	bool NetSerialize(FArchive& Ar, class UPackageMap* Map, bool& bOutSuccess);
+	UE_API bool NetSerialize(FArchive& Ar, class UPackageMap* Map, bool& bOutSuccess);
 
-	void SetRepAnimPositionMethod(ERepAnimPositionMethod InMethod);
+	UE_API void SetRepAnimPositionMethod(ERepAnimPositionMethod InMethod);
 
-	UAnimMontage* GetAnimMontage() const;
+	UE_API UAnimMontage* GetAnimMontage() const;
 };
 
 template<>
@@ -114,3 +116,5 @@ struct TStructOpsTypeTraits<FGameplayAbilityRepAnimMontage> : public TStructOpsT
 		WithNetSerializer = true,
 	};
 };
+
+#undef UE_API
diff --git a/Engine/Plugins/Runtime/GameplayAbilities/Source/GameplayAbilities/Public/Abilities/GameplayAbilityTargetActor.h b/Engine/Plugins/Runtime/GameplayAbilities/Source/GameplayAbilities/Public/Abilities/GameplayAbilityTargetActor.h
--- a/Engine/Plugins/Runtime/GameplayAbilities/Source/GameplayAbilities/Public/Abilities/GameplayAbilityTargetActor.h
+++ b/Engine/Plugins/Runtime/GameplayAbilities/Source/GameplayAbilities/Public/Abilities/GameplayAbilityTargetActor.h
@@ -10,6 +10,8 @@
 #include "Abilities/GameplayAbilityTargetDataFilter.h"
 #include "GameplayAbilityTargetActor.generated.h"
 
+#define UE_API GAMEPLAYABILITIES_API
+
 class APlayerController;
 class UAbilitySystemComponent;
 class UGameplayAbility;
@@ -21,14 +23,14 @@ class UGameplayAbility;
  * For most games you will need to subclass and heavily modify this actor, or you will want to implement similar functions in a game-specific actor or blueprint to avoid actor spawn costs
  * This class is not well tested by internal games, but it is a useful class to look at to learn how target replication occurs
  */
-UCLASS(Blueprintable, abstract, notplaceable)
-class GAMEPLAYABILITIES_API AGameplayAbilityTargetActor : public AActor
+UCLASS(Blueprintable, abstract, notplaceable, MinimalAPI)
+class AGameplayAbilityTargetActor : public AActor
 {
 	GENERATED_UCLASS_BODY()
 
 public:
 
-	void EndPlay(const EEndPlayReason::Type EndPlayReason) override;
+	UE_API void EndPlay(const EEndPlayReason::Type EndPlayReason) override;
 
 	/** The TargetData this class produces can be entirely generated on the server. We don't require the client to send us full or partial TargetData (possibly just a 'confirm') */
 	UPROPERTY(EditAnywhere, Category=Advanced)
@@ -40,34 +42,34 @@ public:
 	FGameplayAbilityTargetingLocationInfo StartLocation;
 	
 	/** Initialize and begin targeting logic  */
-	virtual void StartTargeting(UGameplayAbility* Ability);
+	UE_API virtual void StartTargeting(UGameplayAbility* Ability);
 
-	virtual bool IsConfirmTargetingAllowed();
+	UE_API virtual bool IsConfirmTargetingAllowed();
 
 	/** Requesting targeting data, but not necessarily stopping/destroying the task. Useful for external target data requests. */
-	virtual void ConfirmTargetingAndContinue();
+	UE_API virtual void ConfirmTargetingAndContinue();
 
 	/** Outside code is saying 'stop and just give me what you have.' Returns true if the ability accepts this and can be forgotten. */
 	UFUNCTION()
-	virtual void ConfirmTargeting();
+	UE_API virtual void ConfirmTargeting();
 
 	/** Outside code is saying 'stop everything and just forget about it' */
 	UFUNCTION()
-	virtual void CancelTargeting();
+	UE_API virtual void CancelTargeting();
 
-	virtual void BindToConfirmCancelInputs();
+	UE_API virtual void BindToConfirmCancelInputs();
 
-	virtual bool ShouldProduceTargetData() const;
+	UE_API virtual bool ShouldProduceTargetData() const;
 
 	/** Replicated target data was received from a client. Possibly sanitize/verify. return true if data is good and we should broadcast it as valid data. */
-	virtual bool OnReplicatedTargetDataReceived(FGameplayAbilityTargetDataHandle& Data) const;
+	UE_API virtual bool OnReplicatedTargetDataReceived(FGameplayAbilityTargetDataHandle& Data) const;
 
 	// ------------------------------
 	
 	FAbilityTargetData	TargetDataReadyDelegate;
 	FAbilityTargetData	CanceledDelegate;
 
-	virtual bool IsNetRelevantFor(const AActor* RealViewer, const AActor* ViewTarget, const FVector& SrcLocation) const override;
+	UE_API virtual bool IsNetRelevantFor(const AActor* RealViewer, const AActor* ViewTarget, const FVector& SrcLocation) const override;
 
 #if WITH_EDITOR
 	UE_DEPRECATED(5.1, "This property is deprecated. Please use PrimaryPC instead.")
@@ -107,3 +109,5 @@ public:
 	UPROPERTY()
 	TObjectPtr<UAbilitySystemComponent> GenericDelegateBoundASC;
 };
+
+#undef UE_API
diff --git a/Engine/Plugins/Runtime/GameplayAbilities/Source/GameplayAbilities/Public/Abilities/GameplayAbilityTargetActor_GroundTrace.h b/Engine/Plugins/Runtime/GameplayAbilities/Source/GameplayAbilities/Public/Abilities/GameplayAbilityTargetActor_GroundTrace.h
--- a/Engine/Plugins/Runtime/GameplayAbilities/Source/GameplayAbilities/Public/Abilities/GameplayAbilityTargetActor_GroundTrace.h
+++ b/Engine/Plugins/Runtime/GameplayAbilities/Source/GameplayAbilities/Public/Abilities/GameplayAbilityTargetActor_GroundTrace.h
@@ -10,14 +10,16 @@
 #include "Abilities/GameplayAbilityTargetActor_Trace.h"
 #include "GameplayAbilityTargetActor_GroundTrace.generated.h"
 
+#define UE_API GAMEPLAYABILITIES_API
+
 class UGameplayAbility;
 
-UCLASS(Blueprintable)
-class GAMEPLAYABILITIES_API AGameplayAbilityTargetActor_GroundTrace : public AGameplayAbilityTargetActor_Trace
+UCLASS(Blueprintable, MinimalAPI)
+class AGameplayAbilityTargetActor_GroundTrace : public AGameplayAbilityTargetActor_Trace
 {
 	GENERATED_UCLASS_BODY()
 
-	virtual void StartTargeting(UGameplayAbility* InAbility) override;
+	UE_API virtual void StartTargeting(UGameplayAbility* InAbility) override;
 
 	/** Radius for a sphere or capsule. */
 	UPROPERTY(BlueprintReadWrite, EditAnywhere, meta = (ExposeOnSpawn = true), Category = Targeting)
@@ -30,13 +32,15 @@ class GAMEPLAYABILITIES_API AGameplayAbilityTargetActor_GroundTrace : public AGa
 	//TODO: Consider adding the ability to rotate the capsule with a supplied quaternion.
 
 protected:
-	virtual FHitResult PerformTrace(AActor* InSourceActor) override;
+	UE_API virtual FHitResult PerformTrace(AActor* InSourceActor) override;
 
-	virtual bool IsConfirmTargetingAllowed() override;
+	UE_API virtual bool IsConfirmTargetingAllowed() override;
 
-	bool AdjustCollisionResultForShape(const FVector OriginalStartPoint, const FVector OriginalEndPoint, const FCollisionQueryParams Params, FHitResult& OutHitResult) const;
+	UE_API bool AdjustCollisionResultForShape(const FVector OriginalStartPoint, const FVector OriginalEndPoint, const FCollisionQueryParams Params, FHitResult& OutHitResult) const;
 
 	FCollisionShape CollisionShape;
 	float CollisionHeightOffset;		//When tracing, give this much extra height to avoid start-in-ground problems. Dealing with thick placement actors while standing near walls may be trickier.
 	bool bLastTraceWasGood;
 };
+
+#undef UE_API
diff --git a/Engine/Plugins/Runtime/GameplayAbilities/Source/GameplayAbilities/Public/Abilities/GameplayAbilityTargetActor_Radius.h b/Engine/Plugins/Runtime/GameplayAbilities/Source/GameplayAbilities/Public/Abilities/GameplayAbilityTargetActor_Radius.h
--- a/Engine/Plugins/Runtime/GameplayAbilities/Source/GameplayAbilities/Public/Abilities/GameplayAbilityTargetActor_Radius.h
+++ b/Engine/Plugins/Runtime/GameplayAbilities/Source/GameplayAbilities/Public/Abilities/GameplayAbilityTargetActor_Radius.h
@@ -8,19 +8,21 @@
 #include "Abilities/GameplayAbilityTargetActor.h"
 #include "GameplayAbilityTargetActor_Radius.generated.h"
 
+#define UE_API GAMEPLAYABILITIES_API
+
 class UGameplayAbility;
 
 /** Selects everything within a given radius of the source actor. */
-UCLASS(Blueprintable, notplaceable)
-class GAMEPLAYABILITIES_API AGameplayAbilityTargetActor_Radius : public AGameplayAbilityTargetActor
+UCLASS(Blueprintable, notplaceable, MinimalAPI)
+class AGameplayAbilityTargetActor_Radius : public AGameplayAbilityTargetActor
 {
 	GENERATED_UCLASS_BODY()
 
 public:
 
-	virtual void StartTargeting(UGameplayAbility* Ability) override;
+	UE_API virtual void StartTargeting(UGameplayAbility* Ability) override;
 	
-	virtual void ConfirmTargetingAndContinue() override;
+	UE_API virtual void ConfirmTargetingAndContinue() override;
 
 	/** Radius of target acquisition around the ability's start location. */
 	UPROPERTY(BlueprintReadWrite, EditAnywhere, meta = (ExposeOnSpawn = true), Category = Radius)
@@ -28,8 +30,10 @@ public:
 
 protected:
 
-	TArray<TWeakObjectPtr<AActor> >	PerformOverlap(const FVector& Origin);
+	UE_API TArray<TWeakObjectPtr<AActor> >	PerformOverlap(const FVector& Origin);
 
-	FGameplayAbilityTargetDataHandle MakeTargetData(const TArray<TWeakObjectPtr<AActor>>& Actors, const FVector& Origin) const;
+	UE_API FGameplayAbilityTargetDataHandle MakeTargetData(const TArray<TWeakObjectPtr<AActor>>& Actors, const FVector& Origin) const;
 
 };
+
+#undef UE_API
diff --git a/Engine/Plugins/Runtime/GameplayAbilities/Source/GameplayAbilities/Public/Abilities/GameplayAbilityTargetActor_SingleLineTrace.h b/Engine/Plugins/Runtime/GameplayAbilities/Source/GameplayAbilities/Public/Abilities/GameplayAbilityTargetActor_SingleLineTrace.h
--- a/Engine/Plugins/Runtime/GameplayAbilities/Source/GameplayAbilities/Public/Abilities/GameplayAbilityTargetActor_SingleLineTrace.h
+++ b/Engine/Plugins/Runtime/GameplayAbilities/Source/GameplayAbilities/Public/Abilities/GameplayAbilityTargetActor_SingleLineTrace.h
@@ -8,11 +8,15 @@
 #include "Abilities/GameplayAbilityTargetActor_Trace.h"
 #include "GameplayAbilityTargetActor_SingleLineTrace.generated.h"
 
-UCLASS(Blueprintable)
-class GAMEPLAYABILITIES_API AGameplayAbilityTargetActor_SingleLineTrace : public AGameplayAbilityTargetActor_Trace
+#define UE_API GAMEPLAYABILITIES_API
+
+UCLASS(Blueprintable, MinimalAPI)
+class AGameplayAbilityTargetActor_SingleLineTrace : public AGameplayAbilityTargetActor_Trace
 {
 	GENERATED_UCLASS_BODY()
 
 protected:
-	virtual FHitResult PerformTrace(AActor* InSourceActor) override;
+	UE_API virtual FHitResult PerformTrace(AActor* InSourceActor) override;
 };
+
+#undef UE_API
diff --git a/Engine/Plugins/Runtime/GameplayAbilities/Source/GameplayAbilities/Public/Abilities/GameplayAbilityTargetActor_Trace.h b/Engine/Plugins/Runtime/GameplayAbilities/Source/GameplayAbilities/Public/Abilities/GameplayAbilityTargetActor_Trace.h
--- a/Engine/Plugins/Runtime/GameplayAbilities/Source/GameplayAbilities/Public/Abilities/GameplayAbilityTargetActor_Trace.h
+++ b/Engine/Plugins/Runtime/GameplayAbilities/Source/GameplayAbilities/Public/Abilities/GameplayAbilityTargetActor_Trace.h
@@ -13,32 +13,34 @@
 #include "Abilities/GameplayAbilityTargetActor.h"
 #include "GameplayAbilityTargetActor_Trace.generated.h"
 
+#define UE_API GAMEPLAYABILITIES_API
+
 class UGameplayAbility;
 
 /** Intermediate base class for all line-trace type targeting actors. */
-UCLASS(Abstract, Blueprintable, notplaceable, config=Game)
-class GAMEPLAYABILITIES_API AGameplayAbilityTargetActor_Trace : public AGameplayAbilityTargetActor
+UCLASS(Abstract, Blueprintable, notplaceable, config=Game, MinimalAPI)
+class AGameplayAbilityTargetActor_Trace : public AGameplayAbilityTargetActor
 {
 	GENERATED_UCLASS_BODY()
 
 public:
-	virtual void EndPlay(const EEndPlayReason::Type EndPlayReason) override;
+	UE_API virtual void EndPlay(const EEndPlayReason::Type EndPlayReason) override;
 
 	/** Traces as normal, but will manually filter all hit actors */
-	static void LineTraceWithFilter(FHitResult& OutHitResult, const UWorld* World, const FGameplayTargetDataFilterHandle FilterHandle, const FVector& Start, const FVector& End, FName ProfileName, const FCollisionQueryParams Params);
+	static UE_API void LineTraceWithFilter(FHitResult& OutHitResult, const UWorld* World, const FGameplayTargetDataFilterHandle FilterHandle, const FVector& Start, const FVector& End, FName ProfileName, const FCollisionQueryParams Params);
 
 	/** Sweeps as normal, but will manually filter all hit actors */
-	static void SweepWithFilter(FHitResult& OutHitResult, const UWorld* World, const FGameplayTargetDataFilterHandle FilterHandle, const FVector& Start, const FVector& End, const FQuat& Rotation, const FCollisionShape CollisionShape, FName ProfileName, const FCollisionQueryParams Params);
+	static UE_API void SweepWithFilter(FHitResult& OutHitResult, const UWorld* World, const FGameplayTargetDataFilterHandle FilterHandle, const FVector& Start, const FVector& End, const FQuat& Rotation, const FCollisionShape CollisionShape, FName ProfileName, const FCollisionQueryParams Params);
 
-	void AimWithPlayerController(const AActor* InSourceActor, FCollisionQueryParams Params, const FVector& TraceStart, FVector& OutTraceEnd, bool bIgnorePitch = false) const;
+	UE_API void AimWithPlayerController(const AActor* InSourceActor, FCollisionQueryParams Params, const FVector& TraceStart, FVector& OutTraceEnd, bool bIgnorePitch = false) const;
 
-	static bool ClipCameraRayToAbilityRange(FVector CameraLocation, FVector CameraDirection, FVector AbilityCenter, float AbilityRange, FVector& ClippedPosition);
+	static UE_API bool ClipCameraRayToAbilityRange(FVector CameraLocation, FVector CameraDirection, FVector AbilityCenter, float AbilityRange, FVector& ClippedPosition);
 
-	virtual void StartTargeting(UGameplayAbility* Ability) override;
+	UE_API virtual void StartTargeting(UGameplayAbility* Ability) override;
 
-	virtual void ConfirmTargetingAndContinue() override;
+	UE_API virtual void ConfirmTargetingAndContinue() override;
 
-	virtual void Tick(float DeltaSeconds) override;
+	UE_API virtual void Tick(float DeltaSeconds) override;
 
 	UPROPERTY(BlueprintReadWrite, EditAnywhere, meta = (ExposeOnSpawn = true), Category = Trace)
 	float MaxRange;
@@ -51,9 +53,11 @@ public:
 	bool bTraceAffectsAimPitch;
 
 protected:
-	virtual FHitResult PerformTrace(AActor* InSourceActor) PURE_VIRTUAL(AGameplayAbilityTargetActor_Trace, return FHitResult(););
+	UE_API virtual FHitResult PerformTrace(AActor* InSourceActor) PURE_VIRTUAL(AGameplayAbilityTargetActor_Trace, return FHitResult(););
 
-	FGameplayAbilityTargetDataHandle MakeTargetData(const FHitResult& HitResult) const;
+	UE_API FGameplayAbilityTargetDataHandle MakeTargetData(const FHitResult& HitResult) const;
 
 	TWeakObjectPtr<AGameplayAbilityWorldReticle> ReticleActor;
 };
+
+#undef UE_API
diff --git a/Engine/Plugins/Runtime/GameplayAbilities/Source/GameplayAbilities/Public/Abilities/GameplayAbilityTargetDataFilter.h b/Engine/Plugins/Runtime/GameplayAbilities/Source/GameplayAbilities/Public/Abilities/GameplayAbilityTargetDataFilter.h
--- a/Engine/Plugins/Runtime/GameplayAbilities/Source/GameplayAbilities/Public/Abilities/GameplayAbilityTargetDataFilter.h
+++ b/Engine/Plugins/Runtime/GameplayAbilities/Source/GameplayAbilities/Public/Abilities/GameplayAbilityTargetDataFilter.h
@@ -8,6 +8,8 @@
 #include "GameFramework/Actor.h"
 #include "GameplayAbilityTargetDataFilter.generated.h"
 
+#define UE_API GAMEPLAYABILITIES_API
+
 /** Set rather it is possible to target self */
 UENUM(BlueprintType)
 namespace ETargetDataFilterSelf
@@ -22,7 +24,7 @@ namespace ETargetDataFilterSelf
 
 /** Simple actor target filter, games can subclass this */
 USTRUCT(BlueprintType)
-struct GAMEPLAYABILITIES_API FGameplayTargetDataFilter
+struct FGameplayTargetDataFilter
 {
 	GENERATED_USTRUCT_BODY()
 
@@ -61,7 +63,7 @@ struct GAMEPLAYABILITIES_API FGameplayTargetDataFilter
 	}
 
 	/** Initializes SelfActor */
-	void InitializeFilterContext(AActor* FilterActor);
+	UE_API void InitializeFilterContext(AActor* FilterActor);
 
 	/** Actor we're comparing against. */
 	UPROPERTY()
@@ -82,7 +84,7 @@ struct GAMEPLAYABILITIES_API FGameplayTargetDataFilter
 
 /** Polymorphic handle to filter structure that handles checking if actors should be targeted */
 USTRUCT(BlueprintType)
-struct GAMEPLAYABILITIES_API FGameplayTargetDataFilterHandle
+struct FGameplayTargetDataFilterHandle
 {
 	GENERATED_USTRUCT_BODY()
 
@@ -123,3 +125,5 @@ struct GAMEPLAYABILITIES_API FGameplayTargetDataFilterHandle
 		return FilterPassesForActor(ActorToBeFiltered);
 	}
 };
+
+#undef UE_API
diff --git a/Engine/Plugins/Runtime/GameplayAbilities/Source/GameplayAbilities/Public/Abilities/GameplayAbilityTargetTypes.h b/Engine/Plugins/Runtime/GameplayAbilities/Source/GameplayAbilities/Public/Abilities/GameplayAbilityTargetTypes.h
--- a/Engine/Plugins/Runtime/GameplayAbilities/Source/GameplayAbilities/Public/Abilities/GameplayAbilityTargetTypes.h
+++ b/Engine/Plugins/Runtime/GameplayAbilities/Source/GameplayAbilities/Public/Abilities/GameplayAbilityTargetTypes.h
@@ -15,6 +15,8 @@
 
 #include "GameplayAbilityTargetTypes.generated.h"
 
+#define UE_API GAMEPLAYABILITIES_API
+
 class UGameplayAbility;
 class UGameplayEffect;
 class UMeshComponent;
@@ -66,20 +68,20 @@ namespace EGameplayTargetingConfirmation
  *		-AOE/aura type of targeting data blurs the line
  */
 USTRUCT()
-struct GAMEPLAYABILITIES_API FGameplayAbilityTargetData
+struct FGameplayAbilityTargetData
 {
 	GENERATED_USTRUCT_BODY()
 
 	virtual ~FGameplayAbilityTargetData() { }
 
 	/** Applies a gameplay effect to each target represented */
-	TArray<FActiveGameplayEffectHandle> ApplyGameplayEffect(const UGameplayEffect* GameplayEffect, const FGameplayEffectContextHandle& InEffectContext, float Level, FPredictionKey PredictionKey = FPredictionKey());
+	UE_API TArray<FActiveGameplayEffectHandle> ApplyGameplayEffect(const UGameplayEffect* GameplayEffect, const FGameplayEffectContextHandle& InEffectContext, float Level, FPredictionKey PredictionKey = FPredictionKey());
 
 	/** Applies a previously created gameplay effect spec to each target represented */
-	virtual TArray<FActiveGameplayEffectHandle> ApplyGameplayEffectSpec(FGameplayEffectSpec& Spec, FPredictionKey PredictionKey = FPredictionKey());
+	UE_API virtual TArray<FActiveGameplayEffectHandle> ApplyGameplayEffectSpec(FGameplayEffectSpec& Spec, FPredictionKey PredictionKey = FPredictionKey());
 
 	/** Modifies the context and adds this target data to the target data handle stored within */
-	virtual void AddTargetDataToContext(FGameplayEffectContextHandle& Context, bool bIncludeActorArray) const;
+	UE_API virtual void AddTargetDataToContext(FGameplayEffectContextHandle& Context, bool bIncludeActorArray) const;
 
 	/** Returns all actors targeted, almost always overridden */
 	virtual TArray<TWeakObjectPtr<AActor>> GetActors() const
@@ -143,7 +145,7 @@ struct GAMEPLAYABILITIES_API FGameplayAbilityTargetData
 	}
 
 	/** Returns a debug string representation */
-	virtual FString ToString() const;
+	UE_API virtual FString ToString() const;
 
 	/** See notes on delegate definition FOnTargetActorSwapped */
 	virtual bool ShouldCheckForTargetActorSwap() const
@@ -188,7 +190,7 @@ namespace EGameplayAbilityTargetingLocationType
 *		-Polymophism in TargetData structure
 */
 USTRUCT(BlueprintType)
-struct GAMEPLAYABILITIES_API FGameplayAbilityTargetDataHandle
+struct FGameplayAbilityTargetDataHandle
 {
 	GENERATED_USTRUCT_BODY()
 
@@ -252,7 +254,7 @@ struct GAMEPLAYABILITIES_API FGameplayAbilityTargetDataHandle
 	}
 
 	/** Serialize for networking, handles polymorphism */
-	bool NetSerialize(FArchive& Ar, class UPackageMap* Map, bool& bOutSuccess);
+	UE_API bool NetSerialize(FArchive& Ar, class UPackageMap* Map, bool& bOutSuccess);
 
 	/** Comparison operator */
 	bool operator==(const FGameplayAbilityTargetDataHandle& Other) const
@@ -302,7 +304,7 @@ struct FGameplayAbilityTargetingLocationInfoFixLayout
 
 /** Structure that stores a location in one of several different formats */
 USTRUCT(BlueprintType)
-struct GAMEPLAYABILITIES_API FGameplayAbilityTargetingLocationInfo
+struct FGameplayAbilityTargetingLocationInfo
 #if CPP
 	: public FGameplayAbilityTargetingLocationInfoFixLayout
 #endif
@@ -329,13 +331,13 @@ struct GAMEPLAYABILITIES_API FGameplayAbilityTargetingLocationInfo
 	}
 
 	/** Converts internal format into a literal world space transform */
-	FTransform GetTargetingTransform() const;
+	UE_API FTransform GetTargetingTransform() const;
 	/** Initializes new target data and fills in with hit results */
-	FGameplayAbilityTargetDataHandle MakeTargetDataHandleFromHitResult(TWeakObjectPtr<UGameplayAbility> Ability, const FHitResult& HitResult) const;
-	FGameplayAbilityTargetDataHandle MakeTargetDataHandleFromHitResults(TWeakObjectPtr<UGameplayAbility> Ability, const TArray<FHitResult>& HitResults) const;
+	UE_API FGameplayAbilityTargetDataHandle MakeTargetDataHandleFromHitResult(TWeakObjectPtr<UGameplayAbility> Ability, const FHitResult& HitResult) const;
+	UE_API FGameplayAbilityTargetDataHandle MakeTargetDataHandleFromHitResults(TWeakObjectPtr<UGameplayAbility> Ability, const TArray<FHitResult>& HitResults) const;
 
 	/** Initializes new actor list target data, and sets this as the origin */
-	FGameplayAbilityTargetDataHandle MakeTargetDataHandleFromActors(const TArray<TWeakObjectPtr<AActor>>& TargetActors, bool OneActorPerHandle = false) const;
+	UE_API FGameplayAbilityTargetDataHandle MakeTargetDataHandleFromActors(const TArray<TWeakObjectPtr<AActor>>& TargetActors, bool OneActorPerHandle = false) const;
 
 	/** A source actor is needed for Actor-based targeting, but not for Socket-based targeting. */
 	UPROPERTY(BlueprintReadWrite, meta = (ExposeOnSpawn = true), Category = Targeting)
@@ -362,7 +364,7 @@ struct GAMEPLAYABILITIES_API FGameplayAbilityTargetingLocationInfo
 	TEnumAsByte<EGameplayAbilityTargetingLocationType::Type> LocationType;
 
 	/** Optimized serialize function */
-	bool NetSerialize(FArchive& Ar, class UPackageMap* Map, bool& bOutSuccess);
+	UE_API bool NetSerialize(FArchive& Ar, class UPackageMap* Map, bool& bOutSuccess);
 };
 
 template<>
@@ -376,7 +378,7 @@ struct TStructOpsTypeTraits<FGameplayAbilityTargetingLocationInfo> : public TStr
 
 /** Target data with just a source and target location in space */
 USTRUCT(BlueprintType)
-struct GAMEPLAYABILITIES_API FGameplayAbilityTargetData_LocationInfo : public FGameplayAbilityTargetData
+struct FGameplayAbilityTargetData_LocationInfo : public FGameplayAbilityTargetData
 {
 	GENERATED_USTRUCT_BODY()
 
@@ -429,7 +431,7 @@ struct GAMEPLAYABILITIES_API FGameplayAbilityTargetData_LocationInfo : public FG
 		return TEXT("FGameplayAbilityTargetData_LocationInfo");
 	}
 
-	bool NetSerialize(FArchive& Ar, class UPackageMap* Map, bool& bOutSuccess);
+	UE_API bool NetSerialize(FArchive& Ar, class UPackageMap* Map, bool& bOutSuccess);
 };
 
 template<>
@@ -443,7 +445,7 @@ struct TStructOpsTypeTraits<FGameplayAbilityTargetData_LocationInfo> : public TS
 
 /** Target data with a source location and a list of targeted actors, makes sense for AOE attacks */
 USTRUCT(BlueprintType)
-struct GAMEPLAYABILITIES_API FGameplayAbilityTargetData_ActorArray : public FGameplayAbilityTargetData
+struct FGameplayAbilityTargetData_ActorArray : public FGameplayAbilityTargetData
 {
 	GENERATED_USTRUCT_BODY()
 
@@ -532,7 +534,7 @@ struct GAMEPLAYABILITIES_API FGameplayAbilityTargetData_ActorArray : public FGam
 		return TEXT("FGameplayAbilityTargetData_ActorArray");
 	}
 
-	bool NetSerialize(FArchive& Ar, class UPackageMap* Map, bool& bOutSuccess);
+	UE_API bool NetSerialize(FArchive& Ar, class UPackageMap* Map, bool& bOutSuccess);
 };
 
 template<>
@@ -546,7 +548,7 @@ struct TStructOpsTypeTraits<FGameplayAbilityTargetData_ActorArray> : public TStr
 
 /** Target data with a single hit result, data is packed into the hit result */
 USTRUCT(BlueprintType)
-struct GAMEPLAYABILITIES_API FGameplayAbilityTargetData_SingleTargetHit : public FGameplayAbilityTargetData
+struct FGameplayAbilityTargetData_SingleTargetHit : public FGameplayAbilityTargetData
 {
 	GENERATED_USTRUCT_BODY()
 
@@ -623,7 +625,7 @@ struct GAMEPLAYABILITIES_API FGameplayAbilityTargetData_SingleTargetHit : public
 	UPROPERTY(NotReplicated)
 	bool bHitReplaced = false;
 
-	bool NetSerialize(FArchive& Ar, class UPackageMap* Map, bool& bOutSuccess);
+	UE_API bool NetSerialize(FArchive& Ar, class UPackageMap* Map, bool& bOutSuccess);
 
 	virtual UScriptStruct* GetScriptStruct() const override
 	{
@@ -688,4 +690,6 @@ struct FAbilityReplicatedData
 
 	/** Delegate that will be called on replication */
 	FSimpleMulticastDelegate Delegate;
-};
\ No newline at end of file
+};
+
+#undef UE_API
