[GameplayAbilities]

diff --git a/Engine/Plugins/Runtime/GameplayAbilities/Source/GameplayAbilities/Public/Abilities/GameplayAbilityTypes.h b/Engine/Plugins/Runtime/GameplayAbilities/Source/GameplayAbilities/Public/Abilities/GameplayAbilityTypes.h
--- a/Engine/Plugins/Runtime/GameplayAbilities/Source/GameplayAbilities/Public/Abilities/GameplayAbilityTypes.h
+++ b/Engine/Plugins/Runtime/GameplayAbilities/Source/GameplayAbilities/Public/Abilities/GameplayAbilityTypes.h
@@ -15,6 +15,8 @@
 
 #include "GameplayAbilityTypes.generated.h"
 
+#define UE_API GAMEPLAYABILITIES_API
+
 class APlayerController;
 class UAbilitySystemComponent;
 class UAnimInstance;
@@ -133,7 +135,7 @@ namespace EGameplayAbilityTriggerSource
  *
  */
 USTRUCT(BlueprintType)
-struct GAMEPLAYABILITIES_API FGameplayAbilityActorInfo
+struct FGameplayAbilityActorInfo
 {
 	GENERATED_USTRUCT_BODY()
 
@@ -172,32 +174,32 @@ struct GAMEPLAYABILITIES_API FGameplayAbilityActorInfo
 	FName AffectedAnimInstanceTag; 
 	
 	/** Accessor to get the affected anim instance from the SkeletalMeshComponent */
-	UAnimInstance* GetAnimInstance() const;
+	UE_API UAnimInstance* GetAnimInstance() const;
 	
 	/** Returns true if this actor is locally controlled. Only true for players on the client that owns them (differs from APawn::IsLocallyControlled which requires a Controller) */
-	bool IsLocallyControlled() const;
+	UE_API bool IsLocallyControlled() const;
 
 	/** Returns true if this actor has a PlayerController that is locally controlled. */
-	bool IsLocallyControlledPlayer() const;
+	UE_API bool IsLocallyControlledPlayer() const;
 
 	/** Returns true if the owning actor has net authority */
-	bool IsNetAuthority() const;
+	UE_API bool IsNetAuthority() const;
 
 	/** Initializes the info from an owning actor. Will set both owner and avatar */
-	virtual void InitFromActor(AActor *OwnerActor, AActor *AvatarActor, UAbilitySystemComponent* InAbilitySystemComponent);
+	UE_API virtual void InitFromActor(AActor *OwnerActor, AActor *AvatarActor, UAbilitySystemComponent* InAbilitySystemComponent);
 
 	/** Sets a new avatar actor, keeps same owner and ability system component */
-	virtual void SetAvatarActor(AActor *AvatarActor);
+	UE_API virtual void SetAvatarActor(AActor *AvatarActor);
 
 	/** Clears out any actor info, both owner and avatar */
-	virtual void ClearActorInfo();
+	UE_API virtual void ClearActorInfo();
 };
 
 
 
 /** Data about montages that were played locally (all montages in case of server. predictive montages in case of client). Never replicated directly. */
 USTRUCT()
-struct GAMEPLAYABILITIES_API FGameplayAbilityLocalAnimMontage
+struct FGameplayAbilityLocalAnimMontage
 {
 	GENERATED_USTRUCT_BODY()
 
@@ -226,7 +228,7 @@ struct GAMEPLAYABILITIES_API FGameplayAbilityLocalAnimMontage
 
 /** Metadata for a tag-based Gameplay Event, that can activate other abilities or run ability-specific logic */
 USTRUCT(BlueprintType)
-struct GAMEPLAYABILITIES_API FGameplayEventData
+struct FGameplayEventData
 {
 	GENERATED_USTRUCT_BODY()
 
@@ -403,7 +405,7 @@ PRAGMA_ENABLE_DEPRECATION_WARNINGS
 
 /** Used to initialize default values for attributes */
 USTRUCT()
-struct GAMEPLAYABILITIES_API FAttributeDefaults
+struct FAttributeDefaults
 {
 	GENERATED_USTRUCT_BODY()
 
@@ -421,7 +423,7 @@ struct GAMEPLAYABILITIES_API FAttributeDefaults
 
 /** Debug message emitted by ability tasks */
 USTRUCT()
-struct GAMEPLAYABILITIES_API FAbilityTaskDebugMessage
+struct FAbilityTaskDebugMessage
 {
 	GENERATED_USTRUCT_BODY()
 
@@ -482,10 +484,10 @@ struct FServerAbilityRPCBatch
 
 
 /** Helper struct for defining ServerRPC batch windows. If null ASC is passed in, this becomes a noop. */
-struct GAMEPLAYABILITIES_API FScopedServerAbilityRPCBatcher
+struct FScopedServerAbilityRPCBatcher
 {
-	FScopedServerAbilityRPCBatcher(UAbilitySystemComponent* InASC, FGameplayAbilitySpecHandle InAbilityHandle);
-	~FScopedServerAbilityRPCBatcher();
+	UE_API FScopedServerAbilityRPCBatcher(UAbilitySystemComponent* InASC, FGameplayAbilitySpecHandle InAbilityHandle);
+	UE_API ~FScopedServerAbilityRPCBatcher();
 
 private:
 
@@ -497,7 +499,7 @@ private:
 
 /** Used as a key for storing internal ability data */
 USTRUCT()
-struct GAMEPLAYABILITIES_API FGameplayAbilitySpecHandleAndPredictionKey
+struct FGameplayAbilitySpecHandleAndPredictionKey
 {
 	GENERATED_USTRUCT_BODY()
 
@@ -532,7 +534,7 @@ struct GAMEPLAYABILITIES_API FGameplayAbilitySpecHandleAndPredictionKey
 };
 
 /** Struct defining the cached data for a specific gameplay ability. This data is generally synchronized client->server in a network game. */
-struct GAMEPLAYABILITIES_API FAbilityReplicatedDataCache
+struct FAbilityReplicatedDataCache
 {
 	/** What elements this activation is targeting */
 	FGameplayAbilityTargetDataHandle TargetData;
@@ -611,4 +613,6 @@ private:
 
 	TArray<FKeyDataPair> InUseData;
 	TArray<TSharedRef<FAbilityReplicatedDataCache>> FreeData;
-};
\ No newline at end of file
+};
+
+#undef UE_API
diff --git a/Engine/Plugins/Runtime/GameplayAbilities/Source/GameplayAbilities/Public/Abilities/GameplayAbilityWorldReticle.h b/Engine/Plugins/Runtime/GameplayAbilities/Source/GameplayAbilities/Public/Abilities/GameplayAbilityWorldReticle.h
--- a/Engine/Plugins/Runtime/GameplayAbilities/Source/GameplayAbilities/Public/Abilities/GameplayAbilityWorldReticle.h
+++ b/Engine/Plugins/Runtime/GameplayAbilities/Source/GameplayAbilities/Public/Abilities/GameplayAbilityWorldReticle.h
@@ -6,6 +6,8 @@
 #include "GameFramework/Actor.h"
 #include "GameplayAbilityWorldReticle.generated.h"
 
+#define UE_API GAMEPLAYABILITIES_API
+
 class APlayerController;
 class AGameplayAbilityTargetActor;
 
@@ -25,43 +27,43 @@ struct FWorldReticleParameters
 };
 
 /** Reticles allow targeting to be visualized. Tasks can spawn these. Artists/designers can create BPs for these. */
-UCLASS(Blueprintable, notplaceable)
-class GAMEPLAYABILITIES_API AGameplayAbilityWorldReticle : public AActor
+UCLASS(Blueprintable, notplaceable, MinimalAPI)
+class AGameplayAbilityWorldReticle : public AActor
 {
 	GENERATED_UCLASS_BODY()
 
 public:
 
-	virtual void Tick(float DeltaSeconds) override;
+	UE_API virtual void Tick(float DeltaSeconds) override;
 
 	// ------------------------------
 
-	virtual bool IsNetRelevantFor(const AActor* RealViewer, const AActor* ViewTarget, const FVector& SrcLocation) const override;
+	UE_API virtual bool IsNetRelevantFor(const AActor* RealViewer, const AActor* ViewTarget, const FVector& SrcLocation) const override;
 
-	void InitializeReticle(AGameplayAbilityTargetActor* InTargetingActor, APlayerController* PlayerController, FWorldReticleParameters InParameters);
+	UE_API void InitializeReticle(AGameplayAbilityTargetActor* InTargetingActor, APlayerController* PlayerController, FWorldReticleParameters InParameters);
 
-	void SetIsTargetValid(bool bNewValue);
-	void SetIsTargetAnActor(bool bNewValue);
+	UE_API void SetIsTargetValid(bool bNewValue);
+	UE_API void SetIsTargetAnActor(bool bNewValue);
 
 	/** Called whenever bIsTargetValid changes value. */
 	UFUNCTION(BlueprintImplementableEvent, Category = Reticle)
-	void OnValidTargetChanged(bool bNewValue);
+	UE_API void OnValidTargetChanged(bool bNewValue);
 
 	/** Called whenever bIsTargetAnActor changes value. */
 	UFUNCTION(BlueprintImplementableEvent, Category = Reticle)
-	void OnTargetingAnActor(bool bNewValue);
+	UE_API void OnTargetingAnActor(bool bNewValue);
 
 	UFUNCTION(BlueprintImplementableEvent, Category = Reticle)
-	void OnParametersInitialized();
+	UE_API void OnParametersInitialized();
 
 	UFUNCTION(BlueprintImplementableEvent, Category = Reticle)
-	void SetReticleMaterialParamFloat(FName ParamName, float value);
+	UE_API void SetReticleMaterialParamFloat(FName ParamName, float value);
 
 	UFUNCTION(BlueprintImplementableEvent, Category = Reticle)
-	void SetReticleMaterialParamVector(FName ParamName, FVector value);
+	UE_API void SetReticleMaterialParamVector(FName ParamName, FVector value);
 
 	UFUNCTION(BlueprintCallable, Category = Reticle)
-	void FaceTowardSource(bool bFaceIn2D);
+	UE_API void FaceTowardSource(bool bFaceIn2D);
 
 	UPROPERTY(EditInstanceOnly, BlueprintReadWrite, meta = (ExposeOnSpawn = "true"), Category = "Reticle")
 	FWorldReticleParameters Parameters;
@@ -97,3 +99,5 @@ protected:
 	UPROPERTY(BlueprintReadOnly, Category = "Network")
 	TObjectPtr<AGameplayAbilityTargetActor> TargetingActor;
 };
+
+#undef UE_API
diff --git a/Engine/Plugins/Runtime/GameplayAbilities/Source/GameplayAbilities/Public/Abilities/GameplayAbilityWorldReticle_ActorVisualization.h b/Engine/Plugins/Runtime/GameplayAbilities/Source/GameplayAbilities/Public/Abilities/GameplayAbilityWorldReticle_ActorVisualization.h
--- a/Engine/Plugins/Runtime/GameplayAbilities/Source/GameplayAbilities/Public/Abilities/GameplayAbilityWorldReticle_ActorVisualization.h
+++ b/Engine/Plugins/Runtime/GameplayAbilities/Source/GameplayAbilities/Public/Abilities/GameplayAbilityWorldReticle_ActorVisualization.h
@@ -7,6 +7,8 @@
 
 #include "GameplayAbilityWorldReticle_ActorVisualization.generated.h"
 
+#define UE_API GAMEPLAYABILITIES_API
+
 class AGameplayAbilityTargetActor;
 class UMaterialInterface;
 
@@ -32,5 +34,7 @@ public:
 	virtual void EndPlay(const EEndPlayReason::Type EndPlayReason) override;
 
 	/** Returns CollisionComponent subobject **/
-	GAMEPLAYABILITIES_API class UCapsuleComponent* GetCollisionComponent() { return CollisionComponent; }
+	class UCapsuleComponent* GetCollisionComponent() { return CollisionComponent; }
 };
+
+#undef UE_API
diff --git a/Engine/Plugins/Runtime/GameplayAbilities/Source/GameplayAbilities/Public/Abilities/GameplayAbility_CharacterJump.h b/Engine/Plugins/Runtime/GameplayAbilities/Source/GameplayAbilities/Public/Abilities/GameplayAbility_CharacterJump.h
--- a/Engine/Plugins/Runtime/GameplayAbilities/Source/GameplayAbilities/Public/Abilities/GameplayAbility_CharacterJump.h
+++ b/Engine/Plugins/Runtime/GameplayAbilities/Source/GameplayAbilities/Public/Abilities/GameplayAbility_CharacterJump.h
@@ -8,21 +8,25 @@
 #include "Abilities/GameplayAbility.h"
 #include "GameplayAbility_CharacterJump.generated.h"
 
+#define UE_API GAMEPLAYABILITIES_API
+
 /**
  *	Ability that jumps with a character.
  */
-UCLASS()
-class GAMEPLAYABILITIES_API UGameplayAbility_CharacterJump : public UGameplayAbility
+UCLASS(MinimalAPI)
+class UGameplayAbility_CharacterJump : public UGameplayAbility
 {
 	GENERATED_UCLASS_BODY()
 
 public:
 
-	virtual bool CanActivateAbility(const FGameplayAbilitySpecHandle Handle, const FGameplayAbilityActorInfo* ActorInfo, const FGameplayTagContainer* SourceTags = nullptr, const FGameplayTagContainer* TargetTags = nullptr, OUT FGameplayTagContainer* OptionalRelevantTags = nullptr) const override;
+	UE_API virtual bool CanActivateAbility(const FGameplayAbilitySpecHandle Handle, const FGameplayAbilityActorInfo* ActorInfo, const FGameplayTagContainer* SourceTags = nullptr, const FGameplayTagContainer* TargetTags = nullptr, OUT FGameplayTagContainer* OptionalRelevantTags = nullptr) const override;
 
-	virtual void ActivateAbility(const FGameplayAbilitySpecHandle Handle, const FGameplayAbilityActorInfo* OwnerInfo, const FGameplayAbilityActivationInfo ActivationInfo, const FGameplayEventData* TriggerEventData) override;
+	UE_API virtual void ActivateAbility(const FGameplayAbilitySpecHandle Handle, const FGameplayAbilityActorInfo* OwnerInfo, const FGameplayAbilityActivationInfo ActivationInfo, const FGameplayEventData* TriggerEventData) override;
 
-	virtual void InputReleased(const FGameplayAbilitySpecHandle Handle, const FGameplayAbilityActorInfo* ActorInfo, const FGameplayAbilityActivationInfo ActivationInfo) override;
+	UE_API virtual void InputReleased(const FGameplayAbilitySpecHandle Handle, const FGameplayAbilityActorInfo* ActorInfo, const FGameplayAbilityActivationInfo ActivationInfo) override;
 
-	virtual void CancelAbility(const FGameplayAbilitySpecHandle Handle, const FGameplayAbilityActorInfo* ActorInfo, const FGameplayAbilityActivationInfo ActivationInfo, bool bReplicateCancelAbility) override;
+	UE_API virtual void CancelAbility(const FGameplayAbilitySpecHandle Handle, const FGameplayAbilityActorInfo* ActorInfo, const FGameplayAbilityActivationInfo ActivationInfo, bool bReplicateCancelAbility) override;
 };
+
+#undef UE_API
diff --git a/Engine/Plugins/Runtime/GameplayAbilities/Source/GameplayAbilities/Public/Abilities/GameplayAbility_Montage.h b/Engine/Plugins/Runtime/GameplayAbilities/Source/GameplayAbilities/Public/Abilities/GameplayAbility_Montage.h
--- a/Engine/Plugins/Runtime/GameplayAbilities/Source/GameplayAbilities/Public/Abilities/GameplayAbility_Montage.h
+++ b/Engine/Plugins/Runtime/GameplayAbilities/Source/GameplayAbilities/Public/Abilities/GameplayAbility_Montage.h
@@ -10,20 +10,22 @@
 #include "Abilities/GameplayAbility.h"
 #include "GameplayAbility_Montage.generated.h"
 
+#define UE_API GAMEPLAYABILITIES_API
+
 class UAbilitySystemComponent;
 class UAnimMontage;
 
 /**
  *	A gameplay ability that plays a single montage and applies a GameplayEffect
  */
-UCLASS()
-class GAMEPLAYABILITIES_API UGameplayAbility_Montage : public UGameplayAbility
+UCLASS(MinimalAPI)
+class UGameplayAbility_Montage : public UGameplayAbility
 {
 	GENERATED_UCLASS_BODY()
 
 public:
 	
-	virtual void ActivateAbility(const FGameplayAbilitySpecHandle Handle, const FGameplayAbilityActorInfo* OwnerInfo, const FGameplayAbilityActivationInfo ActivationInfo, const FGameplayEventData* TriggerEventData) override;
+	UE_API virtual void ActivateAbility(const FGameplayAbilitySpecHandle Handle, const FGameplayAbilityActorInfo* OwnerInfo, const FGameplayAbilityActivationInfo ActivationInfo, const FGameplayEventData* TriggerEventData) override;
 
 	UPROPERTY(EditDefaultsOnly, Category = MontageAbility)
 	TObjectPtr<UAnimMontage> 	MontageToPlay;
@@ -42,7 +44,9 @@ public:
 	UPROPERTY(VisibleDefaultsOnly, Category = Deprecated)
 	TArray<TObjectPtr<const UGameplayEffect>>	GameplayEffectsWhileAnimating;
 
-	void OnMontageEnded(UAnimMontage* Montage, bool bInterrupted, TWeakObjectPtr<UAbilitySystemComponent> AbilitySystemComponent, TArray<struct FActiveGameplayEffectHandle>	AppliedEffects);
+	UE_API void OnMontageEnded(UAnimMontage* Montage, bool bInterrupted, TWeakObjectPtr<UAbilitySystemComponent> AbilitySystemComponent, TArray<struct FActiveGameplayEffectHandle>	AppliedEffects);
 
-	void GetGameplayEffectsWhileAnimating(TArray<const UGameplayEffect*>& OutEffects) const;
+	UE_API void GetGameplayEffectsWhileAnimating(TArray<const UGameplayEffect*>& OutEffects) const;
 };
+
+#undef UE_API
diff --git a/Engine/Plugins/Runtime/GameplayAbilities/Source/GameplayAbilities/Public/Abilities/Tasks/AbilityTask.h b/Engine/Plugins/Runtime/GameplayAbilities/Source/GameplayAbilities/Public/Abilities/Tasks/AbilityTask.h
--- a/Engine/Plugins/Runtime/GameplayAbilities/Source/GameplayAbilities/Public/Abilities/Tasks/AbilityTask.h
+++ b/Engine/Plugins/Runtime/GameplayAbilities/Source/GameplayAbilities/Public/Abilities/Tasks/AbilityTask.h
@@ -11,6 +11,8 @@
 #include "Abilities/GameplayAbility.h"
 #include "AbilityTask.generated.h"
 
+#define UE_API GAMEPLAYABILITIES_API
+
 class UAbilitySystemComponent;
 class UGameplayTasksComponent;
 
@@ -84,18 +86,18 @@ enum class EAbilityTaskWaitState : uint8
 	WaitingOnAvatar = 0x04
 };
 
-UCLASS(Abstract)
-class GAMEPLAYABILITIES_API UAbilityTask : public UGameplayTask
+UCLASS(Abstract, MinimalAPI)
+class UAbilityTask : public UGameplayTask
 {
 	GENERATED_UCLASS_BODY()
 
-	virtual void OnDestroy(bool bInOwnerFinished) override;
-	virtual void BeginDestroy() override;
+	UE_API virtual void OnDestroy(bool bInOwnerFinished) override;
+	UE_API virtual void BeginDestroy() override;
 	
 	/** Returns spec handle for owning ability */
-	FGameplayAbilitySpecHandle GetAbilitySpecHandle() const;
+	UE_API FGameplayAbilitySpecHandle GetAbilitySpecHandle() const;
 
-	void SetAbilitySystemComponent(UAbilitySystemComponent* InAbilitySystemComponent);
+	UE_API void SetAbilitySystemComponent(UAbilitySystemComponent* InAbilitySystemComponent);
 
 	/** GameplayAbility that created us */
 	UPROPERTY()
@@ -105,23 +107,23 @@ class GAMEPLAYABILITIES_API UAbilityTask : public UGameplayTask
 	TWeakObjectPtr<UAbilitySystemComponent> AbilitySystemComponent;
 
 	/** Returns true if the ability is a locally predicted ability running on a client. Usually this means we need to tell the server something. */
-	bool IsPredictingClient() const;
+	UE_API bool IsPredictingClient() const;
 
 	/** Returns true if we are executing the ability on the server for a non locally controlled client */
-	bool IsForRemoteClient() const;
+	UE_API bool IsForRemoteClient() const;
 
 	/** Returns true if we are executing the ability on the locally controlled client */
-	bool IsLocallyControlled() const;
+	UE_API bool IsLocallyControlled() const;
 
 	/** Returns ActivationPredictionKey of owning ability */
-	FPredictionKey GetActivationPredictionKey() const;
+	UE_API FPredictionKey GetActivationPredictionKey() const;
 
 	/** This should be called prior to broadcasting delegates back into the ability graph. This makes sure the ability is still active.  */
-	bool ShouldBroadcastAbilityTaskDelegates() const;
+	UE_API bool ShouldBroadcastAbilityTaskDelegates() const;
 
-	virtual void InitSimulatedTask(UGameplayTasksComponent& InGameplayTasksComponent) override;
+	UE_API virtual void InitSimulatedTask(UGameplayTasksComponent& InGameplayTasksComponent) override;
 
-	static void DebugRecordAbilityTaskCreatedByAbility(const UObject* Ability);
+	static UE_API void DebugRecordAbilityTaskCreatedByAbility(const UObject* Ability);
 
 	/** Helper function for instantiating and initializing a new task */
 	template <class T>
@@ -153,14 +155,14 @@ class GAMEPLAYABILITIES_API UAbilityTask : public UGameplayTask
 	}
 
 	/** Called when the ability task is waiting on remote player data. IF the remote player ends the ability prematurely, and a task with this set is still running, the ability is killed. */
-	void SetWaitingOnRemotePlayerData();
-	void ClearWaitingOnRemotePlayerData();
-	virtual bool IsWaitingOnRemotePlayerdata() const override;
+	UE_API void SetWaitingOnRemotePlayerData();
+	UE_API void ClearWaitingOnRemotePlayerData();
+	UE_API virtual bool IsWaitingOnRemotePlayerdata() const override;
 
 	/** same as RemotePlayerData but for ACharacter type of state (movement state, etc) */
-	void SetWaitingOnAvatar();
-	void ClearWaitingOnAvatar();
-	virtual bool IsWaitingOnAvatar() const override;
+	UE_API void SetWaitingOnAvatar();
+	UE_API void ClearWaitingOnAvatar();
+	UE_API virtual bool IsWaitingOnAvatar() const override;
 
 	/** What we are waiting on */
 	uint8 WaitStateBitMask;
@@ -168,7 +170,7 @@ class GAMEPLAYABILITIES_API UAbilityTask : public UGameplayTask
 
 protected:
 	/** Helper method for registering client replicated callbacks */
-	bool CallOrAddReplicatedDelegate(EAbilityGenericReplicatedEvent::Type Event, FSimpleMulticastDelegate::FDelegate Delegate);
+	UE_API bool CallOrAddReplicatedDelegate(EAbilityGenericReplicatedEvent::Type Event, FSimpleMulticastDelegate::FDelegate Delegate);
 };
 
 //For searching through lists of ability instances
@@ -208,3 +210,5 @@ struct FAbilityInstanceClassPredicate
 		if (Ability) \
 			Ability->AddAbilityTaskDebugMessage(this, FString::Printf(TEXT(Format), ##__VA_ARGS__)); \
 	} 
+
+#undef UE_API
diff --git a/Engine/Plugins/Runtime/GameplayAbilities/Source/GameplayAbilities/Public/Abilities/Tasks/AbilityTask_ApplyRootMotionConstantForce.h b/Engine/Plugins/Runtime/GameplayAbilities/Source/GameplayAbilities/Public/Abilities/Tasks/AbilityTask_ApplyRootMotionConstantForce.h
--- a/Engine/Plugins/Runtime/GameplayAbilities/Source/GameplayAbilities/Public/Abilities/Tasks/AbilityTask_ApplyRootMotionConstantForce.h
+++ b/Engine/Plugins/Runtime/GameplayAbilities/Source/GameplayAbilities/Public/Abilities/Tasks/AbilityTask_ApplyRootMotionConstantForce.h
@@ -7,6 +7,8 @@
 #include "Abilities/Tasks/AbilityTask_ApplyRootMotion_Base.h"
 #include "AbilityTask_ApplyRootMotionConstantForce.generated.h"
 
+#define UE_API GAMEPLAYABILITIES_API
+
 class UCharacterMovementComponent;
 class UCurveFloat;
 class UGameplayTasksComponent;
@@ -18,8 +20,8 @@ class AActor;
 /**
  *	Applies force to character's movement
  */
-UCLASS()
-class GAMEPLAYABILITIES_API UAbilityTask_ApplyRootMotionConstantForce : public UAbilityTask_ApplyRootMotion_Base
+UCLASS(MinimalAPI)
+class UAbilityTask_ApplyRootMotionConstantForce : public UAbilityTask_ApplyRootMotion_Base
 {
 	GENERATED_UCLASS_BODY()
 
@@ -28,7 +30,7 @@ class GAMEPLAYABILITIES_API UAbilityTask_ApplyRootMotionConstantForce : public U
 
 	/** Apply force to character's movement */
 	UFUNCTION(BlueprintCallable, Category = "Ability|Tasks", meta = (HidePin = "OwningAbility", DefaultToSelf = "OwningAbility", BlueprintInternalUseOnly = "TRUE"))
-	static UAbilityTask_ApplyRootMotionConstantForce* ApplyRootMotionConstantForce
+	static UE_API UAbilityTask_ApplyRootMotionConstantForce* ApplyRootMotionConstantForce
 	(
 		UGameplayAbility* OwningAbility, 
 		FName TaskInstanceName, 
@@ -44,14 +46,14 @@ class GAMEPLAYABILITIES_API UAbilityTask_ApplyRootMotionConstantForce : public U
 	);
 
 	/** Tick function for this task, if bTickingTask == true */
-	virtual void TickTask(float DeltaTime) override;
+	UE_API virtual void TickTask(float DeltaTime) override;
 
-	virtual void PreDestroyFromReplication() override;
-	virtual void OnDestroy(bool AbilityIsEnding) override;
+	UE_API virtual void PreDestroyFromReplication() override;
+	UE_API virtual void OnDestroy(bool AbilityIsEnding) override;
 
 protected:
 
-	virtual void SharedInitAndApply() override;
+	UE_API virtual void SharedInitAndApply() override;
 
 protected:
 
@@ -80,3 +82,5 @@ protected:
 	bool bEnableGravity;
 
 };
+
+#undef UE_API
diff --git a/Engine/Plugins/Runtime/GameplayAbilities/Source/GameplayAbilities/Public/Abilities/Tasks/AbilityTask_ApplyRootMotionJumpForce.h b/Engine/Plugins/Runtime/GameplayAbilities/Source/GameplayAbilities/Public/Abilities/Tasks/AbilityTask_ApplyRootMotionJumpForce.h
--- a/Engine/Plugins/Runtime/GameplayAbilities/Source/GameplayAbilities/Public/Abilities/Tasks/AbilityTask_ApplyRootMotionJumpForce.h
+++ b/Engine/Plugins/Runtime/GameplayAbilities/Source/GameplayAbilities/Public/Abilities/Tasks/AbilityTask_ApplyRootMotionJumpForce.h
@@ -7,6 +7,8 @@
 #include "Abilities/Tasks/AbilityTask_ApplyRootMotion_Base.h"
 #include "AbilityTask_ApplyRootMotionJumpForce.generated.h"
 
+#define UE_API GAMEPLAYABILITIES_API
+
 class UCharacterMovementComponent;
 class UCurveFloat;
 class UCurveVector;
@@ -19,8 +21,8 @@ class AActor;
 /**
  *	Applies force to character's movement
  */
-UCLASS()
-class GAMEPLAYABILITIES_API UAbilityTask_ApplyRootMotionJumpForce : public UAbilityTask_ApplyRootMotion_Base
+UCLASS(MinimalAPI)
+class UAbilityTask_ApplyRootMotionJumpForce : public UAbilityTask_ApplyRootMotion_Base
 {
 	GENERATED_UCLASS_BODY()
 
@@ -31,27 +33,27 @@ class GAMEPLAYABILITIES_API UAbilityTask_ApplyRootMotionJumpForce : public UAbil
 	FApplyRootMotionJumpForceDelegate OnLanded;
 
 	UFUNCTION(BlueprintCallable, Category="Ability|Tasks")
-	void Finish();
+	UE_API void Finish();
 
 	UFUNCTION()
-	void OnLandedCallback(const FHitResult& Hit);
+	UE_API void OnLandedCallback(const FHitResult& Hit);
 
 	/** Apply force to character's movement */
 	UFUNCTION(BlueprintCallable, Category = "Ability|Tasks", meta = (HidePin = "OwningAbility", DefaultToSelf = "OwningAbility", BlueprintInternalUseOnly = "TRUE"))
-	static UAbilityTask_ApplyRootMotionJumpForce* ApplyRootMotionJumpForce(UGameplayAbility* OwningAbility, FName TaskInstanceName, FRotator Rotation, float Distance, float Height, float Duration, float MinimumLandedTriggerTime,
+	static UE_API UAbilityTask_ApplyRootMotionJumpForce* ApplyRootMotionJumpForce(UGameplayAbility* OwningAbility, FName TaskInstanceName, FRotator Rotation, float Distance, float Height, float Duration, float MinimumLandedTriggerTime,
 		bool bFinishOnLanded, ERootMotionFinishVelocityMode VelocityOnFinishMode, FVector SetVelocityOnFinish, float ClampVelocityOnFinish, UCurveVector* PathOffsetCurve, UCurveFloat* TimeMappingCurve);
 
-	virtual void Activate() override;
+	UE_API virtual void Activate() override;
 
 	/** Tick function for this task, if bTickingTask == true */
-	virtual void TickTask(float DeltaTime) override;
+	UE_API virtual void TickTask(float DeltaTime) override;
 
-	virtual void PreDestroyFromReplication() override;
-	virtual void OnDestroy(bool AbilityIsEnding) override;
+	UE_API virtual void PreDestroyFromReplication() override;
+	UE_API virtual void OnDestroy(bool AbilityIsEnding) override;
 
 protected:
 
-	virtual void SharedInitAndApply() override;
+	UE_API virtual void SharedInitAndApply() override;
 
 	/**
 	* Work-around for OnLanded being called during bClientUpdating in movement replay code
@@ -59,7 +61,7 @@ protected:
 	* If we don't, we end up removing root motion from a replay root motion set instead
 	* of the real one
 	*/
-	void TriggerLanded();
+	UE_API void TriggerLanded();
 
 protected:
 
@@ -96,3 +98,5 @@ protected:
 
 	bool bHasLanded;
 };
+
+#undef UE_API
diff --git a/Engine/Plugins/Runtime/GameplayAbilities/Source/GameplayAbilities/Public/Abilities/Tasks/AbilityTask_ApplyRootMotionMoveToActorForce.h b/Engine/Plugins/Runtime/GameplayAbilities/Source/GameplayAbilities/Public/Abilities/Tasks/AbilityTask_ApplyRootMotionMoveToActorForce.h
--- a/Engine/Plugins/Runtime/GameplayAbilities/Source/GameplayAbilities/Public/Abilities/Tasks/AbilityTask_ApplyRootMotionMoveToActorForce.h
+++ b/Engine/Plugins/Runtime/GameplayAbilities/Source/GameplayAbilities/Public/Abilities/Tasks/AbilityTask_ApplyRootMotionMoveToActorForce.h
@@ -9,6 +9,8 @@
 #include "Delegates/IDelegateInstance.h"
 #include "AbilityTask_ApplyRootMotionMoveToActorForce.generated.h"
 
+#define UE_API GAMEPLAYABILITIES_API
+
 class UCharacterMovementComponent;
 class UCurveFloat;
 class UCurveVector;
@@ -33,8 +35,8 @@ enum class ERootMotionMoveToActorTargetOffsetType : uint8
 /**
  *	Applies force to character's movement
  */
-UCLASS()
-class GAMEPLAYABILITIES_API UAbilityTask_ApplyRootMotionMoveToActorForce : public UAbilityTask_ApplyRootMotion_Base
+UCLASS(MinimalAPI)
+class UAbilityTask_ApplyRootMotionMoveToActorForce : public UAbilityTask_ApplyRootMotion_Base
 {
 	GENERATED_UCLASS_BODY()
 
@@ -43,37 +45,37 @@ class GAMEPLAYABILITIES_API UAbilityTask_ApplyRootMotionMoveToActorForce : publi
 
 	/** Apply force to character's movement to move to a target actor */
 	UFUNCTION(BlueprintCallable, Category = "Ability|Tasks", meta = (HidePin = "OwningAbility", DefaultToSelf = "OwningAbility", BlueprintInternalUseOnly = "TRUE"))
-	static UAbilityTask_ApplyRootMotionMoveToActorForce* ApplyRootMotionMoveToActorForce(UGameplayAbility* OwningAbility, FName TaskInstanceName, AActor* TargetActor, FVector TargetLocationOffset, ERootMotionMoveToActorTargetOffsetType OffsetAlignment, float Duration, UCurveFloat* TargetLerpSpeedHorizontal, UCurveFloat* TargetLerpSpeedVertical, bool bSetNewMovementMode, EMovementMode MovementMode, bool bRestrictSpeedToExpected, UCurveVector* PathOffsetCurve, UCurveFloat* TimeMappingCurve, ERootMotionFinishVelocityMode VelocityOnFinishMode, FVector SetVelocityOnFinish, float ClampVelocityOnFinish, bool bDisableDestinationReachedInterrupt, float ReachedDestinationDistance = 50.0f);
+	static UE_API UAbilityTask_ApplyRootMotionMoveToActorForce* ApplyRootMotionMoveToActorForce(UGameplayAbility* OwningAbility, FName TaskInstanceName, AActor* TargetActor, FVector TargetLocationOffset, ERootMotionMoveToActorTargetOffsetType OffsetAlignment, float Duration, UCurveFloat* TargetLerpSpeedHorizontal, UCurveFloat* TargetLerpSpeedVertical, bool bSetNewMovementMode, EMovementMode MovementMode, bool bRestrictSpeedToExpected, UCurveVector* PathOffsetCurve, UCurveFloat* TimeMappingCurve, ERootMotionFinishVelocityMode VelocityOnFinishMode, FVector SetVelocityOnFinish, float ClampVelocityOnFinish, bool bDisableDestinationReachedInterrupt, float ReachedDestinationDistance = 50.0f);
 
 	/** Apply force to character's movement to move to a target component with relative location*/
 	UFUNCTION(BlueprintCallable, Category = "Ability|Tasks", meta = (HidePin = "OwningAbility", DefaultToSelf = "OwningAbility", BlueprintInternalUseOnly = "TRUE"))
-	static UAbilityTask_ApplyRootMotionMoveToActorForce* ApplyRootMotionMoveToComponentForce(UGameplayAbility* OwningAbility, FName TaskInstanceName, USceneComponent* TargetComponent, FVector TargetComponentRelativeLocation, FVector TargetLocationOffset, ERootMotionMoveToActorTargetOffsetType OffsetAlignment, float Duration, UCurveFloat* TargetLerpSpeedHorizontal, UCurveFloat* TargetLerpSpeedVertical, bool bSetNewMovementMode, EMovementMode MovementMode, bool bRestrictSpeedToExpected, UCurveVector* PathOffsetCurve, UCurveFloat* TimeMappingCurve, ERootMotionFinishVelocityMode VelocityOnFinishMode, FVector SetVelocityOnFinish, float ClampVelocityOnFinish, bool bDisableDestinationReachedInterrupt, float ReachedDestinationDistance = 50.0f);
+	static UE_API UAbilityTask_ApplyRootMotionMoveToActorForce* ApplyRootMotionMoveToComponentForce(UGameplayAbility* OwningAbility, FName TaskInstanceName, USceneComponent* TargetComponent, FVector TargetComponentRelativeLocation, FVector TargetLocationOffset, ERootMotionMoveToActorTargetOffsetType OffsetAlignment, float Duration, UCurveFloat* TargetLerpSpeedHorizontal, UCurveFloat* TargetLerpSpeedVertical, bool bSetNewMovementMode, EMovementMode MovementMode, bool bRestrictSpeedToExpected, UCurveVector* PathOffsetCurve, UCurveFloat* TimeMappingCurve, ERootMotionFinishVelocityMode VelocityOnFinishMode, FVector SetVelocityOnFinish, float ClampVelocityOnFinish, bool bDisableDestinationReachedInterrupt, float ReachedDestinationDistance = 50.0f);
 	
 	/** Apply force to character's movement using an index into targetData instead of using an actor directly. */
 	UFUNCTION(BlueprintCallable, Category = "Ability|Tasks", meta = (HidePin = "OwningAbility", DefaultToSelf = "OwningAbility", BlueprintInternalUseOnly = "TRUE"))
-	static UAbilityTask_ApplyRootMotionMoveToActorForce* ApplyRootMotionMoveToTargetDataActorForce(UGameplayAbility* OwningAbility, FName TaskInstanceName, FGameplayAbilityTargetDataHandle TargetDataHandle, int32 TargetDataIndex, int32 TargetActorIndex, FVector TargetLocationOffset, ERootMotionMoveToActorTargetOffsetType OffsetAlignment, float Duration, UCurveFloat* TargetLerpSpeedHorizontal, UCurveFloat* TargetLerpSpeedVertical, bool bSetNewMovementMode, EMovementMode MovementMode, bool bRestrictSpeedToExpected, UCurveVector* PathOffsetCurve, UCurveFloat* TimeMappingCurve, ERootMotionFinishVelocityMode VelocityOnFinishMode, FVector SetVelocityOnFinish, float ClampVelocityOnFinish, bool bDisableDestinationReachedInterrupt, float ReachedDestinationDistance = 50.0f);
+	static UE_API UAbilityTask_ApplyRootMotionMoveToActorForce* ApplyRootMotionMoveToTargetDataActorForce(UGameplayAbility* OwningAbility, FName TaskInstanceName, FGameplayAbilityTargetDataHandle TargetDataHandle, int32 TargetDataIndex, int32 TargetActorIndex, FVector TargetLocationOffset, ERootMotionMoveToActorTargetOffsetType OffsetAlignment, float Duration, UCurveFloat* TargetLerpSpeedHorizontal, UCurveFloat* TargetLerpSpeedVertical, bool bSetNewMovementMode, EMovementMode MovementMode, bool bRestrictSpeedToExpected, UCurveVector* PathOffsetCurve, UCurveFloat* TimeMappingCurve, ERootMotionFinishVelocityMode VelocityOnFinishMode, FVector SetVelocityOnFinish, float ClampVelocityOnFinish, bool bDisableDestinationReachedInterrupt, float ReachedDestinationDistance = 50.0f);
 
 	/** Tick function for this task, if bTickingTask == true */
-	virtual void TickTask(float DeltaTime) override;
+	UE_API virtual void TickTask(float DeltaTime) override;
 
-	virtual void PreDestroyFromReplication() override;
-	virtual void OnDestroy(bool AbilityIsEnding) override;
+	UE_API virtual void PreDestroyFromReplication() override;
+	UE_API virtual void OnDestroy(bool AbilityIsEnding) override;
 
 	UFUNCTION()
-	void OnTargetActorSwapped(class AActor* OriginalTarget, class AActor* NewTarget);
+	UE_API void OnTargetActorSwapped(class AActor* OriginalTarget, class AActor* NewTarget);
 
 protected:
 
-	virtual void SharedInitAndApply() override;
+	UE_API virtual void SharedInitAndApply() override;
 
-	bool UpdateTargetLocation(float DeltaTime);
+	UE_API bool UpdateTargetLocation(float DeltaTime);
 
-	void SetRootMotionTargetLocation(FVector NewTargetLocation);
+	UE_API void SetRootMotionTargetLocation(FVector NewTargetLocation);
 
-	FVector CalculateTargetOffset() const;
+	UE_API FVector CalculateTargetOffset() const;
 
 	UFUNCTION()
-	void OnRep_TargetLocation();
+	UE_API void OnRep_TargetLocation();
 
 protected:
 
@@ -146,3 +148,5 @@ protected:
 	EMovementMode PreviousMovementMode;
 	uint8 PreviousCustomMode;
 };
+
+#undef UE_API
diff --git a/Engine/Plugins/Runtime/GameplayAbilities/Source/GameplayAbilities/Public/Abilities/Tasks/AbilityTask_ApplyRootMotionMoveToForce.h b/Engine/Plugins/Runtime/GameplayAbilities/Source/GameplayAbilities/Public/Abilities/Tasks/AbilityTask_ApplyRootMotionMoveToForce.h
--- a/Engine/Plugins/Runtime/GameplayAbilities/Source/GameplayAbilities/Public/Abilities/Tasks/AbilityTask_ApplyRootMotionMoveToForce.h
+++ b/Engine/Plugins/Runtime/GameplayAbilities/Source/GameplayAbilities/Public/Abilities/Tasks/AbilityTask_ApplyRootMotionMoveToForce.h
@@ -7,6 +7,8 @@
 #include "Abilities/Tasks/AbilityTask_ApplyRootMotion_Base.h"
 #include "AbilityTask_ApplyRootMotionMoveToForce.generated.h"
 
+#define UE_API GAMEPLAYABILITIES_API
+
 class UCharacterMovementComponent;
 class UCurveVector;
 class UGameplayTasksComponent;
@@ -19,8 +21,8 @@ class AActor;
 /**
  *	Applies force to character's movement
  */
-UCLASS()
-class GAMEPLAYABILITIES_API UAbilityTask_ApplyRootMotionMoveToForce : public UAbilityTask_ApplyRootMotion_Base
+UCLASS(MinimalAPI)
+class UAbilityTask_ApplyRootMotionMoveToForce : public UAbilityTask_ApplyRootMotion_Base
 {
 	GENERATED_BODY()
 	
@@ -33,16 +35,16 @@ public:
 
 	/** Apply force to character's movement */
 	UFUNCTION(BlueprintCallable, Category = "Ability|Tasks", meta = (HidePin = "OwningAbility", DefaultToSelf = "OwningAbility", BlueprintInternalUseOnly = "TRUE"))
-	static UAbilityTask_ApplyRootMotionMoveToForce* ApplyRootMotionMoveToForce(UGameplayAbility* OwningAbility, FName TaskInstanceName, FVector TargetLocation, float Duration, bool bSetNewMovementMode, EMovementMode MovementMode, bool bRestrictSpeedToExpected, UCurveVector* PathOffsetCurve, ERootMotionFinishVelocityMode VelocityOnFinishMode, FVector SetVelocityOnFinish, float ClampVelocityOnFinish);
+	static UE_API UAbilityTask_ApplyRootMotionMoveToForce* ApplyRootMotionMoveToForce(UGameplayAbility* OwningAbility, FName TaskInstanceName, FVector TargetLocation, float Duration, bool bSetNewMovementMode, EMovementMode MovementMode, bool bRestrictSpeedToExpected, UCurveVector* PathOffsetCurve, ERootMotionFinishVelocityMode VelocityOnFinishMode, FVector SetVelocityOnFinish, float ClampVelocityOnFinish);
 
 	/** Tick function for this task, if bTickingTask == true */
-	virtual void TickTask(float DeltaTime) override;
+	UE_API virtual void TickTask(float DeltaTime) override;
 
-	virtual void PreDestroyFromReplication() override;
-	virtual void OnDestroy(bool AbilityIsEnding) override;
+	UE_API virtual void PreDestroyFromReplication() override;
+	UE_API virtual void OnDestroy(bool AbilityIsEnding) override;
 
 protected:
-	virtual void SharedInitAndApply() override;
+	UE_API virtual void SharedInitAndApply() override;
 
 	UPROPERTY(Replicated)
 	FVector StartLocation;
@@ -71,3 +73,5 @@ protected:
 	EMovementMode PreviousMovementMode = EMovementMode::MOVE_None;
 	uint8 PreviousCustomMovementMode = 0;
 };
+
+#undef UE_API
diff --git a/Engine/Plugins/Runtime/GameplayAbilities/Source/GameplayAbilities/Public/Abilities/Tasks/AbilityTask_ApplyRootMotionRadialForce.h b/Engine/Plugins/Runtime/GameplayAbilities/Source/GameplayAbilities/Public/Abilities/Tasks/AbilityTask_ApplyRootMotionRadialForce.h
--- a/Engine/Plugins/Runtime/GameplayAbilities/Source/GameplayAbilities/Public/Abilities/Tasks/AbilityTask_ApplyRootMotionRadialForce.h
+++ b/Engine/Plugins/Runtime/GameplayAbilities/Source/GameplayAbilities/Public/Abilities/Tasks/AbilityTask_ApplyRootMotionRadialForce.h
@@ -7,6 +7,8 @@
 #include "Abilities/Tasks/AbilityTask_ApplyRootMotion_Base.h"
 #include "AbilityTask_ApplyRootMotionRadialForce.generated.h"
 
+#define UE_API GAMEPLAYABILITIES_API
+
 class UCharacterMovementComponent;
 class UCurveFloat;
 class UGameplayTasksComponent;
@@ -19,8 +21,8 @@ class AActor;
 /**
  *	Applies force to character's movement
  */
-UCLASS()
-class GAMEPLAYABILITIES_API UAbilityTask_ApplyRootMotionRadialForce : public UAbilityTask_ApplyRootMotion_Base
+UCLASS(MinimalAPI)
+class UAbilityTask_ApplyRootMotionRadialForce : public UAbilityTask_ApplyRootMotion_Base
 {
 	GENERATED_UCLASS_BODY()
 
@@ -29,17 +31,17 @@ class GAMEPLAYABILITIES_API UAbilityTask_ApplyRootMotionRadialForce : public UAb
 
 	/** Apply force to character's movement */
 	UFUNCTION(BlueprintCallable, Category = "Ability|Tasks", meta = (HidePin = "OwningAbility", DefaultToSelf = "OwningAbility", BlueprintInternalUseOnly = "TRUE"))
-	static UAbilityTask_ApplyRootMotionRadialForce* ApplyRootMotionRadialForce(UGameplayAbility* OwningAbility, FName TaskInstanceName, FVector Location, AActor* LocationActor, float Strength, float Duration, float Radius, bool bIsPush, bool bIsAdditive, bool bNoZForce, UCurveFloat* StrengthDistanceFalloff, UCurveFloat* StrengthOverTime, bool bUseFixedWorldDirection, FRotator FixedWorldDirection, ERootMotionFinishVelocityMode VelocityOnFinishMode, FVector SetVelocityOnFinish, float ClampVelocityOnFinish);
+	static UE_API UAbilityTask_ApplyRootMotionRadialForce* ApplyRootMotionRadialForce(UGameplayAbility* OwningAbility, FName TaskInstanceName, FVector Location, AActor* LocationActor, float Strength, float Duration, float Radius, bool bIsPush, bool bIsAdditive, bool bNoZForce, UCurveFloat* StrengthDistanceFalloff, UCurveFloat* StrengthOverTime, bool bUseFixedWorldDirection, FRotator FixedWorldDirection, ERootMotionFinishVelocityMode VelocityOnFinishMode, FVector SetVelocityOnFinish, float ClampVelocityOnFinish);
 
 	/** Tick function for this task, if bTickingTask == true */
-	virtual void TickTask(float DeltaTime) override;
+	UE_API virtual void TickTask(float DeltaTime) override;
 
-	virtual void PreDestroyFromReplication() override;
-	virtual void OnDestroy(bool AbilityIsEnding) override;
+	UE_API virtual void PreDestroyFromReplication() override;
+	UE_API virtual void OnDestroy(bool AbilityIsEnding) override;
 
 protected:
 
-	virtual void SharedInitAndApply() override;
+	UE_API virtual void SharedInitAndApply() override;
 
 protected:
 
@@ -90,3 +92,5 @@ protected:
 	UPROPERTY(Replicated)
 	FRotator FixedWorldDirection;
 };
+
+#undef UE_API
diff --git a/Engine/Plugins/Runtime/GameplayAbilities/Source/GameplayAbilities/Public/Abilities/Tasks/AbilityTask_ApplyRootMotion_Base.h b/Engine/Plugins/Runtime/GameplayAbilities/Source/GameplayAbilities/Public/Abilities/Tasks/AbilityTask_ApplyRootMotion_Base.h
--- a/Engine/Plugins/Runtime/GameplayAbilities/Source/GameplayAbilities/Public/Abilities/Tasks/AbilityTask_ApplyRootMotion_Base.h
+++ b/Engine/Plugins/Runtime/GameplayAbilities/Source/GameplayAbilities/Public/Abilities/Tasks/AbilityTask_ApplyRootMotion_Base.h
@@ -5,6 +5,8 @@
 #include "Abilities/Tasks/AbilityTask.h"
 #include "AbilityTask_ApplyRootMotion_Base.generated.h"
 
+#define UE_API GAMEPLAYABILITIES_API
+
 class UCharacterMovementComponent;
 enum class ERootMotionFinishVelocityMode : uint8;
 
@@ -12,20 +14,20 @@ enum class ERootMotionFinishVelocityMode : uint8;
 DECLARE_MULTICAST_DELEGATE_TwoParams(FOnTargetActorSwapped, AActor*, AActor*);
 
 /** Base class for ability tasks that apply root motion */
-UCLASS()
-class GAMEPLAYABILITIES_API UAbilityTask_ApplyRootMotion_Base : public UAbilityTask
+UCLASS(MinimalAPI)
+class UAbilityTask_ApplyRootMotion_Base : public UAbilityTask
 {
 	GENERATED_UCLASS_BODY()
 
-	virtual void InitSimulatedTask(UGameplayTasksComponent& InGameplayTasksComponent) override;
+	UE_API virtual void InitSimulatedTask(UGameplayTasksComponent& InGameplayTasksComponent) override;
 
 	//..See notes on delegate definition FOnTargetActorSwapped.
-	static FOnTargetActorSwapped OnTargetActorSwapped;
+	static UE_API FOnTargetActorSwapped OnTargetActorSwapped;
 
 protected:
 
 	virtual void SharedInitAndApply() {};
-	virtual bool HasTimedOut() const;
+	UE_API virtual bool HasTimedOut() const;
 
 	UPROPERTY(Replicated)
 	FName ForceName;
@@ -52,3 +54,5 @@ protected:
 	float StartTime;
 	float EndTime;
 };
+
+#undef UE_API
diff --git a/Engine/Plugins/Runtime/GameplayAbilities/Source/GameplayAbilities/Public/Abilities/Tasks/AbilityTask_MoveToLocation.h b/Engine/Plugins/Runtime/GameplayAbilities/Source/GameplayAbilities/Public/Abilities/Tasks/AbilityTask_MoveToLocation.h
--- a/Engine/Plugins/Runtime/GameplayAbilities/Source/GameplayAbilities/Public/Abilities/Tasks/AbilityTask_MoveToLocation.h
+++ b/Engine/Plugins/Runtime/GameplayAbilities/Source/GameplayAbilities/Public/Abilities/Tasks/AbilityTask_MoveToLocation.h
@@ -6,6 +6,8 @@
 #include "Abilities/Tasks/AbilityTask.h"
 #include "AbilityTask_MoveToLocation.generated.h"
 
+#define UE_API GAMEPLAYABILITIES_API
+
 class UCurveFloat;
 class UCurveVector;
 class UGameplayTasksComponent;
@@ -27,26 +29,26 @@ class UGameplayTasksComponent;
  *	This will RESET your character's current movement mode! If you wish to maintain PHYS_Flying or PHYS_Custom, you must
  *	reset it on completion.!
  */
-UCLASS()
-class GAMEPLAYABILITIES_API UAbilityTask_MoveToLocation : public UAbilityTask
+UCLASS(MinimalAPI)
+class UAbilityTask_MoveToLocation : public UAbilityTask
 {
 	GENERATED_UCLASS_BODY()
 
 	UPROPERTY(BlueprintAssignable)
 	FMoveToLocationDelegate		OnTargetLocationReached;
 
-	virtual void InitSimulatedTask(UGameplayTasksComponent& InGameplayTasksComponent) override;
+	UE_API virtual void InitSimulatedTask(UGameplayTasksComponent& InGameplayTasksComponent) override;
 
 	/** Move to the specified location, using the vector curve (range 0 - 1) if specified, otherwise the float curve (range 0 - 1) or fallback to linear interpolation */
 	UFUNCTION(BlueprintCallable, Category = "Ability|Tasks", meta = (HidePin = "OwningAbility", DefaultToSelf = "OwningAbility", BlueprintInternalUseOnly = "TRUE"))
-	static UAbilityTask_MoveToLocation* MoveToLocation(UGameplayAbility* OwningAbility, FName TaskInstanceName, FVector Location, float Duration, UCurveFloat* OptionalInterpolationCurve, UCurveVector* OptionalVectorInterpolationCurve);
+	static UE_API UAbilityTask_MoveToLocation* MoveToLocation(UGameplayAbility* OwningAbility, FName TaskInstanceName, FVector Location, float Duration, UCurveFloat* OptionalInterpolationCurve, UCurveVector* OptionalVectorInterpolationCurve);
 
-	virtual void Activate() override;
+	UE_API virtual void Activate() override;
 
 	/** Tick function for this task, if bTickingTask == true */
-	virtual void TickTask(float DeltaTime) override;
+	UE_API virtual void TickTask(float DeltaTime) override;
 
-	virtual void OnDestroy(bool AbilityIsEnding) override;
+	UE_API virtual void OnDestroy(bool AbilityIsEnding) override;
 
 protected:
 
@@ -76,3 +78,5 @@ protected:
 
 	
 };
+
+#undef UE_API
diff --git a/Engine/Plugins/Runtime/GameplayAbilities/Source/GameplayAbilities/Public/Abilities/Tasks/AbilityTask_NetworkSyncPoint.h b/Engine/Plugins/Runtime/GameplayAbilities/Source/GameplayAbilities/Public/Abilities/Tasks/AbilityTask_NetworkSyncPoint.h
--- a/Engine/Plugins/Runtime/GameplayAbilities/Source/GameplayAbilities/Public/Abilities/Tasks/AbilityTask_NetworkSyncPoint.h
+++ b/Engine/Plugins/Runtime/GameplayAbilities/Source/GameplayAbilities/Public/Abilities/Tasks/AbilityTask_NetworkSyncPoint.h
@@ -6,6 +6,8 @@
 #include "Abilities/Tasks/AbilityTask.h"
 #include "AbilityTask_NetworkSyncPoint.generated.h"
 
+#define UE_API GAMEPLAYABILITIES_API
+
 DECLARE_DYNAMIC_MULTICAST_DELEGATE(FNetworkSyncDelegate);
 
 UENUM()
@@ -22,8 +24,8 @@ enum class EAbilityTaskNetSyncType : uint8
 };
 
 /** Task for providing a generic sync point for client server (one can wait for a signal from the other) */
-UCLASS()
-class GAMEPLAYABILITIES_API UAbilityTask_NetworkSyncPoint : public UAbilityTask
+UCLASS(MinimalAPI)
+class UAbilityTask_NetworkSyncPoint : public UAbilityTask
 {
 	GENERATED_UCLASS_BODY()
 
@@ -31,9 +33,9 @@ class GAMEPLAYABILITIES_API UAbilityTask_NetworkSyncPoint : public UAbilityTask
 	FNetworkSyncDelegate	OnSync;
 
 	UFUNCTION()
-	void OnSignalCallback();
+	UE_API void OnSignalCallback();
 
-	virtual void Activate() override;
+	UE_API virtual void Activate() override;
 
 	/**
 	 *	
@@ -50,14 +52,16 @@ class GAMEPLAYABILITIES_API UAbilityTask_NetworkSyncPoint : public UAbilityTask
 	 *	
 	 */
 	UFUNCTION(BlueprintCallable, Category="Ability|Tasks", meta = (HidePin = "OwningAbility", DefaultToSelf = "OwningAbility", BlueprintInternalUseOnly = "TRUE"))
-	static UAbilityTask_NetworkSyncPoint* WaitNetSync(UGameplayAbility* OwningAbility, EAbilityTaskNetSyncType SyncType);
+	static UE_API UAbilityTask_NetworkSyncPoint* WaitNetSync(UGameplayAbility* OwningAbility, EAbilityTaskNetSyncType SyncType);
 
 protected:
 
-	void SyncFinished();
+	UE_API void SyncFinished();
 
 	/** The event we replicate */
 	EAbilityGenericReplicatedEvent::Type	ReplicatedEventToListenFor;
 
 	EAbilityTaskNetSyncType SyncType;
 };
+
+#undef UE_API
diff --git a/Engine/Plugins/Runtime/GameplayAbilities/Source/GameplayAbilities/Public/Abilities/Tasks/AbilityTask_PlayAnimAndWait.h b/Engine/Plugins/Runtime/GameplayAbilities/Source/GameplayAbilities/Public/Abilities/Tasks/AbilityTask_PlayAnimAndWait.h
--- a/Engine/Plugins/Runtime/GameplayAbilities/Source/GameplayAbilities/Public/Abilities/Tasks/AbilityTask_PlayAnimAndWait.h
+++ b/Engine/Plugins/Runtime/GameplayAbilities/Source/GameplayAbilities/Public/Abilities/Tasks/AbilityTask_PlayAnimAndWait.h
@@ -7,12 +7,14 @@
 #include "Abilities/Tasks/AbilityTask.h"
 #include "AbilityTask_PlayAnimAndWait.generated.h"
 
+#define UE_API GAMEPLAYABILITIES_API
+
 DECLARE_DYNAMIC_MULTICAST_DELEGATE(FPlayAnimWaitSimpleDelegate);
 
 class UAnimSequence;
 
-UCLASS()
-class GAMEPLAYABILITIES_API UAbilityTask_PlayAnimAndWait : public UAbilityTask
+UCLASS(MinimalAPI)
+class UAbilityTask_PlayAnimAndWait : public UAbilityTask
 {
 	GENERATED_BODY()
 
@@ -33,16 +35,16 @@ public:
 	FPlayAnimWaitSimpleDelegate	OnCancelled;
 
 	UFUNCTION()
-	void OnMontageBlendedIn(UAnimMontage* Montage);
+	UE_API void OnMontageBlendedIn(UAnimMontage* Montage);
 
 	UFUNCTION()
-	void OnMontageBlendingOut(UAnimMontage* Montage, bool bInterrupted);
+	UE_API void OnMontageBlendingOut(UAnimMontage* Montage, bool bInterrupted);
 
 	UFUNCTION()
-	void OnMontageInterrupted();
+	UE_API void OnMontageInterrupted();
 
 	UFUNCTION()
-	void OnMontageEnded(UAnimMontage* Montage, bool bInterrupted);
+	UE_API void OnMontageEnded(UAnimMontage* Montage, bool bInterrupted);
 
 	/** 
 	 * Start playing an animation montage on the avatar actor and wait for it to finish
@@ -60,22 +62,22 @@ public:
 	 */
 	UFUNCTION(BlueprintCallable, Category="Ability|Tasks", meta = (DisplayName="Play Anim And Wait",
 		HidePin = "OwningAbility", DefaultToSelf = "OwningAbility", BlueprintInternalUseOnly = "TRUE"))
-	static UAbilityTask_PlayAnimAndWait* CreatePlayAnimAndWaitProxy(UGameplayAbility* OwningAbility,
+	static UE_API UAbilityTask_PlayAnimAndWait* CreatePlayAnimAndWaitProxy(UGameplayAbility* OwningAbility,
 		FName TaskInstanceName, UAnimSequence* AnimSequence, FName SlotName = NAME_None, float BlendInTime = 0.25f, float BlendOutTime = 0.25f, float InPlayRate = 1.0f, float StartTimeSeconds = 0.0f, bool bStopWhenAbilityEnds = true, float AnimRootMotionTranslationScale = 1.0f, float InPlayCount = 1.0f);
 
-	virtual void Activate() override;
+	UE_API virtual void Activate() override;
 
 	/** Called when the ability is asked to cancel from an outside node. What this means depends on the individual task. By default, this does nothing other than ending the task. */
-	virtual void ExternalCancel() override;
+	UE_API virtual void ExternalCancel() override;
 
-	virtual FString GetDebugString() const override;
+	UE_API virtual FString GetDebugString() const override;
 
 protected:
 
-	virtual void OnDestroy(bool AbilityEnded) override;
+	UE_API virtual void OnDestroy(bool AbilityEnded) override;
 
 	/** Checks if the ability is playing a montage and stops that montage, returns true if a montage was stopped, false if not. */
-	bool StopPlayingMontage();
+	UE_API bool StopPlayingMontage();
 
 	FOnMontageBlendingOutStarted BlendingOutDelegate;
 	FOnMontageBlendedInEnded BlendedInDelegate;
@@ -94,3 +96,5 @@ protected:
 	float StartTimeSeconds = 0.0f;
 	bool bStopWhenAbilityEnds = true;
 };
+
+#undef UE_API
diff --git a/Engine/Plugins/Runtime/GameplayAbilities/Source/GameplayAbilities/Public/Abilities/Tasks/AbilityTask_PlayMontageAndWait.h b/Engine/Plugins/Runtime/GameplayAbilities/Source/GameplayAbilities/Public/Abilities/Tasks/AbilityTask_PlayMontageAndWait.h
--- a/Engine/Plugins/Runtime/GameplayAbilities/Source/GameplayAbilities/Public/Abilities/Tasks/AbilityTask_PlayMontageAndWait.h
+++ b/Engine/Plugins/Runtime/GameplayAbilities/Source/GameplayAbilities/Public/Abilities/Tasks/AbilityTask_PlayMontageAndWait.h
@@ -7,11 +7,13 @@
 #include "Abilities/Tasks/AbilityTask.h"
 #include "AbilityTask_PlayMontageAndWait.generated.h"
 
+#define UE_API GAMEPLAYABILITIES_API
+
 DECLARE_DYNAMIC_MULTICAST_DELEGATE(FMontageWaitSimpleDelegate);
 
 /** Ability task to simply play a montage. Many games will want to make a modified version of this task that looks for game-specific events */
-UCLASS()
-class GAMEPLAYABILITIES_API UAbilityTask_PlayMontageAndWait : public UAbilityTask
+UCLASS(MinimalAPI)
+class UAbilityTask_PlayMontageAndWait : public UAbilityTask
 {
 	GENERATED_BODY()
 public:
@@ -32,21 +34,21 @@ public:
 	FMontageWaitSimpleDelegate	OnCancelled;
 
 	UFUNCTION()
-	void OnMontageBlendedIn(UAnimMontage* Montage);
+	UE_API void OnMontageBlendedIn(UAnimMontage* Montage);
 
 	UFUNCTION()
-	void OnMontageBlendingOut(UAnimMontage* Montage, bool bInterrupted);
+	UE_API void OnMontageBlendingOut(UAnimMontage* Montage, bool bInterrupted);
 
 	UE_DEPRECATED(5.3, "Please use OnGameplayAbilityCancelled instead. This function naming implied the Montage was already interrupted (instead, we are about to interrupt it).")
 	UFUNCTION()
-	void OnMontageInterrupted();
+	UE_API void OnMontageInterrupted();
 
 	/** Callback function for when the owning Gameplay Ability is cancelled */
 	UFUNCTION()
-	void OnGameplayAbilityCancelled();
+	UE_API void OnGameplayAbilityCancelled();
 
 	UFUNCTION()
-	void OnMontageEnded(UAnimMontage* Montage, bool bInterrupted);
+	UE_API void OnMontageEnded(UAnimMontage* Montage, bool bInterrupted);
 
 	/** 
 	 * Start playing an animation montage on the avatar actor and wait for it to finish
@@ -65,22 +67,22 @@ public:
 	 */
 	UFUNCTION(BlueprintCallable, Category="Ability|Tasks", meta = (DisplayName="PlayMontageAndWait",
 		HidePin = "OwningAbility", DefaultToSelf = "OwningAbility", BlueprintInternalUseOnly = "TRUE"))
-	static UAbilityTask_PlayMontageAndWait* CreatePlayMontageAndWaitProxy(UGameplayAbility* OwningAbility,
+	static UE_API UAbilityTask_PlayMontageAndWait* CreatePlayMontageAndWaitProxy(UGameplayAbility* OwningAbility,
 		FName TaskInstanceName, UAnimMontage* MontageToPlay, float Rate = 1.f, FName StartSection = NAME_None, bool bStopWhenAbilityEnds = true, float AnimRootMotionTranslationScale = 1.f, float StartTimeSeconds = 0.f, bool bAllowInterruptAfterBlendOut = false);
 
-	virtual void Activate() override;
+	UE_API virtual void Activate() override;
 
 	/** Called when the ability is asked to cancel from an outside node. What this means depends on the individual task. By default, this does nothing other than ending the task. */
-	virtual void ExternalCancel() override;
+	UE_API virtual void ExternalCancel() override;
 
-	virtual FString GetDebugString() const override;
+	UE_API virtual FString GetDebugString() const override;
 
 protected:
 
-	virtual void OnDestroy(bool AbilityEnded) override;
+	UE_API virtual void OnDestroy(bool AbilityEnded) override;
 
 	/** Checks if the ability is playing a montage and stops that montage, returns true if a montage was stopped, false if not. */
-	bool StopPlayingMontage();
+	UE_API bool StopPlayingMontage();
 
 	FOnMontageBlendedInEnded BlendedInDelegate;
 	FOnMontageBlendingOutStarted BlendingOutDelegate;
@@ -108,3 +110,5 @@ protected:
 	UPROPERTY()
 	bool bAllowInterruptAfterBlendOut;
 };
+
+#undef UE_API
diff --git a/Engine/Plugins/Runtime/GameplayAbilities/Source/GameplayAbilities/Public/Abilities/Tasks/AbilityTask_Repeat.h b/Engine/Plugins/Runtime/GameplayAbilities/Source/GameplayAbilities/Public/Abilities/Tasks/AbilityTask_Repeat.h
--- a/Engine/Plugins/Runtime/GameplayAbilities/Source/GameplayAbilities/Public/Abilities/Tasks/AbilityTask_Repeat.h
+++ b/Engine/Plugins/Runtime/GameplayAbilities/Source/GameplayAbilities/Public/Abilities/Tasks/AbilityTask_Repeat.h
@@ -7,13 +7,15 @@
 #include "Abilities/Tasks/AbilityTask.h"
 #include "AbilityTask_Repeat.generated.h"
 
+#define UE_API GAMEPLAYABILITIES_API
+
 DECLARE_DYNAMIC_MULTICAST_DELEGATE_OneParam(FRepeatedActionDelegate, int32, ActionNumber);
 
 /**
  *	Repeat a task a certain number of times at a given interval.
  */
-UCLASS()
-class GAMEPLAYABILITIES_API UAbilityTask_Repeat : public UAbilityTask
+UCLASS(MinimalAPI)
+class UAbilityTask_Repeat : public UAbilityTask
 {
 	GENERATED_UCLASS_BODY()
 
@@ -23,15 +25,15 @@ class GAMEPLAYABILITIES_API UAbilityTask_Repeat : public UAbilityTask
 	UPROPERTY(BlueprintAssignable)
 	FRepeatedActionDelegate		OnFinished;
 
-	virtual FString GetDebugString() const override;
+	UE_API virtual FString GetDebugString() const override;
 
-	void PerformAction();
+	UE_API void PerformAction();
 
 	/** Start a task that repeats an action or set of actions. */
 	UFUNCTION(BlueprintCallable, Category = "Ability|Tasks", meta = (HidePin = "OwningAbility", DefaultToSelf = "OwningAbility", BlueprintInternalUseOnly = "TRUE"))
-	static UAbilityTask_Repeat* RepeatAction(UGameplayAbility* OwningAbility, float TimeBetweenActions, int32 TotalActionCount);
+	static UE_API UAbilityTask_Repeat* RepeatAction(UGameplayAbility* OwningAbility, float TimeBetweenActions, int32 TotalActionCount);
 
-	virtual void Activate() override;
+	UE_API virtual void Activate() override;
 
 protected:
 	int32 ActionPerformancesDesired;
@@ -41,5 +43,7 @@ protected:
 	/** Handle for efficient management of PerformAction timer */
 	FTimerHandle TimerHandle_PerformAction;
 
-	virtual void OnDestroy(bool AbilityIsEnding) override;
+	UE_API virtual void OnDestroy(bool AbilityIsEnding) override;
 };
+
+#undef UE_API
diff --git a/Engine/Plugins/Runtime/GameplayAbilities/Source/GameplayAbilities/Public/Abilities/Tasks/AbilityTask_SpawnActor.h b/Engine/Plugins/Runtime/GameplayAbilities/Source/GameplayAbilities/Public/Abilities/Tasks/AbilityTask_SpawnActor.h
--- a/Engine/Plugins/Runtime/GameplayAbilities/Source/GameplayAbilities/Public/Abilities/Tasks/AbilityTask_SpawnActor.h
+++ b/Engine/Plugins/Runtime/GameplayAbilities/Source/GameplayAbilities/Public/Abilities/Tasks/AbilityTask_SpawnActor.h
@@ -9,6 +9,8 @@
 #include "Abilities/Tasks/AbilityTask.h"
 #include "AbilityTask_SpawnActor.generated.h"
 
+#define UE_API GAMEPLAYABILITIES_API
+
 DECLARE_DYNAMIC_MULTICAST_DELEGATE_OneParam(FSpawnActorDelegate, AActor*, SpawnedActor);
 
 /**
@@ -23,8 +25,8 @@ DECLARE_DYNAMIC_MULTICAST_DELEGATE_OneParam(FSpawnActorDelegate, AActor*, Spawne
  *	actor down to it. We could potentially also use this to do predictive actor spawning / reconciliation.
  */
 
-UCLASS()
-class GAMEPLAYABILITIES_API UAbilityTask_SpawnActor: public UAbilityTask
+UCLASS(MinimalAPI)
+class UAbilityTask_SpawnActor: public UAbilityTask
 {
 	GENERATED_UCLASS_BODY()
 
@@ -37,14 +39,16 @@ class GAMEPLAYABILITIES_API UAbilityTask_SpawnActor: public UAbilityTask
 	
 	/** Spawn new Actor on the network authority (server) */
 	UFUNCTION(BlueprintCallable, meta=(HidePin = "OwningAbility", DefaultToSelf = "OwningAbility", BlueprintInternalUseOnly = "true"), Category="Ability|Tasks")
-	static UAbilityTask_SpawnActor* SpawnActor(UGameplayAbility* OwningAbility, FGameplayAbilityTargetDataHandle TargetData, TSubclassOf<AActor> Class);
+	static UE_API UAbilityTask_SpawnActor* SpawnActor(UGameplayAbility* OwningAbility, FGameplayAbilityTargetDataHandle TargetData, TSubclassOf<AActor> Class);
 
 	UFUNCTION(BlueprintCallable, meta = (HidePin = "OwningAbility", DefaultToSelf = "OwningAbility", BlueprintInternalUseOnly = "true"), Category = "Abilities")
-	bool BeginSpawningActor(UGameplayAbility* OwningAbility, FGameplayAbilityTargetDataHandle TargetData, TSubclassOf<AActor> Class, AActor*& SpawnedActor);
+	UE_API bool BeginSpawningActor(UGameplayAbility* OwningAbility, FGameplayAbilityTargetDataHandle TargetData, TSubclassOf<AActor> Class, AActor*& SpawnedActor);
 
 	UFUNCTION(BlueprintCallable, meta = (HidePin = "OwningAbility", DefaultToSelf = "OwningAbility", BlueprintInternalUseOnly = "true"), Category = "Abilities")
-	void FinishSpawningActor(UGameplayAbility* OwningAbility, FGameplayAbilityTargetDataHandle TargetData, AActor* SpawnedActor);
+	UE_API void FinishSpawningActor(UGameplayAbility* OwningAbility, FGameplayAbilityTargetDataHandle TargetData, AActor* SpawnedActor);
 
 protected:
 	FGameplayAbilityTargetDataHandle CachedTargetDataHandle;
 };
+
+#undef UE_API
diff --git a/Engine/Plugins/Runtime/GameplayAbilities/Source/GameplayAbilities/Public/Abilities/Tasks/AbilityTask_StartAbilityState.h b/Engine/Plugins/Runtime/GameplayAbilities/Source/GameplayAbilities/Public/Abilities/Tasks/AbilityTask_StartAbilityState.h
--- a/Engine/Plugins/Runtime/GameplayAbilities/Source/GameplayAbilities/Public/Abilities/Tasks/AbilityTask_StartAbilityState.h
+++ b/Engine/Plugins/Runtime/GameplayAbilities/Source/GameplayAbilities/Public/Abilities/Tasks/AbilityTask_StartAbilityState.h
@@ -7,6 +7,8 @@
 #include "Abilities/Tasks/AbilityTask.h"
 #include "AbilityTask_StartAbilityState.generated.h"
 
+#define UE_API GAMEPLAYABILITIES_API
+
 DECLARE_DYNAMIC_MULTICAST_DELEGATE(FAbilityStateDelegate);
 
 /**
@@ -23,8 +25,8 @@ DECLARE_DYNAMIC_MULTICAST_DELEGATE(FAbilityStateDelegate);
  * 'OnStateInterrupted' will be called if:
  * - The ability itself is cancelled via AGameplayAbility::CancelAbility
  */
-UCLASS()
-class GAMEPLAYABILITIES_API UAbilityTask_StartAbilityState : public UAbilityTask
+UCLASS(MinimalAPI)
+class UAbilityTask_StartAbilityState : public UAbilityTask
 {
 	GENERATED_UCLASS_BODY()
 
@@ -36,11 +38,11 @@ class GAMEPLAYABILITIES_API UAbilityTask_StartAbilityState : public UAbilityTask
 	UPROPERTY(BlueprintAssignable)
 	FAbilityStateDelegate OnStateInterrupted;
 
-	virtual void Activate() override;
+	UE_API virtual void Activate() override;
 
-	virtual void ExternalCancel() override;
+	UE_API virtual void ExternalCancel() override;
 
-	virtual FString GetDebugString() const override;
+	UE_API virtual FString GetDebugString() const override;
 
 	/**
 	 * Starts a new ability state.
@@ -49,7 +51,7 @@ class GAMEPLAYABILITIES_API UAbilityTask_StartAbilityState : public UAbilityTask
 	 * @param bEndCurrentState If true, all other active ability states will be ended.
 	 */
 	UFUNCTION(BlueprintCallable, Meta = (HidePin = "OwningAbility", DefaultToSelf = "OwningAbility", BlueprintInternalUseOnly = "true", HideSpawnParms = "Instigator"), Category = "Ability|Tasks")
-	static UAbilityTask_StartAbilityState* StartAbilityState(UGameplayAbility* OwningAbility, FName StateName, bool bEndCurrentState = true);
+	static UE_API UAbilityTask_StartAbilityState* StartAbilityState(UGameplayAbility* OwningAbility, FName StateName, bool bEndCurrentState = true);
 
 private:
 
@@ -61,8 +63,10 @@ private:
 
 	bool bEndCurrentState;
 
-	virtual void OnDestroy(bool AbilityEnded) override;
+	UE_API virtual void OnDestroy(bool AbilityEnded) override;
 
 	void OnEndState(FName StateNameToEnd);
 	void OnInterruptState();
 };
+
+#undef UE_API
diff --git a/Engine/Plugins/Runtime/GameplayAbilities/Source/GameplayAbilities/Public/Abilities/Tasks/AbilityTask_VisualizeTargeting.h b/Engine/Plugins/Runtime/GameplayAbilities/Source/GameplayAbilities/Public/Abilities/Tasks/AbilityTask_VisualizeTargeting.h
--- a/Engine/Plugins/Runtime/GameplayAbilities/Source/GameplayAbilities/Public/Abilities/Tasks/AbilityTask_VisualizeTargeting.h
+++ b/Engine/Plugins/Runtime/GameplayAbilities/Source/GameplayAbilities/Public/Abilities/Tasks/AbilityTask_VisualizeTargeting.h
@@ -9,43 +9,45 @@
 #include "Abilities/Tasks/AbilityTask.h"
 #include "AbilityTask_VisualizeTargeting.generated.h"
 
+#define UE_API GAMEPLAYABILITIES_API
+
 DECLARE_DYNAMIC_MULTICAST_DELEGATE(FVisualizeTargetingDelegate);
 
-UCLASS()
-class GAMEPLAYABILITIES_API UAbilityTask_VisualizeTargeting: public UAbilityTask
+UCLASS(MinimalAPI)
+class UAbilityTask_VisualizeTargeting: public UAbilityTask
 {
 	GENERATED_UCLASS_BODY()
 
 	UPROPERTY(BlueprintAssignable)
 	FVisualizeTargetingDelegate TimeElapsed;
 
-	void OnTimeElapsed();
+	UE_API void OnTimeElapsed();
 
 	/** Spawns target actor and uses it for visualization. */
 	UFUNCTION(BlueprintCallable, meta=(HidePin = "OwningAbility", DefaultToSelf = "OwningAbility", BlueprintInternalUseOnly = "true", HideSpawnParms="Instigator"), Category="Ability|Tasks")
-	static UAbilityTask_VisualizeTargeting* VisualizeTargeting(UGameplayAbility* OwningAbility, TSubclassOf<AGameplayAbilityTargetActor> Class, FName TaskInstanceName, float Duration = -1.0f);
+	static UE_API UAbilityTask_VisualizeTargeting* VisualizeTargeting(UGameplayAbility* OwningAbility, TSubclassOf<AGameplayAbilityTargetActor> Class, FName TaskInstanceName, float Duration = -1.0f);
 
 	/** Visualize target using a specified target actor. */
 	UFUNCTION(BlueprintCallable, meta = (HidePin = "OwningAbility", DefaultToSelf = "OwningAbility", BlueprintInternalUseOnly = "true"), Category = "Ability|Tasks")
-	static UAbilityTask_VisualizeTargeting* VisualizeTargetingUsingActor(UGameplayAbility* OwningAbility, AGameplayAbilityTargetActor* TargetActor, FName TaskInstanceName, float Duration = -1.0f);
+	static UE_API UAbilityTask_VisualizeTargeting* VisualizeTargetingUsingActor(UGameplayAbility* OwningAbility, AGameplayAbilityTargetActor* TargetActor, FName TaskInstanceName, float Duration = -1.0f);
 
-	virtual void Activate() override;
+	UE_API virtual void Activate() override;
 
 	UFUNCTION(BlueprintCallable, meta = (HidePin = "OwningAbility", DefaultToSelf = "OwningAbility", BlueprintInternalUseOnly = "true"), Category = "Abilities")
-	bool BeginSpawningActor(UGameplayAbility* OwningAbility, TSubclassOf<AGameplayAbilityTargetActor> Class, AGameplayAbilityTargetActor*& SpawnedActor);
+	UE_API bool BeginSpawningActor(UGameplayAbility* OwningAbility, TSubclassOf<AGameplayAbilityTargetActor> Class, AGameplayAbilityTargetActor*& SpawnedActor);
 
 	UFUNCTION(BlueprintCallable, meta = (HidePin = "OwningAbility", DefaultToSelf = "OwningAbility", BlueprintInternalUseOnly = "true"), Category = "Abilities")
-	void FinishSpawningActor(UGameplayAbility* OwningAbility, AGameplayAbilityTargetActor* SpawnedActor);
+	UE_API void FinishSpawningActor(UGameplayAbility* OwningAbility, AGameplayAbilityTargetActor* SpawnedActor);
 
 protected:
 
-	void SetDuration(const float Duration);
+	UE_API void SetDuration(const float Duration);
 
-	bool ShouldSpawnTargetActor() const;
-	void InitializeTargetActor(AGameplayAbilityTargetActor* SpawnedActor) const;
-	void FinalizeTargetActor(AGameplayAbilityTargetActor* SpawnedActor) const;
+	UE_API bool ShouldSpawnTargetActor() const;
+	UE_API void InitializeTargetActor(AGameplayAbilityTargetActor* SpawnedActor) const;
+	UE_API void FinalizeTargetActor(AGameplayAbilityTargetActor* SpawnedActor) const;
 
-	virtual void OnDestroy(bool AbilityEnded) override;
+	UE_API virtual void OnDestroy(bool AbilityEnded) override;
 
 protected:
 
@@ -67,3 +69,5 @@ protected:
 *		-Have a function named FinishSpawningActor w/ an AActor* of the class you spawned
 *			-This function *must* call ExecuteConstruction + PostActorConstruction
 */
+
+#undef UE_API
diff --git a/Engine/Plugins/Runtime/GameplayAbilities/Source/GameplayAbilities/Public/Abilities/Tasks/AbilityTask_WaitAbilityActivate.h b/Engine/Plugins/Runtime/GameplayAbilities/Source/GameplayAbilities/Public/Abilities/Tasks/AbilityTask_WaitAbilityActivate.h
--- a/Engine/Plugins/Runtime/GameplayAbilities/Source/GameplayAbilities/Public/Abilities/Tasks/AbilityTask_WaitAbilityActivate.h
+++ b/Engine/Plugins/Runtime/GameplayAbilities/Source/GameplayAbilities/Public/Abilities/Tasks/AbilityTask_WaitAbilityActivate.h
@@ -8,6 +8,8 @@
 #include "Abilities/Tasks/AbilityTask.h"
 #include "AbilityTask_WaitAbilityActivate.generated.h"
 
+#define UE_API GAMEPLAYABILITIES_API
+
 DECLARE_DYNAMIC_MULTICAST_DELEGATE_OneParam(FWaitAbilityActivateDelegate, UGameplayAbility*, ActivatedAbility);
 
 class AActor;
@@ -15,30 +17,30 @@ class AActor;
 /**
  *	Waits for the actor to activate another ability
  */
-UCLASS()
-class GAMEPLAYABILITIES_API UAbilityTask_WaitAbilityActivate : public UAbilityTask
+UCLASS(MinimalAPI)
+class UAbilityTask_WaitAbilityActivate : public UAbilityTask
 {
 	GENERATED_UCLASS_BODY()
 
 	UPROPERTY(BlueprintAssignable)
 	FWaitAbilityActivateDelegate	OnActivate;
 
-	virtual void Activate() override;
+	UE_API virtual void Activate() override;
 
 	UFUNCTION()
-	void OnAbilityActivate(UGameplayAbility *ActivatedAbility);
+	UE_API void OnAbilityActivate(UGameplayAbility *ActivatedAbility);
 
 	/** Wait until a new ability (of the same or different type) is activated. Only input based abilities will be counted unless IncludeTriggeredAbilities is true. */
 	UFUNCTION(BlueprintCallable, Category="Ability|Tasks", meta = (HidePin = "OwningAbility", DefaultToSelf = "OwningAbility", BlueprintInternalUseOnly = "TRUE"))
-	static UAbilityTask_WaitAbilityActivate* WaitForAbilityActivate(UGameplayAbility* OwningAbility, FGameplayTag WithTag, FGameplayTag WithoutTag, bool IncludeTriggeredAbilities=false, bool TriggerOnce=true);
+	static UE_API UAbilityTask_WaitAbilityActivate* WaitForAbilityActivate(UGameplayAbility* OwningAbility, FGameplayTag WithTag, FGameplayTag WithoutTag, bool IncludeTriggeredAbilities=false, bool TriggerOnce=true);
 
 	/** Wait until a new ability (of the same or different type) is activated. Only input based abilities will be counted unless IncludeTriggeredAbilities is true. Uses a tag requirements structure to filter abilities. */
 	UFUNCTION(BlueprintCallable, Category = "Ability|Tasks", meta = (HidePin = "OwningAbility", DefaultToSelf = "OwningAbility", BlueprintInternalUseOnly = "TRUE"))
-	static UAbilityTask_WaitAbilityActivate* WaitForAbilityActivateWithTagRequirements(UGameplayAbility* OwningAbility, FGameplayTagRequirements TagRequirements, bool IncludeTriggeredAbilities = false, bool TriggerOnce = true);
+	static UE_API UAbilityTask_WaitAbilityActivate* WaitForAbilityActivateWithTagRequirements(UGameplayAbility* OwningAbility, FGameplayTagRequirements TagRequirements, bool IncludeTriggeredAbilities = false, bool TriggerOnce = true);
 
 	/** Wait until a new ability (of the same or different type) is activated. Only input based abilities will be counted unless IncludeTriggeredAbilities is true. */
 	UFUNCTION(BlueprintCallable, Category="Ability|Tasks", meta = (HidePin = "OwningAbility", DefaultToSelf = "OwningAbility", BlueprintInternalUseOnly = "TRUE"))
-	static UAbilityTask_WaitAbilityActivate* WaitForAbilityActivate_Query(UGameplayAbility* OwningAbility, FGameplayTagQuery Query, bool IncludeTriggeredAbilities=false, bool TriggerOnce=true);
+	static UE_API UAbilityTask_WaitAbilityActivate* WaitForAbilityActivate_Query(UGameplayAbility* OwningAbility, FGameplayTagQuery Query, bool IncludeTriggeredAbilities=false, bool TriggerOnce=true);
 
 	FGameplayTag WithTag;
 	FGameplayTag WithoutTag;
@@ -50,7 +52,9 @@ class GAMEPLAYABILITIES_API UAbilityTask_WaitAbilityActivate : public UAbilityTa
 
 protected:
 
-	virtual void OnDestroy(bool AbilityEnded) override;
+	UE_API virtual void OnDestroy(bool AbilityEnded) override;
 
 	FDelegateHandle OnAbilityActivateDelegateHandle;
 };
+
+#undef UE_API
diff --git a/Engine/Plugins/Runtime/GameplayAbilities/Source/GameplayAbilities/Public/Abilities/Tasks/AbilityTask_WaitAbilityCommit.h b/Engine/Plugins/Runtime/GameplayAbilities/Source/GameplayAbilities/Public/Abilities/Tasks/AbilityTask_WaitAbilityCommit.h
--- a/Engine/Plugins/Runtime/GameplayAbilities/Source/GameplayAbilities/Public/Abilities/Tasks/AbilityTask_WaitAbilityCommit.h
+++ b/Engine/Plugins/Runtime/GameplayAbilities/Source/GameplayAbilities/Public/Abilities/Tasks/AbilityTask_WaitAbilityCommit.h
@@ -7,30 +7,32 @@
 #include "Abilities/Tasks/AbilityTask.h"
 #include "AbilityTask_WaitAbilityCommit.generated.h"
 
+#define UE_API GAMEPLAYABILITIES_API
+
 DECLARE_DYNAMIC_MULTICAST_DELEGATE_OneParam(FWaitAbilityCommitDelegate, UGameplayAbility*, ActivatedAbility);
 
 /**
  *	Waits for the actor to activate another ability
  */
-UCLASS()
-class GAMEPLAYABILITIES_API UAbilityTask_WaitAbilityCommit : public UAbilityTask
+UCLASS(MinimalAPI)
+class UAbilityTask_WaitAbilityCommit : public UAbilityTask
 {
 	GENERATED_UCLASS_BODY()
 
 	UPROPERTY(BlueprintAssignable)
 	FWaitAbilityCommitDelegate	OnCommit;
 
-	virtual void Activate() override;
+	UE_API virtual void Activate() override;
 
 	UFUNCTION()
-	void OnAbilityCommit(UGameplayAbility *ActivatedAbility);
+	UE_API void OnAbilityCommit(UGameplayAbility *ActivatedAbility);
 
 	/** Wait until a new ability (of the same or different type) is commited. */
 	UFUNCTION(BlueprintCallable, Category="Ability|Tasks", meta = (HidePin = "OwningAbility", DefaultToSelf = "OwningAbility", BlueprintInternalUseOnly = "TRUE", DisplayName = "Wait For New Ability Commit"))
-	static UAbilityTask_WaitAbilityCommit* WaitForAbilityCommit(UGameplayAbility* OwningAbility, FGameplayTag WithTag, FGameplayTag WithoutTage, bool TriggerOnce=true);
+	static UE_API UAbilityTask_WaitAbilityCommit* WaitForAbilityCommit(UGameplayAbility* OwningAbility, FGameplayTag WithTag, FGameplayTag WithoutTage, bool TriggerOnce=true);
 
 	UFUNCTION(BlueprintCallable, Category="Ability|Tasks", meta = (HidePin = "OwningAbility", DefaultToSelf = "OwningAbility", BlueprintInternalUseOnly = "TRUE", DisplayName = "Wait For New Ability Commit Query"))
-	static UAbilityTask_WaitAbilityCommit* WaitForAbilityCommit_Query(UGameplayAbility* OwningAbility, FGameplayTagQuery Query, bool TriggerOnce=true);
+	static UE_API UAbilityTask_WaitAbilityCommit* WaitForAbilityCommit_Query(UGameplayAbility* OwningAbility, FGameplayTagQuery Query, bool TriggerOnce=true);
 
 	FGameplayTag WithTag;
 	FGameplayTag WithoutTag;
@@ -41,7 +43,9 @@ class GAMEPLAYABILITIES_API UAbilityTask_WaitAbilityCommit : public UAbilityTask
 
 protected:
 
-	virtual void OnDestroy(bool AbilityEnded) override;
+	UE_API virtual void OnDestroy(bool AbilityEnded) override;
 
 	FDelegateHandle OnAbilityCommitDelegateHandle;
 };
+
+#undef UE_API
diff --git a/Engine/Plugins/Runtime/GameplayAbilities/Source/GameplayAbilities/Public/Abilities/Tasks/AbilityTask_WaitAttributeChange.h b/Engine/Plugins/Runtime/GameplayAbilities/Source/GameplayAbilities/Public/Abilities/Tasks/AbilityTask_WaitAttributeChange.h
--- a/Engine/Plugins/Runtime/GameplayAbilities/Source/GameplayAbilities/Public/Abilities/Tasks/AbilityTask_WaitAttributeChange.h
+++ b/Engine/Plugins/Runtime/GameplayAbilities/Source/GameplayAbilities/Public/Abilities/Tasks/AbilityTask_WaitAttributeChange.h
@@ -8,6 +8,8 @@
 #include "Abilities/Tasks/AbilityTask.h"
 #include "AbilityTask_WaitAttributeChange.generated.h"
 
+#define UE_API GAMEPLAYABILITIES_API
+
 struct FGameplayEffectModCallbackData;
 
 UENUM()
@@ -31,25 +33,25 @@ DECLARE_DYNAMIC_MULTICAST_DELEGATE(FWaitAttributeChangeDelegate);
 /**
  *	Waits for the actor to activate another ability
  */
-UCLASS()
-class GAMEPLAYABILITIES_API UAbilityTask_WaitAttributeChange : public UAbilityTask
+UCLASS(MinimalAPI)
+class UAbilityTask_WaitAttributeChange : public UAbilityTask
 {
 	GENERATED_UCLASS_BODY()
 
 	UPROPERTY(BlueprintAssignable)
 	FWaitAttributeChangeDelegate	OnChange;
 
-	virtual void Activate() override;
+	UE_API virtual void Activate() override;
 
-	void OnAttributeChange(const FOnAttributeChangeData& CallbackData);
+	UE_API void OnAttributeChange(const FOnAttributeChangeData& CallbackData);
 
 	/** Wait until an attribute changes. */
 	UFUNCTION(BlueprintCallable, Category="Ability|Tasks", meta = (HidePin = "OwningAbility", DefaultToSelf = "OwningAbility", BlueprintInternalUseOnly = "TRUE"))
-	static UAbilityTask_WaitAttributeChange* WaitForAttributeChange(UGameplayAbility* OwningAbility, FGameplayAttribute Attribute, FGameplayTag WithSrcTag, FGameplayTag WithoutSrcTag, bool TriggerOnce=true, AActor* OptionalExternalOwner = nullptr);
+	static UE_API UAbilityTask_WaitAttributeChange* WaitForAttributeChange(UGameplayAbility* OwningAbility, FGameplayAttribute Attribute, FGameplayTag WithSrcTag, FGameplayTag WithoutSrcTag, bool TriggerOnce=true, AActor* OptionalExternalOwner = nullptr);
 
 	/** Wait until an attribute changes to pass a given test. */
 	UFUNCTION(BlueprintCallable, Category = "Ability|Tasks", meta = (HidePin = "OwningAbility", DefaultToSelf = "OwningAbility", BlueprintInternalUseOnly = "TRUE"))
-	static UAbilityTask_WaitAttributeChange* WaitForAttributeChangeWithComparison(UGameplayAbility* OwningAbility, FGameplayAttribute InAttribute, FGameplayTag InWithTag, FGameplayTag InWithoutTag, TEnumAsByte<EWaitAttributeChangeComparison::Type> InComparisonType, float InComparisonValue, bool TriggerOnce=true, AActor* OptionalExternalOwner = nullptr);
+	static UE_API UAbilityTask_WaitAttributeChange* WaitForAttributeChangeWithComparison(UGameplayAbility* OwningAbility, FGameplayAttribute InAttribute, FGameplayTag InWithTag, FGameplayTag InWithoutTag, TEnumAsByte<EWaitAttributeChangeComparison::Type> InComparisonType, float InComparisonValue, bool TriggerOnce=true, AActor* OptionalExternalOwner = nullptr);
 
 	FGameplayTag WithTag;
 	FGameplayTag WithoutTag;
@@ -64,7 +66,9 @@ protected:
 	UPROPERTY()
 	TObjectPtr<UAbilitySystemComponent> ExternalOwner;
 
-	UAbilitySystemComponent* GetFocusedASC();
+	UE_API UAbilitySystemComponent* GetFocusedASC();
 
-	virtual void OnDestroy(bool AbilityEnded) override;
+	UE_API virtual void OnDestroy(bool AbilityEnded) override;
 };
+
+#undef UE_API
diff --git a/Engine/Plugins/Runtime/GameplayAbilities/Source/GameplayAbilities/Public/Abilities/Tasks/AbilityTask_WaitAttributeChangeRatioThreshold.h b/Engine/Plugins/Runtime/GameplayAbilities/Source/GameplayAbilities/Public/Abilities/Tasks/AbilityTask_WaitAttributeChangeRatioThreshold.h
--- a/Engine/Plugins/Runtime/GameplayAbilities/Source/GameplayAbilities/Public/Abilities/Tasks/AbilityTask_WaitAttributeChangeRatioThreshold.h
+++ b/Engine/Plugins/Runtime/GameplayAbilities/Source/GameplayAbilities/Public/Abilities/Tasks/AbilityTask_WaitAttributeChangeRatioThreshold.h
@@ -9,6 +9,8 @@
 #include "Abilities/Tasks/AbilityTask_WaitAttributeChange.h"
 #include "AbilityTask_WaitAttributeChangeRatioThreshold.generated.h"
 
+#define UE_API GAMEPLAYABILITIES_API
+
 struct FGameplayEffectModCallbackData;
 
 DECLARE_DYNAMIC_MULTICAST_DELEGATE_TwoParams(FWaitAttributeChangeRatioThresholdDelegate, bool, bMatchesComparison, float, CurrentRatio);
@@ -16,22 +18,22 @@ DECLARE_DYNAMIC_MULTICAST_DELEGATE_TwoParams(FWaitAttributeChangeRatioThresholdD
 /**
  *	Waits for the ratio between two attributes to match a threshold
  */
-UCLASS()
-class GAMEPLAYABILITIES_API UAbilityTask_WaitAttributeChangeRatioThreshold : public UAbilityTask
+UCLASS(MinimalAPI)
+class UAbilityTask_WaitAttributeChangeRatioThreshold : public UAbilityTask
 {
 	GENERATED_UCLASS_BODY()
 
 	UPROPERTY(BlueprintAssignable)
 	FWaitAttributeChangeRatioThresholdDelegate OnChange;
 
-	virtual void Activate() override;
+	UE_API virtual void Activate() override;
 
-	void OnNumeratorAttributeChange(const FOnAttributeChangeData& CallbackData);
-	void OnDenominatorAttributeChange(const FOnAttributeChangeData& CallbackData);
+	UE_API void OnNumeratorAttributeChange(const FOnAttributeChangeData& CallbackData);
+	UE_API void OnDenominatorAttributeChange(const FOnAttributeChangeData& CallbackData);
 
 	/** Wait on attribute ratio change meeting a comparison threshold. */
 	UFUNCTION(BlueprintCallable, Category="Ability|Tasks", meta = (HidePin = "OwningAbility", DefaultToSelf = "OwningAbility", BlueprintInternalUseOnly = "TRUE"))
-	static UAbilityTask_WaitAttributeChangeRatioThreshold* WaitForAttributeChangeRatioThreshold(UGameplayAbility* OwningAbility, FGameplayAttribute AttributeNumerator, FGameplayAttribute AttributeDenominator, TEnumAsByte<EWaitAttributeChangeComparison::Type> ComparisonType, float ComparisonValue, bool bTriggerOnce, AActor* OptionalExternalOwner = nullptr);
+	static UE_API UAbilityTask_WaitAttributeChangeRatioThreshold* WaitForAttributeChangeRatioThreshold(UGameplayAbility* OwningAbility, FGameplayAttribute AttributeNumerator, FGameplayAttribute AttributeDenominator, TEnumAsByte<EWaitAttributeChangeComparison::Type> ComparisonType, float ComparisonValue, bool bTriggerOnce, AActor* OptionalExternalOwner = nullptr);
 
 	FGameplayAttribute AttributeNumerator;
 	FGameplayAttribute AttributeDenominator;
@@ -53,12 +55,14 @@ protected:
 	UPROPERTY()
 	TObjectPtr<UAbilitySystemComponent> ExternalOwner;
 
-	UAbilitySystemComponent* GetFocusedASC();
+	UE_API UAbilitySystemComponent* GetFocusedASC();
 
-	void OnAttributeChange();
-	void OnRatioChange();
+	UE_API void OnAttributeChange();
+	UE_API void OnRatioChange();
 
-	virtual void OnDestroy(bool AbilityEnded) override;
+	UE_API virtual void OnDestroy(bool AbilityEnded) override;
 
-	bool DoesValuePassComparison(float ValueNumerator, float ValueDenominator) const;
+	UE_API bool DoesValuePassComparison(float ValueNumerator, float ValueDenominator) const;
 };
+
+#undef UE_API
diff --git a/Engine/Plugins/Runtime/GameplayAbilities/Source/GameplayAbilities/Public/Abilities/Tasks/AbilityTask_WaitAttributeChangeThreshold.h b/Engine/Plugins/Runtime/GameplayAbilities/Source/GameplayAbilities/Public/Abilities/Tasks/AbilityTask_WaitAttributeChangeThreshold.h
--- a/Engine/Plugins/Runtime/GameplayAbilities/Source/GameplayAbilities/Public/Abilities/Tasks/AbilityTask_WaitAttributeChangeThreshold.h
+++ b/Engine/Plugins/Runtime/GameplayAbilities/Source/GameplayAbilities/Public/Abilities/Tasks/AbilityTask_WaitAttributeChangeThreshold.h
@@ -8,6 +8,8 @@
 #include "Abilities/Tasks/AbilityTask_WaitAttributeChange.h"
 #include "AbilityTask_WaitAttributeChangeThreshold.generated.h"
 
+#define UE_API GAMEPLAYABILITIES_API
+
 struct FGameplayEffectModCallbackData;
 
 DECLARE_DYNAMIC_MULTICAST_DELEGATE_TwoParams(FWaitAttributeChangeThresholdDelegate, bool, bMatchesComparison, float, CurrentValue);
@@ -15,21 +17,21 @@ DECLARE_DYNAMIC_MULTICAST_DELEGATE_TwoParams(FWaitAttributeChangeThresholdDelega
 /**
  *	Waits for an attribute to match a threshold
  */
-UCLASS()
-class GAMEPLAYABILITIES_API UAbilityTask_WaitAttributeChangeThreshold : public UAbilityTask
+UCLASS(MinimalAPI)
+class UAbilityTask_WaitAttributeChangeThreshold : public UAbilityTask
 {
 	GENERATED_UCLASS_BODY()
 
 	UPROPERTY(BlueprintAssignable)
 	FWaitAttributeChangeThresholdDelegate OnChange;
 
-	virtual void Activate() override;
+	UE_API virtual void Activate() override;
 
-	void OnAttributeChange(const FOnAttributeChangeData& CallbackData);
+	UE_API void OnAttributeChange(const FOnAttributeChangeData& CallbackData);
 
 	/** Wait on attribute change meeting a comparison threshold. */
 	UFUNCTION(BlueprintCallable, Category="Ability|Tasks", meta = (HidePin = "OwningAbility", DefaultToSelf = "OwningAbility", BlueprintInternalUseOnly = "TRUE"))
-	static UAbilityTask_WaitAttributeChangeThreshold* WaitForAttributeChangeThreshold(UGameplayAbility* OwningAbility, FGameplayAttribute Attribute, TEnumAsByte<EWaitAttributeChangeComparison::Type> ComparisonType, float ComparisonValue, bool bTriggerOnce, AActor* OptionalExternalOwner = nullptr);
+	static UE_API UAbilityTask_WaitAttributeChangeThreshold* WaitForAttributeChangeThreshold(UGameplayAbility* OwningAbility, FGameplayAttribute Attribute, TEnumAsByte<EWaitAttributeChangeComparison::Type> ComparisonType, float ComparisonValue, bool bTriggerOnce, AActor* OptionalExternalOwner = nullptr);
 
 	FGameplayAttribute Attribute;
 	TEnumAsByte<EWaitAttributeChangeComparison::Type> ComparisonType;
@@ -44,9 +46,11 @@ protected:
 	UPROPERTY()
 	TObjectPtr<UAbilitySystemComponent> ExternalOwner;
 
-	UAbilitySystemComponent* GetFocusedASC();
+	UE_API UAbilitySystemComponent* GetFocusedASC();
 
-	virtual void OnDestroy(bool AbilityEnded) override;
+	UE_API virtual void OnDestroy(bool AbilityEnded) override;
 
-	bool DoesValuePassComparison(float Value) const;
+	UE_API bool DoesValuePassComparison(float Value) const;
 };
+
+#undef UE_API
diff --git a/Engine/Plugins/Runtime/GameplayAbilities/Source/GameplayAbilities/Public/Abilities/Tasks/AbilityTask_WaitCancel.h b/Engine/Plugins/Runtime/GameplayAbilities/Source/GameplayAbilities/Public/Abilities/Tasks/AbilityTask_WaitCancel.h
--- a/Engine/Plugins/Runtime/GameplayAbilities/Source/GameplayAbilities/Public/Abilities/Tasks/AbilityTask_WaitCancel.h
+++ b/Engine/Plugins/Runtime/GameplayAbilities/Source/GameplayAbilities/Public/Abilities/Tasks/AbilityTask_WaitCancel.h
@@ -6,10 +6,12 @@
 #include "Abilities/Tasks/AbilityTask.h"
 #include "AbilityTask_WaitCancel.generated.h"
 
+#define UE_API GAMEPLAYABILITIES_API
+
 DECLARE_DYNAMIC_MULTICAST_DELEGATE(FWaitCancelDelegate);
 
-UCLASS()
-class GAMEPLAYABILITIES_API UAbilityTask_WaitCancel : public UAbilityTask
+UCLASS(MinimalAPI)
+class UAbilityTask_WaitCancel : public UAbilityTask
 {
 	GENERATED_UCLASS_BODY()
 
@@ -17,20 +19,22 @@ class GAMEPLAYABILITIES_API UAbilityTask_WaitCancel : public UAbilityTask
 	FWaitCancelDelegate	OnCancel;
 	
 	UFUNCTION()
-	void OnCancelCallback();
+	UE_API void OnCancelCallback();
 
 	UFUNCTION()
-	void OnLocalCancelCallback();
+	UE_API void OnLocalCancelCallback();
 
 	UFUNCTION(BlueprintCallable, meta=(HidePin = "OwningAbility", DefaultToSelf = "OwningAbility", BlueprintInternalUseOnly = "true", DisplayName="Wait for Cancel Input"), Category="Ability|Tasks")
-	static UAbilityTask_WaitCancel* WaitCancel(UGameplayAbility* OwningAbility);	
+	static UE_API UAbilityTask_WaitCancel* WaitCancel(UGameplayAbility* OwningAbility);	
 
-	virtual void Activate() override;
+	UE_API virtual void Activate() override;
 
 protected:
 
-	virtual void OnDestroy(bool AbilityEnding) override;
+	UE_API virtual void OnDestroy(bool AbilityEnding) override;
 
 	bool RegisteredCallbacks;
 	
 };
+
+#undef UE_API
diff --git a/Engine/Plugins/Runtime/GameplayAbilities/Source/GameplayAbilities/Public/Abilities/Tasks/AbilityTask_WaitConfirm.h b/Engine/Plugins/Runtime/GameplayAbilities/Source/GameplayAbilities/Public/Abilities/Tasks/AbilityTask_WaitConfirm.h
--- a/Engine/Plugins/Runtime/GameplayAbilities/Source/GameplayAbilities/Public/Abilities/Tasks/AbilityTask_WaitConfirm.h
+++ b/Engine/Plugins/Runtime/GameplayAbilities/Source/GameplayAbilities/Public/Abilities/Tasks/AbilityTask_WaitConfirm.h
@@ -6,8 +6,10 @@
 #include "Abilities/Tasks/AbilityTask.h"
 #include "AbilityTask_WaitConfirm.generated.h"
 
-UCLASS()
-class GAMEPLAYABILITIES_API UAbilityTask_WaitConfirm : public UAbilityTask
+#define UE_API GAMEPLAYABILITIES_API
+
+UCLASS(MinimalAPI)
+class UAbilityTask_WaitConfirm : public UAbilityTask
 {
 	GENERATED_UCLASS_BODY()
 
@@ -15,19 +17,21 @@ class GAMEPLAYABILITIES_API UAbilityTask_WaitConfirm : public UAbilityTask
 	FGenericGameplayTaskDelegate	OnConfirm;
 
 	UFUNCTION()
-	void OnConfirmCallback(UGameplayAbility* InAbility);
+	UE_API void OnConfirmCallback(UGameplayAbility* InAbility);
 
-	virtual void Activate() override;
+	UE_API virtual void Activate() override;
 
 	/** Wait until the server confirms the use of this ability. This is used to gate predictive portions of the ability */
 	UFUNCTION(BlueprintCallable, Category="Ability|Tasks", meta = (HidePin = "OwningAbility", DefaultToSelf = "OwningAbility", BlueprintInternalUseOnly = "TRUE"))
-	static UAbilityTask_WaitConfirm* WaitConfirm(UGameplayAbility* OwningAbility);
+	static UE_API UAbilityTask_WaitConfirm* WaitConfirm(UGameplayAbility* OwningAbility);
 
 protected:
 
-	virtual void OnDestroy(bool AbilityEnded) override;
+	UE_API virtual void OnDestroy(bool AbilityEnded) override;
 
 	bool RegisteredCallback;
 
 	FDelegateHandle OnConfirmCallbackDelegateHandle;
 };
+
+#undef UE_API
diff --git a/Engine/Plugins/Runtime/GameplayAbilities/Source/GameplayAbilities/Public/Abilities/Tasks/AbilityTask_WaitConfirmCancel.h b/Engine/Plugins/Runtime/GameplayAbilities/Source/GameplayAbilities/Public/Abilities/Tasks/AbilityTask_WaitConfirmCancel.h
--- a/Engine/Plugins/Runtime/GameplayAbilities/Source/GameplayAbilities/Public/Abilities/Tasks/AbilityTask_WaitConfirmCancel.h
+++ b/Engine/Plugins/Runtime/GameplayAbilities/Source/GameplayAbilities/Public/Abilities/Tasks/AbilityTask_WaitConfirmCancel.h
@@ -6,14 +6,16 @@
 #include "Abilities/Tasks/AbilityTask.h"
 #include "AbilityTask_WaitConfirmCancel.generated.h"
 
+#define UE_API GAMEPLAYABILITIES_API
+
 DECLARE_DYNAMIC_MULTICAST_DELEGATE(FWaitConfirmCancelDelegate);
 
 // Fixme: this name is conflicting with AbilityTask_WaitConfirm
 // UAbilityTask_WaitConfirmCancel = Wait for Targeting confirm/cancel
 // UAbilityTask_WaitConfirm = Wait for server to confirm ability activation
 
-UCLASS()
-class GAMEPLAYABILITIES_API UAbilityTask_WaitConfirmCancel : public UAbilityTask
+UCLASS(MinimalAPI)
+class UAbilityTask_WaitConfirmCancel : public UAbilityTask
 {
 	GENERATED_UCLASS_BODY()
 
@@ -24,26 +26,28 @@ class GAMEPLAYABILITIES_API UAbilityTask_WaitConfirmCancel : public UAbilityTask
 	FWaitConfirmCancelDelegate	OnCancel;
 	
 	UFUNCTION()
-	void OnConfirmCallback();
+	UE_API void OnConfirmCallback();
 
 	UFUNCTION()
-	void OnCancelCallback();
+	UE_API void OnCancelCallback();
 
 	UFUNCTION()
-	void OnLocalConfirmCallback();
+	UE_API void OnLocalConfirmCallback();
 
 	UFUNCTION()
-	void OnLocalCancelCallback();
+	UE_API void OnLocalCancelCallback();
 
 	UFUNCTION(BlueprintCallable, meta=(HidePin = "OwningAbility", DefaultToSelf = "OwningAbility", BlueprintInternalUseOnly = "true", DisplayName="Wait for Confirm Input"), Category="Ability|Tasks")
-	static UAbilityTask_WaitConfirmCancel* WaitConfirmCancel(UGameplayAbility* OwningAbility);	
+	static UE_API UAbilityTask_WaitConfirmCancel* WaitConfirmCancel(UGameplayAbility* OwningAbility);	
 
-	virtual void Activate() override;
+	UE_API virtual void Activate() override;
 
 protected:
 
-	virtual void OnDestroy(bool AbilityEnding) override;
+	UE_API virtual void OnDestroy(bool AbilityEnding) override;
 
 	bool RegisteredCallbacks;
 	
 };
+
+#undef UE_API
diff --git a/Engine/Plugins/Runtime/GameplayAbilities/Source/GameplayAbilities/Public/Abilities/Tasks/AbilityTask_WaitDelay.h b/Engine/Plugins/Runtime/GameplayAbilities/Source/GameplayAbilities/Public/Abilities/Tasks/AbilityTask_WaitDelay.h
--- a/Engine/Plugins/Runtime/GameplayAbilities/Source/GameplayAbilities/Public/Abilities/Tasks/AbilityTask_WaitDelay.h
+++ b/Engine/Plugins/Runtime/GameplayAbilities/Source/GameplayAbilities/Public/Abilities/Tasks/AbilityTask_WaitDelay.h
@@ -6,24 +6,26 @@
 #include "Abilities/Tasks/AbilityTask.h"
 #include "AbilityTask_WaitDelay.generated.h"
 
+#define UE_API GAMEPLAYABILITIES_API
+
 DECLARE_DYNAMIC_MULTICAST_DELEGATE(FWaitDelayDelegate);
 
-UCLASS()
-class GAMEPLAYABILITIES_API UAbilityTask_WaitDelay : public UAbilityTask
+UCLASS(MinimalAPI)
+class UAbilityTask_WaitDelay : public UAbilityTask
 {
 	GENERATED_UCLASS_BODY()
 
 	UPROPERTY(BlueprintAssignable)
 	FWaitDelayDelegate	OnFinish;
 
-	virtual void Activate() override;
+	UE_API virtual void Activate() override;
 
 	/** Return debug string describing task */
-	virtual FString GetDebugString() const override;
+	UE_API virtual FString GetDebugString() const override;
 
 	/** Wait specified time. This is functionally the same as a standard Delay node. */
 	UFUNCTION(BlueprintCallable, Category="Ability|Tasks", meta = (HidePin = "OwningAbility", DefaultToSelf = "OwningAbility", BlueprintInternalUseOnly = "TRUE"))
-	static UAbilityTask_WaitDelay* WaitDelay(UGameplayAbility* OwningAbility, float Time);
+	static UE_API UAbilityTask_WaitDelay* WaitDelay(UGameplayAbility* OwningAbility, float Time);
 
 private:
 
@@ -32,3 +34,5 @@ private:
 	float Time;
 	float TimeStarted;
 };
+
+#undef UE_API
diff --git a/Engine/Plugins/Runtime/GameplayAbilities/Source/GameplayAbilities/Public/Abilities/Tasks/AbilityTask_WaitGameplayEffectApplied.h b/Engine/Plugins/Runtime/GameplayAbilities/Source/GameplayAbilities/Public/Abilities/Tasks/AbilityTask_WaitGameplayEffectApplied.h
--- a/Engine/Plugins/Runtime/GameplayAbilities/Source/GameplayAbilities/Public/Abilities/Tasks/AbilityTask_WaitGameplayEffectApplied.h
+++ b/Engine/Plugins/Runtime/GameplayAbilities/Source/GameplayAbilities/Public/Abilities/Tasks/AbilityTask_WaitGameplayEffectApplied.h
@@ -8,17 +8,19 @@
 #include "Abilities/Tasks/AbilityTask.h"
 #include "AbilityTask_WaitGameplayEffectApplied.generated.h"
 
+#define UE_API GAMEPLAYABILITIES_API
+
 class UAbilitySystemComponent;
 
-UCLASS()
-class GAMEPLAYABILITIES_API UAbilityTask_WaitGameplayEffectApplied : public UAbilityTask
+UCLASS(MinimalAPI)
+class UAbilityTask_WaitGameplayEffectApplied : public UAbilityTask
 {
 	GENERATED_UCLASS_BODY()
 
 	UFUNCTION()
-	void OnApplyGameplayEffectCallback(UAbilitySystemComponent* Target, const FGameplayEffectSpec& SpecApplied, FActiveGameplayEffectHandle ActiveHandle);
+	UE_API void OnApplyGameplayEffectCallback(UAbilitySystemComponent* Target, const FGameplayEffectSpec& SpecApplied, FActiveGameplayEffectHandle ActiveHandle);
 
-	virtual void Activate() override;
+	UE_API virtual void Activate() override;
 
 	FGameplayTargetDataFilterHandle Filter;
 	FGameplayTagRequirements SourceTagRequirements;
@@ -30,17 +32,17 @@ class GAMEPLAYABILITIES_API UAbilityTask_WaitGameplayEffectApplied : public UAbi
 	bool TriggerOnce;
 	bool ListenForPeriodicEffects;
 
-	void SetExternalActor(AActor* InActor);
+	UE_API void SetExternalActor(AActor* InActor);
 
 protected:
 
-	UAbilitySystemComponent* GetASC();
+	UE_API UAbilitySystemComponent* GetASC();
 
 	virtual void BroadcastDelegate(AActor* Avatar, FGameplayEffectSpecHandle SpecHandle, FActiveGameplayEffectHandle ActiveHandle) { }
 	virtual void RegisterDelegate() { }
 	virtual void RemoveDelegate() { }
 
-	virtual void OnDestroy(bool AbilityEnded) override;
+	UE_API virtual void OnDestroy(bool AbilityEnded) override;
 
 	bool RegisteredCallback;
 	bool UseExternalOwner;
@@ -50,4 +52,6 @@ protected:
 
 	// If we are in the process of broadcasting and should not accept additional GE callbacks
 	bool Locked;
-};
\ No newline at end of file
+};
+
+#undef UE_API
diff --git a/Engine/Plugins/Runtime/GameplayAbilities/Source/GameplayAbilities/Public/Abilities/Tasks/AbilityTask_WaitGameplayEffectApplied_Self.h b/Engine/Plugins/Runtime/GameplayAbilities/Source/GameplayAbilities/Public/Abilities/Tasks/AbilityTask_WaitGameplayEffectApplied_Self.h
--- a/Engine/Plugins/Runtime/GameplayAbilities/Source/GameplayAbilities/Public/Abilities/Tasks/AbilityTask_WaitGameplayEffectApplied_Self.h
+++ b/Engine/Plugins/Runtime/GameplayAbilities/Source/GameplayAbilities/Public/Abilities/Tasks/AbilityTask_WaitGameplayEffectApplied_Self.h
@@ -8,10 +8,12 @@
 #include "Abilities/Tasks/AbilityTask_WaitGameplayEffectApplied.h"
 #include "AbilityTask_WaitGameplayEffectApplied_Self.generated.h"
 
+#define UE_API GAMEPLAYABILITIES_API
+
 DECLARE_DYNAMIC_MULTICAST_DELEGATE_ThreeParams(FGameplayEffectAppliedSelfDelegate, AActor*, Source, FGameplayEffectSpecHandle, SpecHandle,  FActiveGameplayEffectHandle, ActiveHandle );
 
-UCLASS()
-class GAMEPLAYABILITIES_API UAbilityTask_WaitGameplayEffectApplied_Self : public UAbilityTask_WaitGameplayEffectApplied
+UCLASS(MinimalAPI)
+class UAbilityTask_WaitGameplayEffectApplied_Self : public UAbilityTask_WaitGameplayEffectApplied
 {
 	GENERATED_UCLASS_BODY()
 
@@ -23,7 +25,7 @@ class GAMEPLAYABILITIES_API UAbilityTask_WaitGameplayEffectApplied_Self : public
 	 * Optional External Owner can be used to run this task on someone else (not the owner of the ability). By default you can leave this empty.
 	 */	 	 
 	UFUNCTION(BlueprintCallable, Category="Ability|Tasks", meta = (HidePin = "OwningAbility", DefaultToSelf = "OwningAbility", BlueprintInternalUseOnly = "TRUE"))
-	static UAbilityTask_WaitGameplayEffectApplied_Self* WaitGameplayEffectAppliedToSelf(UGameplayAbility* OwningAbility, const FGameplayTargetDataFilterHandle SourceFilter, FGameplayTagRequirements SourceTagRequirements, FGameplayTagRequirements TargetTagRequirements, bool TriggerOnce=false, AActor* OptionalExternalOwner=nullptr, bool ListenForPeriodicEffect=false );
+	static UE_API UAbilityTask_WaitGameplayEffectApplied_Self* WaitGameplayEffectAppliedToSelf(UGameplayAbility* OwningAbility, const FGameplayTargetDataFilterHandle SourceFilter, FGameplayTagRequirements SourceTagRequirements, FGameplayTagRequirements TargetTagRequirements, bool TriggerOnce=false, AActor* OptionalExternalOwner=nullptr, bool ListenForPeriodicEffect=false );
 
 	/**	 
 	 * Wait until the owner *receives* a GameplayEffect from a given source (the source may be the owner too!). If TriggerOnce is true, this task will only return one time. Otherwise it will return everytime a GE is applied that meets the requirements over the life of the ability
@@ -32,15 +34,17 @@ class GAMEPLAYABILITIES_API UAbilityTask_WaitGameplayEffectApplied_Self : public
 	 * This version uses FGameplayTagQuery (more power) instead of FGameplayTagRequirements (faster)
 	 */	
 	UFUNCTION(BlueprintCallable, Category="Ability|Tasks", meta = (HidePin = "OwningAbility", DefaultToSelf = "OwningAbility", BlueprintInternalUseOnly = "TRUE"))
-	static UAbilityTask_WaitGameplayEffectApplied_Self* WaitGameplayEffectAppliedToSelf_Query(UGameplayAbility* OwningAbility, const FGameplayTargetDataFilterHandle SourceFilter, FGameplayTagQuery SourceTagQuery, FGameplayTagQuery TargetTagQuery, bool TriggerOnce=false, AActor* OptionalExternalOwner=nullptr, bool ListenForPeriodicEffect=false );
+	static UE_API UAbilityTask_WaitGameplayEffectApplied_Self* WaitGameplayEffectAppliedToSelf_Query(UGameplayAbility* OwningAbility, const FGameplayTargetDataFilterHandle SourceFilter, FGameplayTagQuery SourceTagQuery, FGameplayTagQuery TargetTagQuery, bool TriggerOnce=false, AActor* OptionalExternalOwner=nullptr, bool ListenForPeriodicEffect=false );
 	
 protected:
 
-	virtual void BroadcastDelegate(AActor* Avatar, FGameplayEffectSpecHandle SpecHandle, FActiveGameplayEffectHandle ActiveHandle) override;
-	virtual void RegisterDelegate() override;
-	virtual void RemoveDelegate() override;
+	UE_API virtual void BroadcastDelegate(AActor* Avatar, FGameplayEffectSpecHandle SpecHandle, FActiveGameplayEffectHandle ActiveHandle) override;
+	UE_API virtual void RegisterDelegate() override;
+	UE_API virtual void RemoveDelegate() override;
 
 private:
 	FDelegateHandle OnApplyGameplayEffectCallbackDelegateHandle;
 	FDelegateHandle OnPeriodicGameplayEffectExecuteCallbackDelegateHandle;
 };
+
+#undef UE_API
