[GameplayAbilities]

diff --git a/Engine/Plugins/Runtime/GameplayAbilities/Source/GameplayAbilities/Public/Abilities/Tasks/AbilityTask_WaitGameplayEffectApplied_Target.h b/Engine/Plugins/Runtime/GameplayAbilities/Source/GameplayAbilities/Public/Abilities/Tasks/AbilityTask_WaitGameplayEffectApplied_Target.h
--- a/Engine/Plugins/Runtime/GameplayAbilities/Source/GameplayAbilities/Public/Abilities/Tasks/AbilityTask_WaitGameplayEffectApplied_Target.h
+++ b/Engine/Plugins/Runtime/GameplayAbilities/Source/GameplayAbilities/Public/Abilities/Tasks/AbilityTask_WaitGameplayEffectApplied_Target.h
@@ -8,10 +8,12 @@
 #include "Abilities/Tasks/AbilityTask_WaitGameplayEffectApplied.h"
 #include "AbilityTask_WaitGameplayEffectApplied_Target.generated.h"
 
+#define UE_API GAMEPLAYABILITIES_API
+
 DECLARE_DYNAMIC_MULTICAST_DELEGATE_ThreeParams(FGameplayEffectAppliedTargetDelegate, AActor*, Target, FGameplayEffectSpecHandle, SpecHandle,  FActiveGameplayEffectHandle, ActiveHandle );
 
-UCLASS()
-class GAMEPLAYABILITIES_API UAbilityTask_WaitGameplayEffectApplied_Target : public UAbilityTask_WaitGameplayEffectApplied
+UCLASS(MinimalAPI)
+class UAbilityTask_WaitGameplayEffectApplied_Target : public UAbilityTask_WaitGameplayEffectApplied
 {
 	GENERATED_UCLASS_BODY()
 
@@ -23,22 +25,24 @@ class GAMEPLAYABILITIES_API UAbilityTask_WaitGameplayEffectApplied_Target : publ
 	 * Optional External Owner can be used to run this task on someone else (not the owner of the ability). By default you can leave this empty.
 	 */
 	UFUNCTION(BlueprintCallable, Category="Ability|Tasks", meta = (HidePin = "OwningAbility", DefaultToSelf = "OwningAbility", BlueprintInternalUseOnly = "TRUE"))
-	static UAbilityTask_WaitGameplayEffectApplied_Target* WaitGameplayEffectAppliedToTarget(UGameplayAbility* OwningAbility, const FGameplayTargetDataFilterHandle TargetFilter, FGameplayTagRequirements SourceTagRequirements, FGameplayTagRequirements TargetTagRequirements, bool TriggerOnce=false, AActor* OptionalExternalOwner=nullptr, bool ListenForPeriodicEffects=false);
+	static UE_API UAbilityTask_WaitGameplayEffectApplied_Target* WaitGameplayEffectAppliedToTarget(UGameplayAbility* OwningAbility, const FGameplayTargetDataFilterHandle TargetFilter, FGameplayTagRequirements SourceTagRequirements, FGameplayTagRequirements TargetTagRequirements, bool TriggerOnce=false, AActor* OptionalExternalOwner=nullptr, bool ListenForPeriodicEffects=false);
 
 	/** 
 	 * Wait until the owner (or External Owner) applies a GameplayEffect to a Target (the target may be the owner too!). If TriggerOnce is true, this task will only return one time. Otherwise it will return everytime a GE is applied that meets the requirements over the life of the ability 
 	 * Optional External Owner can be used to run this task on someone else (not the owner of the ability). By default you can leave this empty.
 	 */
 	UFUNCTION(BlueprintCallable, Category="Ability|Tasks", meta = (HidePin = "OwningAbility", DefaultToSelf = "OwningAbility", BlueprintInternalUseOnly = "TRUE"))
-	static UAbilityTask_WaitGameplayEffectApplied_Target* WaitGameplayEffectAppliedToTarget_Query(UGameplayAbility* OwningAbility, const FGameplayTargetDataFilterHandle SourceFilter, FGameplayTagQuery SourceTagQuery, FGameplayTagQuery TargetTagQuery, bool TriggerOnce=false, AActor* OptionalExternalOwner=nullptr, bool ListenForPeriodicEffect=false );
+	static UE_API UAbilityTask_WaitGameplayEffectApplied_Target* WaitGameplayEffectAppliedToTarget_Query(UGameplayAbility* OwningAbility, const FGameplayTargetDataFilterHandle SourceFilter, FGameplayTagQuery SourceTagQuery, FGameplayTagQuery TargetTagQuery, bool TriggerOnce=false, AActor* OptionalExternalOwner=nullptr, bool ListenForPeriodicEffect=false );
 	
 protected:
 
-	virtual void BroadcastDelegate(AActor* Avatar, FGameplayEffectSpecHandle SpecHandle, FActiveGameplayEffectHandle ActiveHandle) override;
-	virtual void RegisterDelegate() override;
-	virtual void RemoveDelegate() override;
+	UE_API virtual void BroadcastDelegate(AActor* Avatar, FGameplayEffectSpecHandle SpecHandle, FActiveGameplayEffectHandle ActiveHandle) override;
+	UE_API virtual void RegisterDelegate() override;
+	UE_API virtual void RemoveDelegate() override;
 
 private:
 	FDelegateHandle OnApplyGameplayEffectCallbackDelegateHandle;
 	FDelegateHandle OnPeriodicGameplayEffectExecuteCallbackDelegateHandle;
 };
+
+#undef UE_API
diff --git a/Engine/Plugins/Runtime/GameplayAbilities/Source/GameplayAbilities/Public/Abilities/Tasks/AbilityTask_WaitGameplayEffectBlockedImmunity.h b/Engine/Plugins/Runtime/GameplayAbilities/Source/GameplayAbilities/Public/Abilities/Tasks/AbilityTask_WaitGameplayEffectBlockedImmunity.h
--- a/Engine/Plugins/Runtime/GameplayAbilities/Source/GameplayAbilities/Public/Abilities/Tasks/AbilityTask_WaitGameplayEffectBlockedImmunity.h
+++ b/Engine/Plugins/Runtime/GameplayAbilities/Source/GameplayAbilities/Public/Abilities/Tasks/AbilityTask_WaitGameplayEffectBlockedImmunity.h
@@ -7,12 +7,14 @@
 #include "Abilities/Tasks/AbilityTask.h"
 #include "AbilityTask_WaitGameplayEffectBlockedImmunity.generated.h"
 
+#define UE_API GAMEPLAYABILITIES_API
+
 class UAbilitySystemComponent;
 
 DECLARE_DYNAMIC_MULTICAST_DELEGATE_TwoParams(FGameplayEffectBlockedDelegate, FGameplayEffectSpecHandle, BlockedSpec, FActiveGameplayEffectHandle, ImmunityGameplayEffectHandle);
 
-UCLASS()
-class GAMEPLAYABILITIES_API UAbilityTask_WaitGameplayEffectBlockedImmunity : public UAbilityTask
+UCLASS(MinimalAPI)
+class UAbilityTask_WaitGameplayEffectBlockedImmunity : public UAbilityTask
 {
 	GENERATED_UCLASS_BODY()
 
@@ -21,11 +23,11 @@ class GAMEPLAYABILITIES_API UAbilityTask_WaitGameplayEffectBlockedImmunity : pub
 
 	/** Listens for GE immunity. By default this means "this hero blocked a GE due to immunity". Setting OptionalExternalTarget will listen for a GE being blocked on an external target. Note this only works on the server. */
 	UFUNCTION(BlueprintCallable, Category = "Ability|Tasks", meta = (HidePin = "OwningAbility", DefaultToSelf = "OwningAbility", BlueprintInternalUseOnly = "TRUE"))
-	static UAbilityTask_WaitGameplayEffectBlockedImmunity* WaitGameplayEffectBlockedByImmunity(UGameplayAbility* OwningAbility, FGameplayTagRequirements SourceTagRequirements, FGameplayTagRequirements TargetTagRequirements, AActor* OptionalExternalTarget=nullptr, bool OnlyTriggerOnce=false);
+	static UE_API UAbilityTask_WaitGameplayEffectBlockedImmunity* WaitGameplayEffectBlockedByImmunity(UGameplayAbility* OwningAbility, FGameplayTagRequirements SourceTagRequirements, FGameplayTagRequirements TargetTagRequirements, AActor* OptionalExternalTarget=nullptr, bool OnlyTriggerOnce=false);
 
-	virtual void Activate() override;
+	UE_API virtual void Activate() override;
 
-	virtual void ImmunityCallback(const FGameplayEffectSpec& BlockedSpec, const FActiveGameplayEffect* ImmunityGE);
+	UE_API virtual void ImmunityCallback(const FGameplayEffectSpec& BlockedSpec, const FActiveGameplayEffect* ImmunityGE);
 	
 	FGameplayTagRequirements SourceTagRequirements;
 	FGameplayTagRequirements TargetTagRequirements;
@@ -33,16 +35,16 @@ class GAMEPLAYABILITIES_API UAbilityTask_WaitGameplayEffectBlockedImmunity : pub
 	bool TriggerOnce;
 	bool ListenForPeriodicEffects;
 
-	void SetExternalActor(AActor* InActor);
+	UE_API void SetExternalActor(AActor* InActor);
 
 protected:
 
-	UAbilitySystemComponent* GetASC();
+	UE_API UAbilitySystemComponent* GetASC();
 		
-	void RegisterDelegate();
-	void RemoveDelegate();
+	UE_API void RegisterDelegate();
+	UE_API void RemoveDelegate();
 
-	virtual void OnDestroy(bool AbilityEnded) override;
+	UE_API virtual void OnDestroy(bool AbilityEnded) override;
 
 	bool RegisteredCallback;
 	bool UseExternalOwner;
@@ -52,3 +54,5 @@ protected:
 
 	FDelegateHandle DelegateHandle;
 };
+
+#undef UE_API
diff --git a/Engine/Plugins/Runtime/GameplayAbilities/Source/GameplayAbilities/Public/Abilities/Tasks/AbilityTask_WaitGameplayEffectRemoved.h b/Engine/Plugins/Runtime/GameplayAbilities/Source/GameplayAbilities/Public/Abilities/Tasks/AbilityTask_WaitGameplayEffectRemoved.h
--- a/Engine/Plugins/Runtime/GameplayAbilities/Source/GameplayAbilities/Public/Abilities/Tasks/AbilityTask_WaitGameplayEffectRemoved.h
+++ b/Engine/Plugins/Runtime/GameplayAbilities/Source/GameplayAbilities/Public/Abilities/Tasks/AbilityTask_WaitGameplayEffectRemoved.h
@@ -7,6 +7,8 @@
 #include "Abilities/Tasks/AbilityTask.h"
 #include "AbilityTask_WaitGameplayEffectRemoved.generated.h"
 
+#define UE_API GAMEPLAYABILITIES_API
+
 DECLARE_DYNAMIC_MULTICAST_DELEGATE_OneParam(FWaitGameplayEffectRemovedDelegate, const FGameplayEffectRemovalInfo &, GameplayEffectRemovalInfo);
 
 class AActor;
@@ -15,8 +17,8 @@ class UPrimitiveComponent;
 /**
  *	Waits for the actor to activate another ability
  */
-UCLASS()
-class GAMEPLAYABILITIES_API UAbilityTask_WaitGameplayEffectRemoved : public UAbilityTask
+UCLASS(MinimalAPI)
+class UAbilityTask_WaitGameplayEffectRemoved : public UAbilityTask
 {
 	GENERATED_UCLASS_BODY()
 
@@ -26,22 +28,24 @@ class GAMEPLAYABILITIES_API UAbilityTask_WaitGameplayEffectRemoved : public UAbi
 	UPROPERTY(BlueprintAssignable)
 	FWaitGameplayEffectRemovedDelegate	InvalidHandle;
 
-	virtual void Activate() override;
+	UE_API virtual void Activate() override;
 
 	UFUNCTION()
-	void OnGameplayEffectRemoved(const FGameplayEffectRemovalInfo& InGameplayEffectRemovalInfo);
+	UE_API void OnGameplayEffectRemoved(const FGameplayEffectRemovalInfo& InGameplayEffectRemovalInfo);
 
 	/** Wait until the specified gameplay effect is removed. */
 	UFUNCTION(BlueprintCallable, Category="Ability|Tasks", meta = (HidePin = "OwningAbility", DefaultToSelf = "OwningAbility", BlueprintInternalUseOnly = "TRUE"))
-	static UAbilityTask_WaitGameplayEffectRemoved* WaitForGameplayEffectRemoved(UGameplayAbility* OwningAbility, FActiveGameplayEffectHandle Handle);
+	static UE_API UAbilityTask_WaitGameplayEffectRemoved* WaitForGameplayEffectRemoved(UGameplayAbility* OwningAbility, FActiveGameplayEffectHandle Handle);
 
 	FActiveGameplayEffectHandle Handle;
 
 protected:
 
-	virtual void OnDestroy(bool AbilityIsEnding) override;
+	UE_API virtual void OnDestroy(bool AbilityIsEnding) override;
 	bool Registered;
 
 	FDelegateHandle OnGameplayEffectRemovedDelegateHandle;
 	FDelegateHandle OnGameplayEffectRemovedDelegateHandle_Deprecated;
 };
+
+#undef UE_API
diff --git a/Engine/Plugins/Runtime/GameplayAbilities/Source/GameplayAbilities/Public/Abilities/Tasks/AbilityTask_WaitGameplayEffectStackChange.h b/Engine/Plugins/Runtime/GameplayAbilities/Source/GameplayAbilities/Public/Abilities/Tasks/AbilityTask_WaitGameplayEffectStackChange.h
--- a/Engine/Plugins/Runtime/GameplayAbilities/Source/GameplayAbilities/Public/Abilities/Tasks/AbilityTask_WaitGameplayEffectStackChange.h
+++ b/Engine/Plugins/Runtime/GameplayAbilities/Source/GameplayAbilities/Public/Abilities/Tasks/AbilityTask_WaitGameplayEffectStackChange.h
@@ -7,6 +7,8 @@
 #include "Abilities/Tasks/AbilityTask.h"
 #include "AbilityTask_WaitGameplayEffectStackChange.generated.h"
 
+#define UE_API GAMEPLAYABILITIES_API
+
 DECLARE_DYNAMIC_MULTICAST_DELEGATE_ThreeParams(FWaitGameplayEffectStackChangeDelegate, FActiveGameplayEffectHandle, Handle, int32, NewCount, int32, OldCount);
 
 class AActor;
@@ -15,8 +17,8 @@ class UPrimitiveComponent;
 /**
  *	Waits for the actor to activate another ability
  */
-UCLASS()
-class GAMEPLAYABILITIES_API UAbilityTask_WaitGameplayEffectStackChange : public UAbilityTask
+UCLASS(MinimalAPI)
+class UAbilityTask_WaitGameplayEffectStackChange : public UAbilityTask
 {
 	GENERATED_UCLASS_BODY()
 
@@ -26,21 +28,23 @@ class GAMEPLAYABILITIES_API UAbilityTask_WaitGameplayEffectStackChange : public
 	UPROPERTY(BlueprintAssignable)
 	FWaitGameplayEffectStackChangeDelegate	InvalidHandle;
 
-	virtual void Activate() override;
+	UE_API virtual void Activate() override;
 
 	UFUNCTION()
-	void OnGameplayEffectStackChange(FActiveGameplayEffectHandle Handle, int32 NewCount, int32 OldCount);
+	UE_API void OnGameplayEffectStackChange(FActiveGameplayEffectHandle Handle, int32 NewCount, int32 OldCount);
 
 	/** Wait until the specified gameplay effect is removed. */
 	UFUNCTION(BlueprintCallable, Category="Ability|Tasks", meta = (HidePin = "OwningAbility", DefaultToSelf = "OwningAbility", BlueprintInternalUseOnly = "TRUE"))
-	static UAbilityTask_WaitGameplayEffectStackChange* WaitForGameplayEffectStackChange(UGameplayAbility* OwningAbility, FActiveGameplayEffectHandle Handle);
+	static UE_API UAbilityTask_WaitGameplayEffectStackChange* WaitForGameplayEffectStackChange(UGameplayAbility* OwningAbility, FActiveGameplayEffectHandle Handle);
 
 	FActiveGameplayEffectHandle Handle;
 
 protected:
 
-	virtual void OnDestroy(bool AbilityIsEnding) override;
+	UE_API virtual void OnDestroy(bool AbilityIsEnding) override;
 	bool Registered;
 
 	FDelegateHandle OnGameplayEffectStackChangeDelegateHandle;
 };
+
+#undef UE_API
diff --git a/Engine/Plugins/Runtime/GameplayAbilities/Source/GameplayAbilities/Public/Abilities/Tasks/AbilityTask_WaitGameplayEvent.h b/Engine/Plugins/Runtime/GameplayAbilities/Source/GameplayAbilities/Public/Abilities/Tasks/AbilityTask_WaitGameplayEvent.h
--- a/Engine/Plugins/Runtime/GameplayAbilities/Source/GameplayAbilities/Public/Abilities/Tasks/AbilityTask_WaitGameplayEvent.h
+++ b/Engine/Plugins/Runtime/GameplayAbilities/Source/GameplayAbilities/Public/Abilities/Tasks/AbilityTask_WaitGameplayEvent.h
@@ -8,12 +8,14 @@
 #include "Abilities/Tasks/AbilityTask.h"
 #include "AbilityTask_WaitGameplayEvent.generated.h"
 
+#define UE_API GAMEPLAYABILITIES_API
+
 class UAbilitySystemComponent;
 
 DECLARE_DYNAMIC_MULTICAST_DELEGATE_OneParam(FWaitGameplayEventDelegate, FGameplayEventData, Payload);
 
-UCLASS()
-class GAMEPLAYABILITIES_API UAbilityTask_WaitGameplayEvent : public UAbilityTask
+UCLASS(MinimalAPI)
+class UAbilityTask_WaitGameplayEvent : public UAbilityTask
 {
 	GENERATED_UCLASS_BODY()
 
@@ -26,18 +28,18 @@ class GAMEPLAYABILITIES_API UAbilityTask_WaitGameplayEvent : public UAbilityTask
 	 * If OnlyMatchExact = false it will trigger for nested tags
 	 */
 	UFUNCTION(BlueprintCallable, Category = "Ability|Tasks", meta = (HidePin = "OwningAbility", DefaultToSelf = "OwningAbility", BlueprintInternalUseOnly = "TRUE"))
-	static UAbilityTask_WaitGameplayEvent* WaitGameplayEvent(UGameplayAbility* OwningAbility, UPARAM(meta=(GameplayTagFilter="GameplayEventTagsCategory"))FGameplayTag EventTag, AActor* OptionalExternalTarget=nullptr, bool OnlyTriggerOnce=false, bool OnlyMatchExact = true);
+	static UE_API UAbilityTask_WaitGameplayEvent* WaitGameplayEvent(UGameplayAbility* OwningAbility, UPARAM(meta=(GameplayTagFilter="GameplayEventTagsCategory"))FGameplayTag EventTag, AActor* OptionalExternalTarget=nullptr, bool OnlyTriggerOnce=false, bool OnlyMatchExact = true);
 
-	void SetExternalTarget(AActor* Actor);
+	UE_API void SetExternalTarget(AActor* Actor);
 
-	UAbilitySystemComponent* GetTargetASC();
+	UE_API UAbilitySystemComponent* GetTargetASC();
 
-	virtual void Activate() override;
+	UE_API virtual void Activate() override;
 
-	virtual void GameplayEventCallback(const FGameplayEventData* Payload);
-	virtual void GameplayEventContainerCallback(FGameplayTag MatchingTag, const FGameplayEventData* Payload);
+	UE_API virtual void GameplayEventCallback(const FGameplayEventData* Payload);
+	UE_API virtual void GameplayEventContainerCallback(FGameplayTag MatchingTag, const FGameplayEventData* Payload);
 
-	void OnDestroy(bool AbilityEnding) override;
+	UE_API void OnDestroy(bool AbilityEnding) override;
 
 	FGameplayTag Tag;
 
@@ -50,3 +52,5 @@ class GAMEPLAYABILITIES_API UAbilityTask_WaitGameplayEvent : public UAbilityTask
 
 	FDelegateHandle MyHandle;
 };
+
+#undef UE_API
diff --git a/Engine/Plugins/Runtime/GameplayAbilities/Source/GameplayAbilities/Public/Abilities/Tasks/AbilityTask_WaitGameplayTag.h b/Engine/Plugins/Runtime/GameplayAbilities/Source/GameplayAbilities/Public/Abilities/Tasks/AbilityTask_WaitGameplayTag.h
--- a/Engine/Plugins/Runtime/GameplayAbilities/Source/GameplayAbilities/Public/Abilities/Tasks/AbilityTask_WaitGameplayTag.h
+++ b/Engine/Plugins/Runtime/GameplayAbilities/Source/GameplayAbilities/Public/Abilities/Tasks/AbilityTask_WaitGameplayTag.h
@@ -7,10 +7,12 @@
 #include "Abilities/Tasks/AbilityTask_WaitGameplayTagBase.h"
 #include "AbilityTask_WaitGameplayTag.generated.h"
 
+#define UE_API GAMEPLAYABILITIES_API
+
 DECLARE_DYNAMIC_MULTICAST_DELEGATE(FWaitGameplayTagDelegate);
 
-UCLASS()
-class GAMEPLAYABILITIES_API UAbilityTask_WaitGameplayTagAdded : public UAbilityTask_WaitGameplayTag
+UCLASS(MinimalAPI)
+class UAbilityTask_WaitGameplayTagAdded : public UAbilityTask_WaitGameplayTag
 {
 	GENERATED_UCLASS_BODY()
 
@@ -22,15 +24,15 @@ class GAMEPLAYABILITIES_API UAbilityTask_WaitGameplayTagAdded : public UAbilityT
 	 *  If the tag is already present when this task is started, it will immediately broadcast the Added event. It will keep listening as long as OnlyTriggerOnce = false.
 	 */
 	UFUNCTION(BlueprintCallable, Category = "Ability|Tasks", meta = (HidePin = "OwningAbility", DefaultToSelf = "OwningAbility", BlueprintInternalUseOnly = "TRUE"))
-	static UAbilityTask_WaitGameplayTagAdded* WaitGameplayTagAdd(UGameplayAbility* OwningAbility, FGameplayTag Tag, AActor* InOptionalExternalTarget=nullptr, bool OnlyTriggerOnce=false);
+	static UE_API UAbilityTask_WaitGameplayTagAdded* WaitGameplayTagAdd(UGameplayAbility* OwningAbility, FGameplayTag Tag, AActor* InOptionalExternalTarget=nullptr, bool OnlyTriggerOnce=false);
 
-	virtual void Activate() override;
+	UE_API virtual void Activate() override;
 
-	virtual void GameplayTagCallback(const FGameplayTag Tag, int32 NewCount) override;
+	UE_API virtual void GameplayTagCallback(const FGameplayTag Tag, int32 NewCount) override;
 };
 
-UCLASS()
-class GAMEPLAYABILITIES_API UAbilityTask_WaitGameplayTagRemoved : public UAbilityTask_WaitGameplayTag
+UCLASS(MinimalAPI)
+class UAbilityTask_WaitGameplayTagRemoved : public UAbilityTask_WaitGameplayTag
 {
 	GENERATED_UCLASS_BODY()
 
@@ -42,9 +44,11 @@ class GAMEPLAYABILITIES_API UAbilityTask_WaitGameplayTagRemoved : public UAbilit
 	 *  If the tag is not present when this task is started, it will immediately broadcast the Removed event. It will keep listening as long as OnlyTriggerOnce = false.
 	 */
 	UFUNCTION(BlueprintCallable, Category = "Ability|Tasks", meta = (HidePin = "OwningAbility", DefaultToSelf = "OwningAbility", BlueprintInternalUseOnly = "TRUE"))
-	static UAbilityTask_WaitGameplayTagRemoved* WaitGameplayTagRemove(UGameplayAbility* OwningAbility, FGameplayTag Tag, AActor* InOptionalExternalTarget=nullptr, bool OnlyTriggerOnce=false);
+	static UE_API UAbilityTask_WaitGameplayTagRemoved* WaitGameplayTagRemove(UGameplayAbility* OwningAbility, FGameplayTag Tag, AActor* InOptionalExternalTarget=nullptr, bool OnlyTriggerOnce=false);
 
-	virtual void Activate() override;
+	UE_API virtual void Activate() override;
 
-	virtual void GameplayTagCallback(const FGameplayTag Tag, int32 NewCount) override;
+	UE_API virtual void GameplayTagCallback(const FGameplayTag Tag, int32 NewCount) override;
 };
+
+#undef UE_API
diff --git a/Engine/Plugins/Runtime/GameplayAbilities/Source/GameplayAbilities/Public/Abilities/Tasks/AbilityTask_WaitGameplayTagBase.h b/Engine/Plugins/Runtime/GameplayAbilities/Source/GameplayAbilities/Public/Abilities/Tasks/AbilityTask_WaitGameplayTagBase.h
--- a/Engine/Plugins/Runtime/GameplayAbilities/Source/GameplayAbilities/Public/Abilities/Tasks/AbilityTask_WaitGameplayTagBase.h
+++ b/Engine/Plugins/Runtime/GameplayAbilities/Source/GameplayAbilities/Public/Abilities/Tasks/AbilityTask_WaitGameplayTagBase.h
@@ -7,27 +7,29 @@
 #include "Abilities/Tasks/AbilityTask.h"
 #include "AbilityTask_WaitGameplayTagBase.generated.h"
 
+#define UE_API GAMEPLAYABILITIES_API
+
 class UAbilitySystemComponent;
 
-UCLASS()
-class GAMEPLAYABILITIES_API UAbilityTask_WaitGameplayTag : public UAbilityTask
+UCLASS(MinimalAPI)
+class UAbilityTask_WaitGameplayTag : public UAbilityTask
 {
 	GENERATED_UCLASS_BODY()
 
-	virtual void Activate() override;
+	UE_API virtual void Activate() override;
 
 	UFUNCTION()
-	virtual void GameplayTagCallback(const FGameplayTag Tag, int32 NewCount);
+	UE_API virtual void GameplayTagCallback(const FGameplayTag Tag, int32 NewCount);
 	
-	void SetExternalTarget(AActor* Actor);
+	UE_API void SetExternalTarget(AActor* Actor);
 
 	FGameplayTag	Tag;
 
 protected:
 
-	UAbilitySystemComponent* GetTargetASC();
+	UE_API UAbilitySystemComponent* GetTargetASC();
 
-	virtual void OnDestroy(bool AbilityIsEnding) override;
+	UE_API virtual void OnDestroy(bool AbilityIsEnding) override;
 	
 	bool RegisteredCallback;
 	
@@ -39,3 +41,5 @@ protected:
 
 	FDelegateHandle DelegateHandle;
 };
+
+#undef UE_API
diff --git a/Engine/Plugins/Runtime/GameplayAbilities/Source/GameplayAbilities/Public/Abilities/Tasks/AbilityTask_WaitGameplayTagCountChanged.h b/Engine/Plugins/Runtime/GameplayAbilities/Source/GameplayAbilities/Public/Abilities/Tasks/AbilityTask_WaitGameplayTagCountChanged.h
--- a/Engine/Plugins/Runtime/GameplayAbilities/Source/GameplayAbilities/Public/Abilities/Tasks/AbilityTask_WaitGameplayTagCountChanged.h
+++ b/Engine/Plugins/Runtime/GameplayAbilities/Source/GameplayAbilities/Public/Abilities/Tasks/AbilityTask_WaitGameplayTagCountChanged.h
@@ -5,10 +5,12 @@
 #include "GameplayTagContainer.h"
 #include "AbilityTask_WaitGameplayTagCountChanged.generated.h"
 
+#define UE_API GAMEPLAYABILITIES_API
+
 DECLARE_DYNAMIC_MULTICAST_DELEGATE_OneParam(FWaitGameplayTagCountDelegate, int32, TagCount);
 
-UCLASS()
-class GAMEPLAYABILITIES_API UAbilityTask_WaitGameplayTagCountChanged : public UAbilityTask
+UCLASS(MinimalAPI)
+class UAbilityTask_WaitGameplayTagCountChanged : public UAbilityTask
 {
 	GENERATED_BODY()
 
@@ -20,15 +22,15 @@ protected:
 	 * 	Wait until the specified gameplay tag count has changed. By default this will look at the owner of this ability. OptionalExternalTarget can be set to make this look at another actor's tags for changes. 
 	 */
 	UFUNCTION(BlueprintCallable, Category = "Ability|Tasks", meta = (HidePin = "OwningAbility", DefaultToSelf = "OwningAbility", BlueprintInternalUseOnly = "TRUE"))
-	static UAbilityTask_WaitGameplayTagCountChanged* WaitGameplayTagCountChange(UGameplayAbility* OwningAbility, FGameplayTag Tag, AActor* InOptionalExternalTarget = nullptr);
+	static UE_API UAbilityTask_WaitGameplayTagCountChanged* WaitGameplayTagCountChange(UGameplayAbility* OwningAbility, FGameplayTag Tag, AActor* InOptionalExternalTarget = nullptr);
 
-	virtual void Activate() override;
-	virtual void GameplayTagCallback(const FGameplayTag Tag, int32 NewCount);
+	UE_API virtual void Activate() override;
+	UE_API virtual void GameplayTagCallback(const FGameplayTag Tag, int32 NewCount);
 
 	FGameplayTag Tag;
 
-	UAbilitySystemComponent* GetTargetASC();
-	virtual void OnDestroy(bool bAbilityIsEnding) override;
+	UE_API UAbilitySystemComponent* GetTargetASC();
+	UE_API virtual void OnDestroy(bool bAbilityIsEnding) override;
 
 	UPROPERTY()
 	TObjectPtr<UAbilitySystemComponent> OptionalExternalTarget;
@@ -36,3 +38,4 @@ protected:
 	FDelegateHandle GameplayTagCountChangedHandle;
 };
 
+#undef UE_API
diff --git a/Engine/Plugins/Runtime/GameplayAbilities/Source/GameplayAbilities/Public/Abilities/Tasks/AbilityTask_WaitGameplayTagQuery.h b/Engine/Plugins/Runtime/GameplayAbilities/Source/GameplayAbilities/Public/Abilities/Tasks/AbilityTask_WaitGameplayTagQuery.h
--- a/Engine/Plugins/Runtime/GameplayAbilities/Source/GameplayAbilities/Public/Abilities/Tasks/AbilityTask_WaitGameplayTagQuery.h
+++ b/Engine/Plugins/Runtime/GameplayAbilities/Source/GameplayAbilities/Public/Abilities/Tasks/AbilityTask_WaitGameplayTagQuery.h
@@ -5,6 +5,8 @@
 
 #include "AbilityTask_WaitGameplayTagQuery.generated.h"
 
+#define UE_API GAMEPLAYABILITIES_API
+
 DECLARE_DYNAMIC_MULTICAST_DELEGATE(FWaitGameplayTagQueryDelegate);
 
 /** This enum defines the condition in which the wait gameplay tag query node will trigger. */
@@ -16,17 +18,17 @@ enum class EWaitGameplayTagQueryTriggerCondition : uint8
 };
 
 /** This class defines a node to wait on a gameplay tag query. */
-UCLASS()
-class GAMEPLAYABILITIES_API UAbilityTask_WaitGameplayTagQuery : public UAbilityTask
+UCLASS(MinimalAPI)
+class UAbilityTask_WaitGameplayTagQuery : public UAbilityTask
 {
 	GENERATED_BODY()
 	
 public:
 	/** Activates this AbilityTask. */
-	virtual void Activate() override;
+	UE_API virtual void Activate() override;
 	
 	/** Sets the external target that this node whould use for checking tags. */
-	void SetExternalTarget(const AActor* Actor);
+	UE_API void SetExternalTarget(const AActor* Actor);
 
 	/**
 	 * 	Wait until the given gameplay tag query has become true or false, based on TriggerCondition. 
@@ -36,7 +38,7 @@ public:
 	 *  event. It will keep listening as long as bOnlyTriggerOnce = false.
 	 */
 	UFUNCTION(BlueprintCallable, Category = "Ability|Tasks", meta = (HidePin = "OwningAbility", DefaultToSelf = "OwningAbility", BlueprintInternalUseOnly = "TRUE"))
-	static UAbilityTask_WaitGameplayTagQuery* WaitGameplayTagQuery(UGameplayAbility* OwningAbility, 
+	static UE_API UAbilityTask_WaitGameplayTagQuery* WaitGameplayTagQuery(UGameplayAbility* OwningAbility, 
 																   const FGameplayTagQuery TagQuery, 
 																   const AActor* InOptionalExternalTarget = nullptr, 
 																   const EWaitGameplayTagQueryTriggerCondition TriggerCondition = EWaitGameplayTagQueryTriggerCondition::WhenTrue, 
@@ -46,16 +48,16 @@ protected:
 
 	/** This will update the tags in the TargetTags container to reflect that tags that are on the target ASC. */
 	UFUNCTION()
-	void UpdateTargetTags(const FGameplayTag Tag, int32 NewCount);
+	UE_API void UpdateTargetTags(const FGameplayTag Tag, int32 NewCount);
 
 	/** This will evaluate the TargetTags using the given TagQuery, executing the Trigger delegate if needed. */
-	void EvaluateTagQuery();
+	UE_API void EvaluateTagQuery();
 
 	/** This gets the ASC to use to listen to tag changed events. */
-	UAbilitySystemComponent* GetTargetASC();
+	UE_API UAbilitySystemComponent* GetTargetASC();
 
 	/** This will handle cleaning up any registered delegates. */
-	virtual void OnDestroy(bool AbilityIsEnding) override;
+	UE_API virtual void OnDestroy(bool AbilityIsEnding) override;
 
 	/** This is the tag query to evaluate for triggering this node. */
 	FGameplayTagQuery TagQuery;
@@ -89,4 +91,6 @@ protected:
 
 	/** This indicates if this node should only trigger once, or any number of times. */
 	bool bOnlyTriggerOnce = false;
-};
\ No newline at end of file
+};
+
+#undef UE_API
diff --git a/Engine/Plugins/Runtime/GameplayAbilities/Source/GameplayAbilities/Public/Abilities/Tasks/AbilityTask_WaitInputPress.h b/Engine/Plugins/Runtime/GameplayAbilities/Source/GameplayAbilities/Public/Abilities/Tasks/AbilityTask_WaitInputPress.h
--- a/Engine/Plugins/Runtime/GameplayAbilities/Source/GameplayAbilities/Public/Abilities/Tasks/AbilityTask_WaitInputPress.h
+++ b/Engine/Plugins/Runtime/GameplayAbilities/Source/GameplayAbilities/Public/Abilities/Tasks/AbilityTask_WaitInputPress.h
@@ -6,14 +6,16 @@
 #include "Abilities/Tasks/AbilityTask.h"
 #include "AbilityTask_WaitInputPress.generated.h"
 
+#define UE_API GAMEPLAYABILITIES_API
+
 DECLARE_DYNAMIC_MULTICAST_DELEGATE_OneParam(FInputPressDelegate, float, TimeWaited);
 
 /**
  *	Waits until the input is pressed from activating an ability. This should be true immediately upon starting the ability, since the key was pressed to activate it.
  *	We expect server to execute this task in parallel and keep its own time. We do not keep track of 
  */
-UCLASS()
-class GAMEPLAYABILITIES_API UAbilityTask_WaitInputPress : public UAbilityTask
+UCLASS(MinimalAPI)
+class UAbilityTask_WaitInputPress : public UAbilityTask
 {
 	GENERATED_UCLASS_BODY()
 
@@ -21,13 +23,13 @@ class GAMEPLAYABILITIES_API UAbilityTask_WaitInputPress : public UAbilityTask
 	FInputPressDelegate		OnPress;
 
 	UFUNCTION()
-	void OnPressCallback();
+	UE_API void OnPressCallback();
 
-	virtual void Activate() override;
+	UE_API virtual void Activate() override;
 
 	/** Wait until the user presses the input button for this ability's activation. Returns time this node spent waiting for the press. Will return 0 if input was already down. */
 	UFUNCTION(BlueprintCallable, Category = "Ability|Tasks", meta = (HidePin = "OwningAbility", DefaultToSelf = "OwningAbility", BlueprintInternalUseOnly = "TRUE"))
-	static UAbilityTask_WaitInputPress* WaitInputPress(UGameplayAbility* OwningAbility, bool bTestAlreadyPressed=false);
+	static UE_API UAbilityTask_WaitInputPress* WaitInputPress(UGameplayAbility* OwningAbility, bool bTestAlreadyPressed=false);
 
 protected:
 
@@ -35,3 +37,5 @@ protected:
 	bool bTestInitialState;
 	FDelegateHandle DelegateHandle;
 };
+
+#undef UE_API
diff --git a/Engine/Plugins/Runtime/GameplayAbilities/Source/GameplayAbilities/Public/Abilities/Tasks/AbilityTask_WaitInputRelease.h b/Engine/Plugins/Runtime/GameplayAbilities/Source/GameplayAbilities/Public/Abilities/Tasks/AbilityTask_WaitInputRelease.h
--- a/Engine/Plugins/Runtime/GameplayAbilities/Source/GameplayAbilities/Public/Abilities/Tasks/AbilityTask_WaitInputRelease.h
+++ b/Engine/Plugins/Runtime/GameplayAbilities/Source/GameplayAbilities/Public/Abilities/Tasks/AbilityTask_WaitInputRelease.h
@@ -6,14 +6,16 @@
 #include "Abilities/Tasks/AbilityTask.h"
 #include "AbilityTask_WaitInputRelease.generated.h"
 
+#define UE_API GAMEPLAYABILITIES_API
+
 DECLARE_DYNAMIC_MULTICAST_DELEGATE_OneParam(FInputReleaseDelegate, float, TimeHeld);
 
 /**
  *	Waits until the input is released from activating an ability. Clients will replicate a 'release input' event to the server, but not the exact time it was held locally.
  *	We expect server to execute this task in parallel and keep its own time.
  */
-UCLASS()
-class GAMEPLAYABILITIES_API UAbilityTask_WaitInputRelease : public UAbilityTask
+UCLASS(MinimalAPI)
+class UAbilityTask_WaitInputRelease : public UAbilityTask
 {
 	GENERATED_UCLASS_BODY()
 
@@ -21,13 +23,13 @@ class GAMEPLAYABILITIES_API UAbilityTask_WaitInputRelease : public UAbilityTask
 	FInputReleaseDelegate	OnRelease;
 
 	UFUNCTION()
-	void OnReleaseCallback();
+	UE_API void OnReleaseCallback();
 
-	virtual void Activate() override;
+	UE_API virtual void Activate() override;
 
 	/** Wait until the user releases the input button for this ability's activation. Returns time from hitting this node, till release. Will return 0 if input was already released. */
 	UFUNCTION(BlueprintCallable, Category="Ability|Tasks", meta = (HidePin = "OwningAbility", DefaultToSelf = "OwningAbility", BlueprintInternalUseOnly = "TRUE"))
-	static UAbilityTask_WaitInputRelease* WaitInputRelease(UGameplayAbility* OwningAbility, bool bTestAlreadyReleased=false);
+	static UE_API UAbilityTask_WaitInputRelease* WaitInputRelease(UGameplayAbility* OwningAbility, bool bTestAlreadyReleased=false);
 
 protected:
 
@@ -35,3 +37,5 @@ protected:
 	bool bTestInitialState;
 	FDelegateHandle DelegateHandle;
 };
+
+#undef UE_API
diff --git a/Engine/Plugins/Runtime/GameplayAbilities/Source/GameplayAbilities/Public/Abilities/Tasks/AbilityTask_WaitMovementModeChange.h b/Engine/Plugins/Runtime/GameplayAbilities/Source/GameplayAbilities/Public/Abilities/Tasks/AbilityTask_WaitMovementModeChange.h
--- a/Engine/Plugins/Runtime/GameplayAbilities/Source/GameplayAbilities/Public/Abilities/Tasks/AbilityTask_WaitMovementModeChange.h
+++ b/Engine/Plugins/Runtime/GameplayAbilities/Source/GameplayAbilities/Public/Abilities/Tasks/AbilityTask_WaitMovementModeChange.h
@@ -6,14 +6,16 @@
 #include "Abilities/Tasks/AbilityTask.h"
 #include "AbilityTask_WaitMovementModeChange.generated.h"
 
+#define UE_API GAMEPLAYABILITIES_API
+
 class ACharacter;
 
 DECLARE_DYNAMIC_MULTICAST_DELEGATE_OneParam(FMovementModeChangedDelegate, EMovementMode, NewMovementMode);
 
 class ACharacter;
 
-UCLASS()
-class GAMEPLAYABILITIES_API UAbilityTask_WaitMovementModeChange : public UAbilityTask
+UCLASS(MinimalAPI)
+class UAbilityTask_WaitMovementModeChange : public UAbilityTask
 {
 	GENERATED_UCLASS_BODY()
 
@@ -23,18 +25,20 @@ public:
 	FMovementModeChangedDelegate	OnChange;
 
 	UFUNCTION()
-	void OnMovementModeChange(ACharacter * Character, EMovementMode PrevMovementMode, uint8 PreviousCustomMode);
+	UE_API void OnMovementModeChange(ACharacter * Character, EMovementMode PrevMovementMode, uint8 PreviousCustomMode);
 
 	EMovementMode	RequiredMode;
 
 	/** Wait until movement mode changes (E.g., landing) */
 	UFUNCTION(BlueprintCallable, Category="Ability|Tasks", meta = (HidePin = "OwningAbility", DefaultToSelf = "OwningAbility", BlueprintInternalUseOnly = "TRUE", DisplayName="WaitMovementModeChange"))
-	static UAbilityTask_WaitMovementModeChange* CreateWaitMovementModeChange(UGameplayAbility* OwningAbility, EMovementMode NewMode);
+	static UE_API UAbilityTask_WaitMovementModeChange* CreateWaitMovementModeChange(UGameplayAbility* OwningAbility, EMovementMode NewMode);
 
-	virtual void Activate() override;
+	UE_API virtual void Activate() override;
 private:
 
-	virtual void OnDestroy(bool AbilityEnded) override;
+	UE_API virtual void OnDestroy(bool AbilityEnded) override;
 
 	TWeakObjectPtr<ACharacter>	MyCharacter;
 };
+
+#undef UE_API
diff --git a/Engine/Plugins/Runtime/GameplayAbilities/Source/GameplayAbilities/Public/Abilities/Tasks/AbilityTask_WaitOverlap.h b/Engine/Plugins/Runtime/GameplayAbilities/Source/GameplayAbilities/Public/Abilities/Tasks/AbilityTask_WaitOverlap.h
--- a/Engine/Plugins/Runtime/GameplayAbilities/Source/GameplayAbilities/Public/Abilities/Tasks/AbilityTask_WaitOverlap.h
+++ b/Engine/Plugins/Runtime/GameplayAbilities/Source/GameplayAbilities/Public/Abilities/Tasks/AbilityTask_WaitOverlap.h
@@ -6,6 +6,8 @@
 #include "Abilities/Tasks/AbilityTask.h"
 #include "AbilityTask_WaitOverlap.generated.h"
 
+#define UE_API GAMEPLAYABILITIES_API
+
 DECLARE_DYNAMIC_MULTICAST_DELEGATE_OneParam(FWaitOverlapDelegate, const FGameplayAbilityTargetDataHandle&, TargetData);
 
 class AActor;
@@ -16,8 +18,8 @@ class UPrimitiveComponent;
  *		-Only actually activates on Blocking hits
  *		-Uses first PrimitiveComponent instead of being able to specify arbitrary component.
  */
-UCLASS()
-class GAMEPLAYABILITIES_API UAbilityTask_WaitOverlap : public UAbilityTask
+UCLASS(MinimalAPI)
+class UAbilityTask_WaitOverlap : public UAbilityTask
 {
 	GENERATED_UCLASS_BODY()
 
@@ -25,18 +27,20 @@ class GAMEPLAYABILITIES_API UAbilityTask_WaitOverlap : public UAbilityTask
 	FWaitOverlapDelegate	OnOverlap;
 
 	UFUNCTION()
-	void OnHitCallback(UPrimitiveComponent* HitComp, AActor* OtherActor, UPrimitiveComponent* OtherComp, FVector NormalImpulse, const FHitResult& Hit);
+	UE_API void OnHitCallback(UPrimitiveComponent* HitComp, AActor* OtherActor, UPrimitiveComponent* OtherComp, FVector NormalImpulse, const FHitResult& Hit);
 
-	virtual void Activate() override;
+	UE_API virtual void Activate() override;
 
 	/** Wait until an overlap occurs. This will need to be better fleshed out so we can specify game specific collision requirements */
 	UFUNCTION(BlueprintCallable, Category="Ability|Tasks", meta = (HidePin = "OwningAbility", DefaultToSelf = "OwningAbility", BlueprintInternalUseOnly = "TRUE"))
-	static UAbilityTask_WaitOverlap* WaitForOverlap(UGameplayAbility* OwningAbility);
+	static UE_API UAbilityTask_WaitOverlap* WaitForOverlap(UGameplayAbility* OwningAbility);
 
 private:
 
-	virtual void OnDestroy(bool AbilityEnded) override;
+	UE_API virtual void OnDestroy(bool AbilityEnded) override;
 
 	UPrimitiveComponent* GetComponent();
 	
 };
+
+#undef UE_API
diff --git a/Engine/Plugins/Runtime/GameplayAbilities/Source/GameplayAbilities/Public/Abilities/Tasks/AbilityTask_WaitTargetData.h b/Engine/Plugins/Runtime/GameplayAbilities/Source/GameplayAbilities/Public/Abilities/Tasks/AbilityTask_WaitTargetData.h
--- a/Engine/Plugins/Runtime/GameplayAbilities/Source/GameplayAbilities/Public/Abilities/Tasks/AbilityTask_WaitTargetData.h
+++ b/Engine/Plugins/Runtime/GameplayAbilities/Source/GameplayAbilities/Public/Abilities/Tasks/AbilityTask_WaitTargetData.h
@@ -9,6 +9,8 @@
 #include "Abilities/Tasks/AbilityTask.h"
 #include "AbilityTask_WaitTargetData.generated.h"
 
+#define UE_API GAMEPLAYABILITIES_API
+
 DECLARE_DYNAMIC_MULTICAST_DELEGATE_OneParam(FWaitTargetDataDelegate, const FGameplayAbilityTargetDataHandle&, Data);
 
 /**
@@ -18,8 +20,8 @@ DECLARE_DYNAMIC_MULTICAST_DELEGATE_OneParam(FWaitTargetDataDelegate, const FGame
  * For most games you will need to subclass and heavily modify this actor, or you will want to implement similar functions in a game-specific actor or blueprint to avoid actor spawn costs
  * This task is not well tested by internal games, but it is a useful class to look at to learn how target replication occurs
  */
-UCLASS()
-class GAMEPLAYABILITIES_API UAbilityTask_WaitTargetData: public UAbilityTask
+UCLASS(MinimalAPI)
+class UAbilityTask_WaitTargetData: public UAbilityTask
 {
 	GENERATED_UCLASS_BODY()
 
@@ -30,50 +32,50 @@ class GAMEPLAYABILITIES_API UAbilityTask_WaitTargetData: public UAbilityTask
 	FWaitTargetDataDelegate	Cancelled;
 
 	UFUNCTION()
-	virtual void OnTargetDataReplicatedCallback(const FGameplayAbilityTargetDataHandle& Data, FGameplayTag ActivationTag);
+	UE_API virtual void OnTargetDataReplicatedCallback(const FGameplayAbilityTargetDataHandle& Data, FGameplayTag ActivationTag);
 
 	UFUNCTION()
-	virtual void OnTargetDataReplicatedCancelledCallback();
+	UE_API virtual void OnTargetDataReplicatedCancelledCallback();
 
 	UFUNCTION()
-	virtual void OnTargetDataReadyCallback(const FGameplayAbilityTargetDataHandle& Data);
+	UE_API virtual void OnTargetDataReadyCallback(const FGameplayAbilityTargetDataHandle& Data);
 
 	UFUNCTION()
-	virtual void OnTargetDataCancelledCallback(const FGameplayAbilityTargetDataHandle& Data);
+	UE_API virtual void OnTargetDataCancelledCallback(const FGameplayAbilityTargetDataHandle& Data);
 
 	/** Spawns target actor and waits for it to return valid data or to be canceled. */
 	UFUNCTION(BlueprintCallable, meta=(HidePin = "OwningAbility", DefaultToSelf = "OwningAbility", BlueprintInternalUseOnly = "true", HideSpawnParms="Instigator"), Category="Ability|Tasks")
-	static UAbilityTask_WaitTargetData* WaitTargetData(UGameplayAbility* OwningAbility, FName TaskInstanceName, TEnumAsByte<EGameplayTargetingConfirmation::Type> ConfirmationType, TSubclassOf<AGameplayAbilityTargetActor> Class);
+	static UE_API UAbilityTask_WaitTargetData* WaitTargetData(UGameplayAbility* OwningAbility, FName TaskInstanceName, TEnumAsByte<EGameplayTargetingConfirmation::Type> ConfirmationType, TSubclassOf<AGameplayAbilityTargetActor> Class);
 
 	/** Uses specified target actor and waits for it to return valid data or to be canceled. */
 	UFUNCTION(BlueprintCallable, meta = (HidePin = "OwningAbility", DefaultToSelf = "OwningAbility", BlueprintInternalUseOnly = "true", HideSpawnParms = "Instigator"), Category = "Ability|Tasks")
-	static UAbilityTask_WaitTargetData* WaitTargetDataUsingActor(UGameplayAbility* OwningAbility, FName TaskInstanceName, TEnumAsByte<EGameplayTargetingConfirmation::Type> ConfirmationType, AGameplayAbilityTargetActor* TargetActor);
+	static UE_API UAbilityTask_WaitTargetData* WaitTargetDataUsingActor(UGameplayAbility* OwningAbility, FName TaskInstanceName, TEnumAsByte<EGameplayTargetingConfirmation::Type> ConfirmationType, AGameplayAbilityTargetActor* TargetActor);
 
-	virtual void Activate() override;
+	UE_API virtual void Activate() override;
 
 	UFUNCTION(BlueprintCallable, meta = (HidePin = "OwningAbility", DefaultToSelf = "OwningAbility", BlueprintInternalUseOnly = "true"), Category = "Abilities")
-	virtual bool BeginSpawningActor(UGameplayAbility* OwningAbility, TSubclassOf<AGameplayAbilityTargetActor> Class, AGameplayAbilityTargetActor*& SpawnedActor);
+	UE_API virtual bool BeginSpawningActor(UGameplayAbility* OwningAbility, TSubclassOf<AGameplayAbilityTargetActor> Class, AGameplayAbilityTargetActor*& SpawnedActor);
 
 	UFUNCTION(BlueprintCallable, meta = (HidePin = "OwningAbility", DefaultToSelf = "OwningAbility", BlueprintInternalUseOnly = "true"), Category = "Abilities")
-	virtual void FinishSpawningActor(UGameplayAbility* OwningAbility, AGameplayAbilityTargetActor* SpawnedActor);
+	UE_API virtual void FinishSpawningActor(UGameplayAbility* OwningAbility, AGameplayAbilityTargetActor* SpawnedActor);
 
 	/** Called when the ability is asked to confirm from an outside node. What this means depends on the individual task. By default, this does nothing other than ending if bEndTask is true. */
-	virtual void ExternalConfirm(bool bEndTask) override;
+	UE_API virtual void ExternalConfirm(bool bEndTask) override;
 
 	/** Called when the ability is asked to cancel from an outside node. What this means depends on the individual task. By default, this does nothing other than ending the task. */
-	virtual void ExternalCancel() override;
+	UE_API virtual void ExternalCancel() override;
 
 protected:
 
-	virtual bool ShouldSpawnTargetActor() const;
-	virtual void InitializeTargetActor(AGameplayAbilityTargetActor* SpawnedActor) const;
-	virtual void FinalizeTargetActor(AGameplayAbilityTargetActor* SpawnedActor) const;
+	UE_API virtual bool ShouldSpawnTargetActor() const;
+	UE_API virtual void InitializeTargetActor(AGameplayAbilityTargetActor* SpawnedActor) const;
+	UE_API virtual void FinalizeTargetActor(AGameplayAbilityTargetActor* SpawnedActor) const;
 
-	virtual void RegisterTargetDataCallbacks();
+	UE_API virtual void RegisterTargetDataCallbacks();
 
-	virtual void OnDestroy(bool AbilityEnded) override;
+	UE_API virtual void OnDestroy(bool AbilityEnded) override;
 
-	virtual bool ShouldReplicateDataToServer() const;
+	UE_API virtual bool ShouldReplicateDataToServer() const;
 
 protected:
 
@@ -98,3 +100,5 @@ protected:
 *		-Have a function named FinishSpawningActor w/ an AActor* of the class you spawned
 *			-This function *must* call ExecuteConstruction + PostActorConstruction
 */
+
+#undef UE_API
diff --git a/Engine/Plugins/Runtime/GameplayAbilities/Source/GameplayAbilities/Public/Abilities/Tasks/AbilityTask_WaitVelocityChange.h b/Engine/Plugins/Runtime/GameplayAbilities/Source/GameplayAbilities/Public/Abilities/Tasks/AbilityTask_WaitVelocityChange.h
--- a/Engine/Plugins/Runtime/GameplayAbilities/Source/GameplayAbilities/Public/Abilities/Tasks/AbilityTask_WaitVelocityChange.h
+++ b/Engine/Plugins/Runtime/GameplayAbilities/Source/GameplayAbilities/Public/Abilities/Tasks/AbilityTask_WaitVelocityChange.h
@@ -6,12 +6,14 @@
 #include "Abilities/Tasks/AbilityTask.h"
 #include "AbilityTask_WaitVelocityChange.generated.h"
 
+#define UE_API GAMEPLAYABILITIES_API
+
 class UMovementComponent;
 
 DECLARE_DYNAMIC_MULTICAST_DELEGATE(FWaitVelocityChangeDelegate);
 
-UCLASS()
-class GAMEPLAYABILITIES_API UAbilityTask_WaitVelocityChange: public UAbilityTask
+UCLASS(MinimalAPI)
+class UAbilityTask_WaitVelocityChange: public UAbilityTask
 {
 	GENERATED_UCLASS_BODY()
 
@@ -19,13 +21,13 @@ class GAMEPLAYABILITIES_API UAbilityTask_WaitVelocityChange: public UAbilityTask
 	UPROPERTY(BlueprintAssignable)
 	FWaitVelocityChangeDelegate OnVelocityChage;
 
-	virtual void TickTask(float DeltaTime) override;
+	UE_API virtual void TickTask(float DeltaTime) override;
 
 	/** Wait for the actor's movement component velocity to be of minimum magnitude when projected along given direction */
 	UFUNCTION(BlueprintCallable, Category = "Ability|Tasks", meta = (DisplayName="WaitVelocityChange",HidePin = "OwningAbility", DefaultToSelf = "OwningAbility", BlueprintInternalUseOnly = "TRUE"))
-	static UAbilityTask_WaitVelocityChange* CreateWaitVelocityChange(UGameplayAbility* OwningAbility, FVector Direction, float MinimumMagnitude);
+	static UE_API UAbilityTask_WaitVelocityChange* CreateWaitVelocityChange(UGameplayAbility* OwningAbility, FVector Direction, float MinimumMagnitude);
 		
-	virtual void Activate() override;
+	UE_API virtual void Activate() override;
 
 protected:
 
@@ -35,3 +37,5 @@ protected:
 	float	MinimumMagnitude;
 	FVector Direction;
 };
+
+#undef UE_API
diff --git a/Engine/Plugins/Runtime/GameplayAbilities/Source/GameplayAbilities/Public/AbilitySystemBlueprintLibrary.h b/Engine/Plugins/Runtime/GameplayAbilities/Source/GameplayAbilities/Public/AbilitySystemBlueprintLibrary.h
--- a/Engine/Plugins/Runtime/GameplayAbilities/Source/GameplayAbilities/Public/AbilitySystemBlueprintLibrary.h
+++ b/Engine/Plugins/Runtime/GameplayAbilities/Source/GameplayAbilities/Public/AbilitySystemBlueprintLibrary.h
@@ -17,6 +17,8 @@
 #include "ScalableFloat.h"
 #include "AbilitySystemBlueprintLibrary.generated.h"
 
+#define UE_API GAMEPLAYABILITIES_API
+
 class UAbilitySystemComponent;
 class UGameplayEffect;
 class UGameplayEffectUIData;
@@ -63,14 +65,14 @@ struct FGameplayTagChangedEventWrapperSpecHandle
 };
 
 /** Blueprint library for ability system. Many of these functions are useful to call from native as well */
-UCLASS(meta=(ScriptName="AbilitySystemLibrary"))
-class GAMEPLAYABILITIES_API UAbilitySystemBlueprintLibrary : public UBlueprintFunctionLibrary
+UCLASS(meta=(ScriptName="AbilitySystemLibrary"), MinimalAPI)
+class UAbilitySystemBlueprintLibrary : public UBlueprintFunctionLibrary
 {
 	GENERATED_UCLASS_BODY()
 
 	/** Tries to find an ability system component on the actor, will use AbilitySystemInterface or fall back to a component search */
 	UFUNCTION(BlueprintPure, Category = Ability, Meta=(DefaultToSelf = "Actor"))
-	static UAbilitySystemComponent* GetAbilitySystemComponent(AActor* Actor);
+	static UE_API UAbilitySystemComponent* GetAbilitySystemComponent(AActor* Actor);
 
 	/**
 	 * This function can be used to trigger an ability on the actor in question with useful payload data.
@@ -78,7 +80,7 @@ class GAMEPLAYABILITIES_API UAbilitySystemBlueprintLibrary : public UBlueprintFu
 	 * found, the event will not be sent.
 	 */
 	UFUNCTION(BlueprintCallable, Category = Ability, Meta = (Tooltip = "This function can be used to trigger an ability on the actor in question with useful payload data."))
-	static void SendGameplayEventToActor(AActor* Actor, UPARAM(meta=(GameplayTagFilter="GameplayEventTagsCategory")) FGameplayTag EventTag, FGameplayEventData Payload);
+	static UE_API void SendGameplayEventToActor(AActor* Actor, UPARAM(meta=(GameplayTagFilter="GameplayEventTagsCategory")) FGameplayTag EventTag, FGameplayEventData Payload);
 	
 	// -------------------------------------------------------------------------------
 	//		Gameplay Tag
@@ -95,7 +97,7 @@ class GAMEPLAYABILITIES_API UAbilitySystemBlueprintLibrary : public UBlueprintFu
 	 * @return                                  FGameplayTagChangedEventWrapperSpecHandle Handle by which this binding request can be unbound.
 	 */
 	UFUNCTION(BlueprintCallable, Category = "Gameplay Tags", meta = (AdvancedDisplay = "bExecuteImmediatelyIfTagApplied, TagListeningPolicy"))
-	static FGameplayTagChangedEventWrapperSpecHandle BindEventWrapperToGameplayTagChanged(
+	static UE_API FGameplayTagChangedEventWrapperSpecHandle BindEventWrapperToGameplayTagChanged(
 		UAbilitySystemComponent* AbilitySystemComponent,
 		FGameplayTag Tag, 
 		FOnGameplayTagChangedEventWrapperSignature GameplayTagChangedEventWrapperDelegate, 
@@ -113,7 +115,7 @@ class GAMEPLAYABILITIES_API UAbilitySystemBlueprintLibrary : public UBlueprintFu
 	 * @return                                  FGameplayTagChangedEventWrapperSpecHandle Handle by which this binding request can be unbound.
 	 */
 	UFUNCTION(BlueprintCallable, Category = "Gameplay Tags", meta = (AdvancedDisplay = "bExecuteImmediatelyIfTagApplied, TagListeningPolicy"))
-	static FGameplayTagChangedEventWrapperSpecHandle BindEventWrapperToAnyOfGameplayTagsChanged(
+	static UE_API FGameplayTagChangedEventWrapperSpecHandle BindEventWrapperToAnyOfGameplayTagsChanged(
 		UAbilitySystemComponent* AbilitySystemComponent,
 		const TArray<FGameplayTag>& Tags,
 		FOnGameplayTagChangedEventWrapperSignature GameplayTagChangedEventWrapperDelegate, 
@@ -131,7 +133,7 @@ class GAMEPLAYABILITIES_API UAbilitySystemBlueprintLibrary : public UBlueprintFu
 	 * @return                                  FGameplayTagChangedEventWrapperSpecHandle Handle by which this binding request can be unbound.
 	 */
 	UFUNCTION(BlueprintCallable, Category = "Gameplay Tags", meta = (AdvancedDisplay = "bExecuteImmediatelyIfTagApplied, TagListeningPolicy"))
-	static FGameplayTagChangedEventWrapperSpecHandle BindEventWrapperToAnyOfGameplayTagContainerChanged(
+	static UE_API FGameplayTagChangedEventWrapperSpecHandle BindEventWrapperToAnyOfGameplayTagContainerChanged(
 		UAbilitySystemComponent* AbilitySystemComponent,
 		const FGameplayTagContainer TagContainer, 
 		FOnGameplayTagChangedEventWrapperSignature GameplayTagChangedEventWrapperDelegate, 
@@ -144,7 +146,7 @@ class GAMEPLAYABILITIES_API UAbilitySystemBlueprintLibrary : public UBlueprintFu
 	 * @param Handle     FGameplayTagChangedEventWrapperSpecHandle Handle provided when binding to a delegate
 	 */
 	UFUNCTION(BlueprintCallable, Category = "Gameplay Tags")
-	static void UnbindAllGameplayTagChangedEventWrappersForHandle(FGameplayTagChangedEventWrapperSpecHandle Handle);
+	static UE_API void UnbindAllGameplayTagChangedEventWrappersForHandle(FGameplayTagChangedEventWrapperSpecHandle Handle);
 
 	/**
 	 * Unbinds the event wrapper tag change event bound via a BindEventWrapper<to gameplay tag(s)> method that is tied to the given Handle, for the specific tag.
@@ -154,12 +156,12 @@ class GAMEPLAYABILITIES_API UAbilitySystemBlueprintLibrary : public UBlueprintFu
 	 * @param Handle     int Handle provided when binding to a delegate
 	 */
 	UFUNCTION(BlueprintCallable, Category = "Gameplay Tags")
-	static void UnbindGameplayTagChangedEventWrapperForHandle(FGameplayTag Tag, FGameplayTagChangedEventWrapperSpecHandle Handle);
+	static UE_API void UnbindGameplayTagChangedEventWrapperForHandle(FGameplayTag Tag, FGameplayTagChangedEventWrapperSpecHandle Handle);
 
 protected:
 
 	// Helper fn to process gameplay tag changed event wrappers.
-	static void ProcessGameplayTagChangedEventWrapper(const FGameplayTag GameplayTag, int32 GameplayTagCount, FOnGameplayTagChangedEventWrapperSignature GameplayTagChangedEventWrapperDelegate);
+	static UE_API void ProcessGameplayTagChangedEventWrapper(const FGameplayTag GameplayTag, int32 GameplayTagCount, FOnGameplayTagChangedEventWrapperSignature GameplayTagChangedEventWrapperDelegate);
 
 public:
 
@@ -169,43 +171,43 @@ public:
 
 	/** Returns true if the attribute actually exists */
 	UFUNCTION(BlueprintPure, Category = "Ability|Attribute")
-	static bool IsValid(FGameplayAttribute Attribute);
+	static UE_API bool IsValid(FGameplayAttribute Attribute);
 
 	/** Returns the value of Attribute from the ability system component belonging to Actor. */
 	UFUNCTION(BlueprintPure, Category = "Ability|Attribute")
-	static float GetFloatAttribute(const AActor* Actor, FGameplayAttribute Attribute, bool& bSuccessfullyFoundAttribute);
+	static UE_API float GetFloatAttribute(const AActor* Actor, FGameplayAttribute Attribute, bool& bSuccessfullyFoundAttribute);
 
 	/** Returns the value of Attribute from the ability system component AbilitySystem. */
 	UFUNCTION(BlueprintPure, Category = "Ability|Attribute")
-	static float GetFloatAttributeFromAbilitySystemComponent(const UAbilitySystemComponent* AbilitySystem, FGameplayAttribute Attribute, bool& bSuccessfullyFoundAttribute);
+	static UE_API float GetFloatAttributeFromAbilitySystemComponent(const UAbilitySystemComponent* AbilitySystem, FGameplayAttribute Attribute, bool& bSuccessfullyFoundAttribute);
 
 	/** Returns the base value of Attribute from the ability system component belonging to Actor. */
 	UFUNCTION(BlueprintPure, Category = "Ability|Attribute")
-	static float GetFloatAttributeBase(const AActor* Actor, FGameplayAttribute Attribute, bool& bSuccessfullyFoundAttribute);
+	static UE_API float GetFloatAttributeBase(const AActor* Actor, FGameplayAttribute Attribute, bool& bSuccessfullyFoundAttribute);
 
 	/** Returns the base value of Attribute from the ability system component AbilitySystemComponent. */
 	UFUNCTION(BlueprintPure, Category = "Ability|Attribute")
-	static float GetFloatAttributeBaseFromAbilitySystemComponent(const UAbilitySystemComponent* AbilitySystemComponent, FGameplayAttribute Attribute, bool& bSuccessfullyFoundAttribute);
+	static UE_API float GetFloatAttributeBaseFromAbilitySystemComponent(const UAbilitySystemComponent* AbilitySystemComponent, FGameplayAttribute Attribute, bool& bSuccessfullyFoundAttribute);
 
 	/** Returns the value of Attribute from the ability system component AbilitySystem after evaluating it with source and target tags. bSuccess indicates the success or failure of this operation. */
 	UFUNCTION(BlueprintPure, Category = "Ability|Attribute")
-	static float EvaluateAttributeValueWithTags(UAbilitySystemComponent* AbilitySystem, FGameplayAttribute Attribute, const FGameplayTagContainer& SourceTags, const FGameplayTagContainer& TargetTags, bool& bSuccess);
+	static UE_API float EvaluateAttributeValueWithTags(UAbilitySystemComponent* AbilitySystem, FGameplayAttribute Attribute, const FGameplayTagContainer& SourceTags, const FGameplayTagContainer& TargetTags, bool& bSuccess);
 
 	/** Returns the value of Attribute from the ability system component AbilitySystem after evaluating it with source and target tags using the passed in base value instead of the real base value. bSuccess indicates the success or failure of this operation. */
 	UFUNCTION(BlueprintPure, Category = "Ability|Attribute")
-	static float EvaluateAttributeValueWithTagsAndBase(UAbilitySystemComponent* AbilitySystem, FGameplayAttribute Attribute, const FGameplayTagContainer& SourceTags, const FGameplayTagContainer& TargetTags, float BaseValue, bool& bSuccess);
+	static UE_API float EvaluateAttributeValueWithTagsAndBase(UAbilitySystemComponent* AbilitySystem, FGameplayAttribute Attribute, const FGameplayTagContainer& SourceTags, const FGameplayTagContainer& TargetTags, float BaseValue, bool& bSuccess);
 
 	/** Simple equality operator for gameplay attributes */
 	UFUNCTION(BlueprintPure, meta=(DisplayName = "Equal (Gameplay Attribute)", CompactNodeTitle = "==", Keywords = "== equal"), Category="Ability|Attribute")
-	static bool EqualEqual_GameplayAttributeGameplayAttribute(FGameplayAttribute AttributeA, FGameplayAttribute AttributeB);
+	static UE_API bool EqualEqual_GameplayAttributeGameplayAttribute(FGameplayAttribute AttributeA, FGameplayAttribute AttributeB);
 
 	/** Simple inequality operator for gameplay attributes */
 	UFUNCTION(BlueprintPure, meta=(DisplayName = "Not Equal (Gameplay Attribute)", CompactNodeTitle = "!=", Keywords = "!= not equal"), Category="Ability|Attribute")
-	static bool NotEqual_GameplayAttributeGameplayAttribute(FGameplayAttribute AttributeA, FGameplayAttribute AttributeB);
+	static UE_API bool NotEqual_GameplayAttributeGameplayAttribute(FGameplayAttribute AttributeA, FGameplayAttribute AttributeB);
 
 	/** Returns FString representation of a gameplay attribute's set class and name, in the form of AttrSetName.AttrName (or just AttrName if not part of a set).  */
 	UFUNCTION(BlueprintPure, Category="Ability|Attribute")
-	static FString GetDebugStringFromGameplayAttribute(const FGameplayAttribute& Attribute);
+	static UE_API FString GetDebugStringFromGameplayAttribute(const FGameplayAttribute& Attribute);
 
 	// -------------------------------------------------------------------------------
 	//		TargetData
@@ -213,92 +215,92 @@ public:
 
 	/** Copies targets from HandleToAdd to TargetHandle */
 	UFUNCTION(BlueprintCallable, Category = "Ability|TargetData")
-	static FGameplayAbilityTargetDataHandle AppendTargetDataHandle(FGameplayAbilityTargetDataHandle TargetHandle, const FGameplayAbilityTargetDataHandle& HandleToAdd);
+	static UE_API FGameplayAbilityTargetDataHandle AppendTargetDataHandle(FGameplayAbilityTargetDataHandle TargetHandle, const FGameplayAbilityTargetDataHandle& HandleToAdd);
 
 	/** Creates a target data with a source and destination location */
 	UFUNCTION(BlueprintPure, Category = "Ability|TargetData")
-	static FGameplayAbilityTargetDataHandle	AbilityTargetDataFromLocations(const FGameplayAbilityTargetingLocationInfo& SourceLocation, const FGameplayAbilityTargetingLocationInfo& TargetLocation);
+	static UE_API FGameplayAbilityTargetDataHandle	AbilityTargetDataFromLocations(const FGameplayAbilityTargetingLocationInfo& SourceLocation, const FGameplayAbilityTargetingLocationInfo& TargetLocation);
 
 	/** Creates a target data with a single hit result */
 	UFUNCTION(BlueprintPure, Category = "Ability|TargetData")
-	static FGameplayAbilityTargetDataHandle	AbilityTargetDataFromHitResult(const FHitResult& HitResult);
+	static UE_API FGameplayAbilityTargetDataHandle	AbilityTargetDataFromHitResult(const FHitResult& HitResult);
 
 	/** Returns number of target data objects, not necessarily number of distinct targets */
 	UFUNCTION(BlueprintPure, Category = "Ability|TargetData")
-	static int32 GetDataCountFromTargetData(const FGameplayAbilityTargetDataHandle& TargetData);
+	static UE_API int32 GetDataCountFromTargetData(const FGameplayAbilityTargetDataHandle& TargetData);
 
 	/** Creates single actor target data */
 	UFUNCTION(BlueprintPure, Category = "Ability|TargetData")
-	static FGameplayAbilityTargetDataHandle	AbilityTargetDataFromActor(AActor* Actor);
+	static UE_API FGameplayAbilityTargetDataHandle	AbilityTargetDataFromActor(AActor* Actor);
 
 	/** Creates actor array target data */
 	UFUNCTION(BlueprintPure, Category = "Ability|TargetData")
-	static FGameplayAbilityTargetDataHandle	AbilityTargetDataFromActorArray(const TArray<AActor*>& ActorArray, bool OneTargetPerHandle);
+	static UE_API FGameplayAbilityTargetDataHandle	AbilityTargetDataFromActorArray(const TArray<AActor*>& ActorArray, bool OneTargetPerHandle);
 
 	/** Create a new target data handle with filtration performed on the data */
 	UFUNCTION(BlueprintPure, Category = "Ability|TargetData")
-	static FGameplayAbilityTargetDataHandle	FilterTargetData(const FGameplayAbilityTargetDataHandle& TargetDataHandle, FGameplayTargetDataFilterHandle ActorFilterClass);
+	static UE_API FGameplayAbilityTargetDataHandle	FilterTargetData(const FGameplayAbilityTargetDataHandle& TargetDataHandle, FGameplayTargetDataFilterHandle ActorFilterClass);
 
 	/** Create a handle for filtering target data, filling out all fields */
 	UFUNCTION(BlueprintPure, Category = "Filter")
-	static FGameplayTargetDataFilterHandle MakeFilterHandle(FGameplayTargetDataFilter Filter, AActor* FilterActor);
+	static UE_API FGameplayTargetDataFilterHandle MakeFilterHandle(FGameplayTargetDataFilter Filter, AActor* FilterActor);
 
 	/** Create a spec handle, filling out all fields */
 	UE_DEPRECATED(5.5, "Use MakeSpecHandleByClass. It's safer as InGameplayEffect needs to be a CDO")
 	UFUNCTION(BlueprintPure, Category = "Spec", meta=(DeprecatedFunction, DeprecatedMessage="Use Make Spec Handle (By Class)"))
-	static FGameplayEffectSpecHandle MakeSpecHandle(UGameplayEffect* InGameplayEffect, AActor* InInstigator, AActor* InEffectCauser, float InLevel = 1.0f);
+	static UE_API FGameplayEffectSpecHandle MakeSpecHandle(UGameplayEffect* InGameplayEffect, AActor* InInstigator, AActor* InEffectCauser, float InLevel = 1.0f);
 
 	/** Create a spec handle, filling out all fields */
 	UFUNCTION(BlueprintCallable, Category = "Spec", meta=(DisplayName="Make Spec Handle (By Class)"))
-	static FGameplayEffectSpecHandle MakeSpecHandleByClass(TSubclassOf<UGameplayEffect> GameplayEffect, AActor* Instigator, AActor* EffectCauser, float Level = 1.0f);
+	static UE_API FGameplayEffectSpecHandle MakeSpecHandleByClass(TSubclassOf<UGameplayEffect> GameplayEffect, AActor* Instigator, AActor* EffectCauser, float Level = 1.0f);
 
 	/** Create a spec handle, cloning another */
 	UFUNCTION(BlueprintPure, Category = "Spec")
-	static FGameplayEffectSpecHandle CloneSpecHandle(AActor* InNewInstigator, AActor* InEffectCauser, FGameplayEffectSpecHandle GameplayEffectSpecHandle_Clone);
+	static UE_API FGameplayEffectSpecHandle CloneSpecHandle(AActor* InNewInstigator, AActor* InEffectCauser, FGameplayEffectSpecHandle GameplayEffectSpecHandle_Clone);
 
 	/** Returns all actors targeted, for a given index */
 	UFUNCTION(BlueprintPure, Category = "Ability|TargetData")
-	static TArray<AActor*> GetActorsFromTargetData(const FGameplayAbilityTargetDataHandle& TargetData, int32 Index);
+	static UE_API TArray<AActor*> GetActorsFromTargetData(const FGameplayAbilityTargetDataHandle& TargetData, int32 Index);
 
 	/** Returns all actors targeted */
 	UFUNCTION(BlueprintPure, Category = "Ability|TargetData")
-	static TArray<AActor*> GetAllActorsFromTargetData(const FGameplayAbilityTargetDataHandle& TargetData);
+	static UE_API TArray<AActor*> GetAllActorsFromTargetData(const FGameplayAbilityTargetDataHandle& TargetData);
 
 	/** Returns true if the given TargetData has the actor passed in targeted */
 	UFUNCTION(BlueprintPure, Category = "Ability|TargetData")
-	static bool DoesTargetDataContainActor(const FGameplayAbilityTargetDataHandle& TargetData, int32 Index, AActor* Actor);
+	static UE_API bool DoesTargetDataContainActor(const FGameplayAbilityTargetDataHandle& TargetData, int32 Index, AActor* Actor);
 
 	/** Returns true if the given TargetData has at least 1 actor targeted */
 	UFUNCTION(BlueprintPure, Category = "Ability|TargetData")
-	static bool TargetDataHasActor(const FGameplayAbilityTargetDataHandle& TargetData, int32 Index);
+	static UE_API bool TargetDataHasActor(const FGameplayAbilityTargetDataHandle& TargetData, int32 Index);
 
 	/** Returns true if the target data has a hit result */
 	UFUNCTION(BlueprintPure, Category = "Ability|TargetData")
-	static bool TargetDataHasHitResult(const FGameplayAbilityTargetDataHandle& HitResult, int32 Index);
+	static UE_API bool TargetDataHasHitResult(const FGameplayAbilityTargetDataHandle& HitResult, int32 Index);
 
 	/** Returns the hit result for a given index if it exists */
 	UFUNCTION(BlueprintPure, Category = "Ability|TargetData")
-	static FHitResult GetHitResultFromTargetData(const FGameplayAbilityTargetDataHandle& HitResult, int32 Index);
+	static UE_API FHitResult GetHitResultFromTargetData(const FGameplayAbilityTargetDataHandle& HitResult, int32 Index);
 
 	/** Returns true if the target data has an origin */
 	UFUNCTION(BlueprintPure, Category = "Ability|TargetData")
-	static bool TargetDataHasOrigin(const FGameplayAbilityTargetDataHandle& TargetData, int32 Index);
+	static UE_API bool TargetDataHasOrigin(const FGameplayAbilityTargetDataHandle& TargetData, int32 Index);
 
 	/** Returns the origin for a given index if it exists */
 	UFUNCTION(BlueprintPure, Category = "Ability|TargetData")
-	static FTransform GetTargetDataOrigin(const FGameplayAbilityTargetDataHandle& TargetData, int32 Index);
+	static UE_API FTransform GetTargetDataOrigin(const FGameplayAbilityTargetDataHandle& TargetData, int32 Index);
 
 	/** Returns true if the target data has an end point */
 	UFUNCTION(BlueprintPure, Category = "Ability|TargetData")
-	static bool TargetDataHasEndPoint(const FGameplayAbilityTargetDataHandle& TargetData, int32 Index);
+	static UE_API bool TargetDataHasEndPoint(const FGameplayAbilityTargetDataHandle& TargetData, int32 Index);
 
 	/** Returns the end point for a given index if it exists */
 	UFUNCTION(BlueprintPure, Category = "Ability|TargetData")
-	static FVector GetTargetDataEndPoint(const FGameplayAbilityTargetDataHandle& TargetData, int32 Index);
+	static UE_API FVector GetTargetDataEndPoint(const FGameplayAbilityTargetDataHandle& TargetData, int32 Index);
 
 	/** Returns the end point transform for a given index if it exists */
 	UFUNCTION(BlueprintPure, Category = "Ability|TargetData")
-	static FTransform GetTargetDataEndPointTransform(const FGameplayAbilityTargetDataHandle& TargetData, int32 Index);
+	static UE_API FTransform GetTargetDataEndPointTransform(const FGameplayAbilityTargetDataHandle& TargetData, int32 Index);
 
 	// -------------------------------------------------------------------------------
 	//		GameplayEffectContext
@@ -306,47 +308,47 @@ public:
 
 	/** Returns true if this context has ever been initialized */
 	UFUNCTION(BlueprintPure, Category = "Ability|EffectContext", Meta = (DisplayName = "IsValid"))
-	static bool EffectContextIsValid(FGameplayEffectContextHandle EffectContext);
+	static UE_API bool EffectContextIsValid(FGameplayEffectContextHandle EffectContext);
 
 	/** Returns true if the ability system component that instigated this is locally controlled */
 	UFUNCTION(BlueprintPure, Category = "Ability|EffectContext", Meta = (DisplayName = "IsInstigatorLocallyControlled"))
-	static bool EffectContextIsInstigatorLocallyControlled(FGameplayEffectContextHandle EffectContext);
+	static UE_API bool EffectContextIsInstigatorLocallyControlled(FGameplayEffectContextHandle EffectContext);
 
 	/** Extracts a hit result from the effect context if it is set */
 	UFUNCTION(BlueprintPure, Category = "Ability|EffectContext", Meta = (DisplayName = "GetHitResult"))
-	static FHitResult EffectContextGetHitResult(FGameplayEffectContextHandle EffectContext);
+	static UE_API FHitResult EffectContextGetHitResult(FGameplayEffectContextHandle EffectContext);
 
 	/** Returns true if there is a valid hit result inside the effect context */
 	UFUNCTION(BlueprintPure, Category = "Ability|EffectContext", Meta = (DisplayName = "HasHitResult"))
-	static bool EffectContextHasHitResult(FGameplayEffectContextHandle EffectContext);
+	static UE_API bool EffectContextHasHitResult(FGameplayEffectContextHandle EffectContext);
 
 	/** Adds a hit result to the effect context */
 	UFUNCTION(BlueprintCallable, Category = "Ability|EffectContext", Meta = (DisplayName = "AddHitResult"))
-	static void EffectContextAddHitResult(FGameplayEffectContextHandle EffectContext, FHitResult HitResult, bool bReset);
+	static UE_API void EffectContextAddHitResult(FGameplayEffectContextHandle EffectContext, FHitResult HitResult, bool bReset);
 
 	/** Gets the location the effect originated from */
 	UFUNCTION(BlueprintPure, Category = "Ability|EffectContext", Meta = (DisplayName = "GetOrigin"))
-	static FVector EffectContextGetOrigin(FGameplayEffectContextHandle EffectContext);
+	static UE_API FVector EffectContextGetOrigin(FGameplayEffectContextHandle EffectContext);
 
 	/** Sets the location the effect originated from */
 	UFUNCTION(BlueprintCallable, Category = "Ability|EffectContext", Meta = (DisplayName = "SetOrigin"))
-	static void EffectContextSetOrigin(FGameplayEffectContextHandle EffectContext, FVector Origin);
+	static UE_API void EffectContextSetOrigin(FGameplayEffectContextHandle EffectContext, FVector Origin);
 
 	/** Gets the instigating actor (that holds the ability system component) of the EffectContext */
 	UFUNCTION(BlueprintPure, Category = "Ability|EffectContext", Meta = (DisplayName = "GetInstigatorActor"))
-	static AActor* EffectContextGetInstigatorActor(FGameplayEffectContextHandle EffectContext);
+	static UE_API AActor* EffectContextGetInstigatorActor(FGameplayEffectContextHandle EffectContext);
 
 	/** Gets the original instigator actor that started the chain of events to cause this effect */
 	UFUNCTION(BlueprintPure, Category = "Ability|EffectContext", Meta = (DisplayName = "GetOriginalInstigatorActor"))
-	static AActor* EffectContextGetOriginalInstigatorActor(FGameplayEffectContextHandle EffectContext);
+	static UE_API AActor* EffectContextGetOriginalInstigatorActor(FGameplayEffectContextHandle EffectContext);
 
 	/** Gets the physical actor that caused the effect, possibly a projectile or weapon */
 	UFUNCTION(BlueprintPure, Category = "Ability|EffectContext", Meta = (DisplayName = "GetEffectCauser"))
-	static AActor* EffectContextGetEffectCauser(FGameplayEffectContextHandle EffectContext);
+	static UE_API AActor* EffectContextGetEffectCauser(FGameplayEffectContextHandle EffectContext);
 
 	/** Gets the source object of the effect. */
 	UFUNCTION(BlueprintPure, Category = "Ability|EffectContext", Meta = (DisplayName = "GetSourceObject"))
-	static UObject* EffectContextGetSourceObject(FGameplayEffectContextHandle EffectContext);
+	static UE_API UObject* EffectContextGetSourceObject(FGameplayEffectContextHandle EffectContext);
 
 	// -------------------------------------------------------------------------------
 	//		GameplayCue
@@ -354,63 +356,63 @@ public:
 
 	/** Returns true if the ability system component that spawned this cue is locally controlled */
 	UFUNCTION(BlueprintPure, Category="Ability|GameplayCue")
-	static bool IsInstigatorLocallyControlled(FGameplayCueParameters Parameters);
+	static UE_API bool IsInstigatorLocallyControlled(FGameplayCueParameters Parameters);
 
 	/** Returns true if the ability system component that spawned this cue is locally controlled and a player */
 	UFUNCTION(BlueprintPure, Category="Ability|GameplayCue")
-	static bool IsInstigatorLocallyControlledPlayer(FGameplayCueParameters Parameters);
+	static UE_API bool IsInstigatorLocallyControlledPlayer(FGameplayCueParameters Parameters);
 
 	/** Returns number of actors stored in the Effect Context used by this cue */
 	UFUNCTION(BlueprintPure, Category = "Ability|GameplayCue")
-	static int32 GetActorCount(FGameplayCueParameters Parameters);
+	static UE_API int32 GetActorCount(FGameplayCueParameters Parameters);
 
 	/** Returns actor stored in the Effect Context used by this cue */
 	UFUNCTION(BlueprintPure, Category = "Ability|GameplayCue")
-	static AActor* GetActorByIndex(FGameplayCueParameters Parameters, int32 Index);
+	static UE_API AActor* GetActorByIndex(FGameplayCueParameters Parameters, int32 Index);
 
 	/** Returns a hit result stored in the effect context if valid */
 	UFUNCTION(BlueprintPure, Category = "Ability|GameplayCue")
-	static FHitResult GetHitResult(FGameplayCueParameters Parameters);
+	static UE_API FHitResult GetHitResult(FGameplayCueParameters Parameters);
 
 	/** Checks if the effect context has a hit reslt stored inside */
 	UFUNCTION(BlueprintPure, Category = "Ability|GameplayCue")
-	static bool HasHitResult(FGameplayCueParameters Parameters);
+	static UE_API bool HasHitResult(FGameplayCueParameters Parameters);
 
 	/** Forwards the gameplay cue to another gameplay cue interface object */
 	UFUNCTION(BlueprintCallable, Category = "Ability|GameplayCue")
-	static void ForwardGameplayCueToTarget(TScriptInterface<IGameplayCueInterface> TargetCueInterface, EGameplayCueEvent::Type EventType, FGameplayCueParameters Parameters);
+	static UE_API void ForwardGameplayCueToTarget(TScriptInterface<IGameplayCueInterface> TargetCueInterface, EGameplayCueEvent::Type EventType, FGameplayCueParameters Parameters);
 
 	/** Gets the instigating actor (that holds the ability system component) of the GameplayCue */
 	UFUNCTION(BlueprintPure, Category = "Ability|GameplayCue")
-	static AActor* GetInstigatorActor(FGameplayCueParameters Parameters);
+	static UE_API AActor* GetInstigatorActor(FGameplayCueParameters Parameters);
 
 	/** Gets instigating world location */
 	UFUNCTION(BlueprintPure, Category = "Ability|GameplayCue")
-	static FTransform GetInstigatorTransform(FGameplayCueParameters Parameters);
+	static UE_API FTransform GetInstigatorTransform(FGameplayCueParameters Parameters);
 
 	/** Gets instigating world location */
 	UFUNCTION(BlueprintPure, Category = "Ability|GameplayCue")
-	static FVector GetOrigin(FGameplayCueParameters Parameters);
+	static UE_API FVector GetOrigin(FGameplayCueParameters Parameters);
 
 	/** Gets the best end location and normal for this gameplay cue. If there is hit result data, it will return this. Otherwise it will return the target actor's location/rotation. If none of this is available, it will return false. */
 	UFUNCTION(BlueprintPure, Category = "Ability|GameplayCue")
-	static bool GetGameplayCueEndLocationAndNormal(AActor* TargetActor, FGameplayCueParameters Parameters, FVector& Location, FVector& Normal);
+	static UE_API bool GetGameplayCueEndLocationAndNormal(AActor* TargetActor, FGameplayCueParameters Parameters, FVector& Location, FVector& Normal);
 
 	/** Gets the best normalized effect direction for this gameplay cue. This is useful for effects that require the direction of an enemy attack. Returns true if a valid direction could be calculated. */
 	UFUNCTION(BlueprintPure, Category = "Ability|GameplayCue")
-	static bool GetGameplayCueDirection(AActor* TargetActor, FGameplayCueParameters Parameters, FVector& Direction);
+	static UE_API bool GetGameplayCueDirection(AActor* TargetActor, FGameplayCueParameters Parameters, FVector& Direction);
 
 	/** Returns true if the aggregated source and target tags from the effect spec meets the tag requirements */
 	UFUNCTION(BlueprintPure, Category = "Ability|GameplayCue")
-	static bool DoesGameplayCueMeetTagRequirements(FGameplayCueParameters Parameters, const FGameplayTagRequirements& SourceTagReqs, const FGameplayTagRequirements& TargetTagReqs);
+	static UE_API bool DoesGameplayCueMeetTagRequirements(FGameplayCueParameters Parameters, const FGameplayTagRequirements& SourceTagReqs, const FGameplayTagRequirements& TargetTagReqs);
 
 	/** Native make, to avoid having to deal with quantized vector types */
 	UFUNCTION(BlueprintPure, Category = "Ability|GameplayCue", meta = (NativeMakeFunc, AdvancedDisplay=5, Location="0,0,0", Normal= "0,0,0", GameplayEffectLevel = "1", AbilityLevel = "1"))
-	static FGameplayCueParameters MakeGameplayCueParameters(float NormalizedMagnitude, float RawMagnitude, FGameplayEffectContextHandle EffectContext, FGameplayTag MatchedTagName, FGameplayTag OriginalTag, FGameplayTagContainer AggregatedSourceTags, FGameplayTagContainer AggregatedTargetTags, FVector Location, FVector Normal, AActor* Instigator, AActor* EffectCauser, UObject* SourceObject, UPhysicalMaterial* PhysicalMaterial, int32 GameplayEffectLevel, int32 AbilityLevel, USceneComponent* TargetAttachComponent, bool bReplicateLocationWhenUsingMinimalRepProxy);
+	static UE_API FGameplayCueParameters MakeGameplayCueParameters(float NormalizedMagnitude, float RawMagnitude, FGameplayEffectContextHandle EffectContext, FGameplayTag MatchedTagName, FGameplayTag OriginalTag, FGameplayTagContainer AggregatedSourceTags, FGameplayTagContainer AggregatedTargetTags, FVector Location, FVector Normal, AActor* Instigator, AActor* EffectCauser, UObject* SourceObject, UPhysicalMaterial* PhysicalMaterial, int32 GameplayEffectLevel, int32 AbilityLevel, USceneComponent* TargetAttachComponent, bool bReplicateLocationWhenUsingMinimalRepProxy);
 
 	/** Native break, to avoid having to deal with quantized vector types */
 	UFUNCTION(BlueprintPure, Category = "Ability|GameplayCue", meta = (NativeBreakFunc, AdvancedDisplay=6))
-	static void BreakGameplayCueParameters(const struct FGameplayCueParameters& Parameters, float& NormalizedMagnitude, float& RawMagnitude, FGameplayEffectContextHandle& EffectContext, FGameplayTag& MatchedTagName, FGameplayTag& OriginalTag, FGameplayTagContainer& AggregatedSourceTags, FGameplayTagContainer& AggregatedTargetTags, FVector& Location, FVector& Normal, AActor*& Instigator, AActor*& EffectCauser, UObject*& SourceObject, UPhysicalMaterial*& PhysicalMaterial, int32& GameplayEffectLevel, int32& AbilityLevel, USceneComponent*& TargetAttachComponent, bool& bReplicateLocationWhenUsingMinimalRepProxy);
+	static UE_API void BreakGameplayCueParameters(const struct FGameplayCueParameters& Parameters, float& NormalizedMagnitude, float& RawMagnitude, FGameplayEffectContextHandle& EffectContext, FGameplayTag& MatchedTagName, FGameplayTag& OriginalTag, FGameplayTagContainer& AggregatedSourceTags, FGameplayTagContainer& AggregatedTargetTags, FVector& Location, FVector& Normal, AActor*& Instigator, AActor*& EffectCauser, UObject*& SourceObject, UPhysicalMaterial*& PhysicalMaterial, int32& GameplayEffectLevel, int32& AbilityLevel, USceneComponent*& TargetAttachComponent, bool& bReplicateLocationWhenUsingMinimalRepProxy);
 
 	// -------------------------------------------------------------------------------
 	//		GameplayEffectSpec
@@ -418,66 +420,66 @@ public:
 
 	/** Sets a raw name Set By Caller magnitude value, the tag version should normally be used */
 	UFUNCTION(BlueprintCallable, Category = "Ability|GameplayEffect")
-	static FGameplayEffectSpecHandle AssignSetByCallerMagnitude(FGameplayEffectSpecHandle SpecHandle, FName DataName, float Magnitude);
+	static UE_API FGameplayEffectSpecHandle AssignSetByCallerMagnitude(FGameplayEffectSpecHandle SpecHandle, FName DataName, float Magnitude);
 
 	/** Sets a gameplay tag Set By Caller magnitude value */
 	UFUNCTION(BlueprintCallable, Category = "Ability|GameplayEffect", meta = (GameplayTagFilter = "SetByCaller"))
-	static FGameplayEffectSpecHandle AssignTagSetByCallerMagnitude(FGameplayEffectSpecHandle SpecHandle, FGameplayTag DataTag, float Magnitude);
+	static UE_API FGameplayEffectSpecHandle AssignTagSetByCallerMagnitude(FGameplayEffectSpecHandle SpecHandle, FGameplayTag DataTag, float Magnitude);
 
 	/** Manually sets the duration on a specific effect */
 	UFUNCTION(BlueprintCallable, Category = "Ability|GameplayEffect")
-	static FGameplayEffectSpecHandle SetDuration(FGameplayEffectSpecHandle SpecHandle, float Duration);
+	static UE_API FGameplayEffectSpecHandle SetDuration(FGameplayEffectSpecHandle SpecHandle, float Duration);
 
 	/** This instance of the effect will now grant NewGameplayTag to the object that this effect is applied to */
 	UFUNCTION(BlueprintCallable, Category = "Ability|GameplayEffect")
-	static FGameplayEffectSpecHandle AddGrantedTag(FGameplayEffectSpecHandle SpecHandle, FGameplayTag NewGameplayTag);
+	static UE_API FGameplayEffectSpecHandle AddGrantedTag(FGameplayEffectSpecHandle SpecHandle, FGameplayTag NewGameplayTag);
 
 	/** This instance of the effect will now grant NewGameplayTags to the object that this effect is applied to */
 	UFUNCTION(BlueprintCallable, Category = "Ability|GameplayEffect")
-	static FGameplayEffectSpecHandle AddGrantedTags(FGameplayEffectSpecHandle SpecHandle, FGameplayTagContainer NewGameplayTags);
+	static UE_API FGameplayEffectSpecHandle AddGrantedTags(FGameplayEffectSpecHandle SpecHandle, FGameplayTagContainer NewGameplayTags);
 
 	/** Adds NewGameplayTag to this instance of the effect */
 	UFUNCTION(BlueprintCallable, Category = "Ability|GameplayEffect")
-	static FGameplayEffectSpecHandle AddAssetTag(FGameplayEffectSpecHandle SpecHandle, FGameplayTag NewGameplayTag);
+	static UE_API FGameplayEffectSpecHandle AddAssetTag(FGameplayEffectSpecHandle SpecHandle, FGameplayTag NewGameplayTag);
 
 	/** Adds NewGameplayTags to this instance of the effect */
 	UFUNCTION(BlueprintCallable, Category = "Ability|GameplayEffect")
-	static FGameplayEffectSpecHandle AddAssetTags(FGameplayEffectSpecHandle SpecHandle, FGameplayTagContainer NewGameplayTags);
+	static UE_API FGameplayEffectSpecHandle AddAssetTags(FGameplayEffectSpecHandle SpecHandle, FGameplayTagContainer NewGameplayTags);
 
 	/** Adds LinkedGameplayEffectSpec to SpecHandles. LinkedGameplayEffectSpec will be applied when/if SpecHandle is applied successfully. LinkedGameplayEffectSpec will not be modified here. Returns the ORIGINAL SpecHandle (legacy decision) */
 	UE_DEPRECATED(5.3, "Linked GameplayEffects aren't replicated.  The new UAdditionalGameplayEffectsComponent renders this functionality obsolete (and since it's configured in the Asset, it is properly synced).")
 	UFUNCTION(BlueprintCallable, Category = "Ability|GameplayEffect", meta = (DeprecatedFunction, DeprecatedMessage="Linked GameplayEffects aren't replicated.  Configure the GameplayEffect asset with a suitable GameplayEffectComponent."))
-	static FGameplayEffectSpecHandle AddLinkedGameplayEffectSpec(FGameplayEffectSpecHandle SpecHandle, FGameplayEffectSpecHandle LinkedGameplayEffectSpec);
+	static UE_API FGameplayEffectSpecHandle AddLinkedGameplayEffectSpec(FGameplayEffectSpecHandle SpecHandle, FGameplayEffectSpecHandle LinkedGameplayEffectSpec);
 
 	/** Adds LinkedGameplayEffect to SpecHandles. LinkedGameplayEffectSpec will be applied when/if SpecHandle is applied successfully. This will initialize the LinkedGameplayEffect's Spec for you. Returns to NEW linked spec in case you want to add more to it. */
 	UE_DEPRECATED(5.3, "Linked GameplayEffects aren't replicated.  The new UAdditionalGameplayEffectsComponent renders this functionality obsolete (and since it's configured in the Asset, it is properly synced).")
 	UFUNCTION(BlueprintCallable, Category = "Ability|GameplayEffect", meta = (DeprecatedFunction, DeprecatedMessage = "Linked GameplayEffects aren't replicated.  Configure the GameplayEffect asset with a suitable GameplayEffectComponent."))
-	static FGameplayEffectSpecHandle AddLinkedGameplayEffect(FGameplayEffectSpecHandle SpecHandle, TSubclassOf<UGameplayEffect> LinkedGameplayEffect);
+	static UE_API FGameplayEffectSpecHandle AddLinkedGameplayEffect(FGameplayEffectSpecHandle SpecHandle, TSubclassOf<UGameplayEffect> LinkedGameplayEffect);
 
 	/** Sets the GameplayEffectSpec's StackCount to the specified amount (prior to applying) */
 	UFUNCTION(BlueprintCallable, Category = "Ability|GameplayEffect")
-	static FGameplayEffectSpecHandle SetStackCount(FGameplayEffectSpecHandle SpecHandle, int32 StackCount);
+	static UE_API FGameplayEffectSpecHandle SetStackCount(FGameplayEffectSpecHandle SpecHandle, int32 StackCount);
 
 	/** Sets the GameplayEffectSpec's StackCount to the max stack count defined in the GameplayEffect definition */
 	UFUNCTION(BlueprintCallable, Category = "Ability|GameplayEffect")
-	static FGameplayEffectSpecHandle SetStackCountToMax(FGameplayEffectSpecHandle SpecHandle);
+	static UE_API FGameplayEffectSpecHandle SetStackCountToMax(FGameplayEffectSpecHandle SpecHandle);
 
 	/** Gets the GameplayEffectSpec's effect context handle */
 	UFUNCTION(BlueprintCallable, Category = "Ability|GameplayEffect")
-	static FGameplayEffectContextHandle GetEffectContext(FGameplayEffectSpecHandle SpecHandle);
+	static UE_API FGameplayEffectContextHandle GetEffectContext(FGameplayEffectSpecHandle SpecHandle);
 
 	/** Returns handles for all Linked GE Specs that SpecHandle may apply. Useful if you want to append additional information to them. */
 	UE_DEPRECATED(5.3, "Linked GameplayEffects aren't replicated.  The new UAdditionalGameplayEffectsComponent renders this functionality obsolete (and since it's configured in the Asset, it is properly synced).")
 	UFUNCTION(BlueprintPure, Category = "Ability|GameplayEffect", meta = (DeprecatedFunction, DeprecatedMessage = "Linked GameplayEffects aren't replicated.  Configure the GameplayEffect asset with a suitable GameplayEffectComponent."))
-	static TArray<FGameplayEffectSpecHandle> GetAllLinkedGameplayEffectSpecHandles(FGameplayEffectSpecHandle SpecHandle);
+	static UE_API TArray<FGameplayEffectSpecHandle> GetAllLinkedGameplayEffectSpecHandles(FGameplayEffectSpecHandle SpecHandle);
 
 	/** Manually adds a set of tags to a given actor, and optionally replicates them. */
 	UFUNCTION(BlueprintCallable, Category="Ability|GameplayEffect")
-	static bool AddLooseGameplayTags(AActor* Actor, const FGameplayTagContainer& GameplayTags, bool bShouldReplicate=false);
+	static UE_API bool AddLooseGameplayTags(AActor* Actor, const FGameplayTagContainer& GameplayTags, bool bShouldReplicate=false);
 
 	/** Manually removes a set of tags from a given actor, with optional replication. */
 	UFUNCTION(BlueprintCallable, Category="Ability|GameplayEffect")
-	static bool RemoveLooseGameplayTags(AActor* Actor, const FGameplayTagContainer& GameplayTags, bool bShouldReplicate=false);
+	static UE_API bool RemoveLooseGameplayTags(AActor* Actor, const FGameplayTagContainer& GameplayTags, bool bShouldReplicate=false);
 
 	// -------------------------------------------------------------------------------
 	//		GameplayEffectSpec
@@ -485,10 +487,10 @@ public:
 
 	/** Gets the magnitude of change for an attribute on an APPLIED GameplayEffectSpec. */
 	UFUNCTION(BlueprintCallable, Category = "Ability|GameplayEffect")
-	static float GetModifiedAttributeMagnitude(FGameplayEffectSpecHandle SpecHandle, FGameplayAttribute Attribute);
+	static UE_API float GetModifiedAttributeMagnitude(FGameplayEffectSpecHandle SpecHandle, FGameplayAttribute Attribute);
 
 	/** Helper function that may be useful to call from native as well */
-	static float GetModifiedAttributeMagnitude(const FGameplayEffectSpec& SpecHandle, FGameplayAttribute Attribute);
+	static UE_API float GetModifiedAttributeMagnitude(const FGameplayEffectSpec& SpecHandle, FGameplayAttribute Attribute);
 
 	// -------------------------------------------------------------------------------
 	//		FActiveGameplayEffectHandle
@@ -496,43 +498,43 @@ public:
 	
 	/** Returns current stack count of an active Gameplay Effect. Will return 0 if the GameplayEffect is no longer valid. */
 	UFUNCTION(BlueprintCallable, Category = "Ability|GameplayEffect")
-	static int32 GetActiveGameplayEffectStackCount(FActiveGameplayEffectHandle ActiveHandle);
+	static UE_API int32 GetActiveGameplayEffectStackCount(FActiveGameplayEffectHandle ActiveHandle);
 
 	/** Returns stack limit count of an active Gameplay Effect. Will return 0 if the GameplayEffect is no longer valid. */
 	UFUNCTION(BlueprintCallable, Category = "Ability|GameplayEffect")
-	static int32 GetActiveGameplayEffectStackLimitCount(FActiveGameplayEffectHandle ActiveHandle);
+	static UE_API int32 GetActiveGameplayEffectStackLimitCount(FActiveGameplayEffectHandle ActiveHandle);
 
 	/** Returns the start time (time which the GE was added) for a given GameplayEffect */
 	UFUNCTION(BlueprintCallable, Category = "Ability|GameplayEffect")
-	static float GetActiveGameplayEffectStartTime(FActiveGameplayEffectHandle ActiveHandle);
+	static UE_API float GetActiveGameplayEffectStartTime(FActiveGameplayEffectHandle ActiveHandle);
 
 	/** Returns the expected end time (when we think the GE will expire) for a given GameplayEffect (note someone could remove or change it before that happens!) */
 	UFUNCTION(BlueprintCallable, Category = "Ability|GameplayEffect")
-	static float GetActiveGameplayEffectExpectedEndTime(FActiveGameplayEffectHandle ActiveHandle);
+	static UE_API float GetActiveGameplayEffectExpectedEndTime(FActiveGameplayEffectHandle ActiveHandle);
 
 	/** Returns the total duration for a given GameplayEffect */
 	UFUNCTION(BlueprintCallable, Category = "Ability|GameplayEffect")
-	static float GetActiveGameplayEffectTotalDuration(FActiveGameplayEffectHandle ActiveHandle);
+	static UE_API float GetActiveGameplayEffectTotalDuration(FActiveGameplayEffectHandle ActiveHandle);
 
 	/** Returns the total duration for a given GameplayEffect, basically ExpectedEndTime - Current Time */
 	UFUNCTION(BlueprintCallable, Category = "Ability|GameplayEffect", meta = (WorldContext = "WorldContextObject"))
-	static float GetActiveGameplayEffectRemainingDuration(UObject* WorldContextObject, FActiveGameplayEffectHandle ActiveHandle);
+	static UE_API float GetActiveGameplayEffectRemainingDuration(UObject* WorldContextObject, FActiveGameplayEffectHandle ActiveHandle);
 
 	/** Returns a debug string for display */
 	UFUNCTION(BlueprintPure, Category = "Ability|GameplayEffect", Meta = (DisplayName = "Get Active GameplayEffect Debug String "))
-	static FString GetActiveGameplayEffectDebugString(FActiveGameplayEffectHandle ActiveHandle);
+	static UE_API FString GetActiveGameplayEffectDebugString(FActiveGameplayEffectHandle ActiveHandle);
 
 	/** Returns the UI data for a gameplay effect class (if any) */
 	UFUNCTION(BlueprintCallable, Category = "Ability|GameplayEffect", Meta = (DisplayName = "Get GameplayEffect UI Data", DeterminesOutputType="DataType"))
-	static const UGameplayEffectUIData* GetGameplayEffectUIData(TSubclassOf<UGameplayEffect> EffectClass, TSubclassOf<UGameplayEffectUIData> DataType);
+	static UE_API const UGameplayEffectUIData* GetGameplayEffectUIData(TSubclassOf<UGameplayEffect> EffectClass, TSubclassOf<UGameplayEffectUIData> DataType);
 
 	/** Equality operator for two Active Gameplay Effect Handles */
 	UFUNCTION(BlueprintPure, Category = "Ability|GameplayEffect", meta = (DisplayName = "Equal (Active Gameplay Effect Handle)", CompactNodeTitle = "==", ScriptOperator = "=="))
-	static bool EqualEqual_ActiveGameplayEffectHandle(const FActiveGameplayEffectHandle& A, const FActiveGameplayEffectHandle& B);
+	static UE_API bool EqualEqual_ActiveGameplayEffectHandle(const FActiveGameplayEffectHandle& A, const FActiveGameplayEffectHandle& B);
 
 	/** Inequality operator for two Active Gameplay Effect Handles */
 	UFUNCTION(BlueprintPure, Category = "Ability|GameplayEffect", meta = (DisplayName = "Not Equal (Active Gameplay Effect Handle)", CompactNodeTitle = "!=", ScriptOperator = "!="))
-	static bool NotEqual_ActiveGameplayEffectHandle(const FActiveGameplayEffectHandle& A, const FActiveGameplayEffectHandle& B);
+	static UE_API bool NotEqual_ActiveGameplayEffectHandle(const FActiveGameplayEffectHandle& A, const FActiveGameplayEffectHandle& B);
 
 	/**
 	 * Returns the Gameplay Effect CDO from an active handle.
@@ -540,7 +542,7 @@ public:
 	 * but you can use it to read additional Gameplay Effect info, such as icon, description, etc. 
 	 */
 	UFUNCTION(BlueprintPure, Category = "Ability|GameplayEffect")
-	static const UGameplayEffect* GetGameplayEffectFromActiveEffectHandle(const FActiveGameplayEffectHandle& ActiveHandle);
+	static UE_API const UGameplayEffect* GetGameplayEffectFromActiveEffectHandle(const FActiveGameplayEffectHandle& ActiveHandle);
 
 	// -------------------------------------------------------------------------------
 	//		Gameplay Effect
@@ -548,11 +550,11 @@ public:
 
 	/** Returns all tags that the Gameplay Effect *has* (that denote the GE Asset itself) and *does not* grant to any Actor. */
 	UFUNCTION(BlueprintPure, Category = "Ability|GameplayEffect", meta = (DisplayName = "Get Asset Tags"))
-	static const FGameplayTagContainer& GetGameplayEffectAssetTags(TSubclassOf<UGameplayEffect> EffectClass);
+	static UE_API const FGameplayTagContainer& GetGameplayEffectAssetTags(TSubclassOf<UGameplayEffect> EffectClass);
 
 	/** Returns all tags that the Gameplay Effect grants to the target Actor */
 	UFUNCTION(BlueprintPure, Category = "Ability|GameplayEffect", meta = (DisplayName = "Get Granted Tags"))
-	static const FGameplayTagContainer& GetGameplayEffectGrantedTags(TSubclassOf<UGameplayEffect> EffectClass);
+	static UE_API const FGameplayTagContainer& GetGameplayEffectGrantedTags(TSubclassOf<UGameplayEffect> EffectClass);
 
 	// -------------------------------------------------------------------------------
 	//		GameplayAbility
@@ -568,19 +570,19 @@ public:
 	 * @return Pointer to the Gameplay Ability object
 	 */
 	UFUNCTION(BlueprintPure, Category = "Ability|GameplayAbility")
-	static const UGameplayAbility* GetGameplayAbilityFromSpecHandle(UAbilitySystemComponent* AbilitySystem, const FGameplayAbilitySpecHandle& AbilitySpecHandle, bool& bIsInstance);
+	static UE_API const UGameplayAbility* GetGameplayAbilityFromSpecHandle(UAbilitySystemComponent* AbilitySystem, const FGameplayAbilitySpecHandle& AbilitySpecHandle, bool& bIsInstance);
 
 	/** Returns true if the passed-in Gameplay Ability instance is active (activated and not yet ended). */
 	UFUNCTION(BlueprintPure, Category = "Ability|GameplayAbility", meta=(DisplayName="Is Active",DefaultToSelf=GameplayAbility))
-	static bool IsGameplayAbilityActive(const UGameplayAbility* GameplayAbility);
+	static UE_API bool IsGameplayAbilityActive(const UGameplayAbility* GameplayAbility);
 
 	/** Equality operator for two Gameplay Ability Spec Handles */
 	UFUNCTION(BlueprintPure, Category = "Ability|GameplayAbility", meta = (DisplayName = "Equal (Gameplay Ability Spec Handle)", CompactNodeTitle = "==", ScriptOperator = "=="))
-	static bool EqualEqual_GameplayAbilitySpecHandle(const FGameplayAbilitySpecHandle& A, const FGameplayAbilitySpecHandle& B);
+	static UE_API bool EqualEqual_GameplayAbilitySpecHandle(const FGameplayAbilitySpecHandle& A, const FGameplayAbilitySpecHandle& B);
 
 	/** Inequality operator for two Gameplay Ability Spec Handles */
 	UFUNCTION(BlueprintPure, Category = "Ability|GameplayAbility", meta = (DisplayName = "Not Equal (Gameplay Ability Spec Handle)", CompactNodeTitle = "!=", ScriptOperator = "!="))
-	static bool NotEqual_GameplayAbilitySpecHandle(const FGameplayAbilitySpecHandle& A, const FGameplayAbilitySpecHandle& B);
+	static UE_API bool NotEqual_GameplayAbilitySpecHandle(const FGameplayAbilitySpecHandle& A, const FGameplayAbilitySpecHandle& B);
 
 	/**
 	 * Converts the given Scalable float to the value at the given level.
@@ -590,7 +592,7 @@ public:
 	 * @return The single-precision float value at the given level from the scalable float
 	 */
 	UFUNCTION(BlueprintPure, Category="ScalableFloat", meta=(DisplayName="Get Value At Level (single-precision)", BlueprintAutocast, IgnoreTypePromotion))
-	static float Conv_ScalableFloatToFloat(const FScalableFloat& Input, float Level = 0.0f);
+	static UE_API float Conv_ScalableFloatToFloat(const FScalableFloat& Input, float Level = 0.0f);
 
 	/**
 	 * Converts the given Scalable float to the value at the given level.
@@ -600,5 +602,7 @@ public:
 	 * @return The double-precision float value at the given level from the scalable float
 	 */
 	UFUNCTION(BlueprintPure, Category="ScalableFloat", meta=(DisplayName="Get Value At Level (double-precision)", BlueprintAutocast, IgnoreTypePromotion))
-	static double Conv_ScalableFloatToDouble(const FScalableFloat& Input, float Level = 0.0f);
+	static UE_API double Conv_ScalableFloatToDouble(const FScalableFloat& Input, float Level = 0.0f);
 };
+
+#undef UE_API
