[GameplayAbilities]

diff --git a/Engine/Plugins/Runtime/GameplayAbilities/Source/GameplayAbilities/Public/AbilitySystemComponent.h b/Engine/Plugins/Runtime/GameplayAbilities/Source/GameplayAbilities/Public/AbilitySystemComponent.h
--- a/Engine/Plugins/Runtime/GameplayAbilities/Source/GameplayAbilities/Public/AbilitySystemComponent.h
+++ b/Engine/Plugins/Runtime/GameplayAbilities/Source/GameplayAbilities/Public/AbilitySystemComponent.h
@@ -24,6 +24,8 @@
 #include "Net/Core/PushModel/PushModel.h"
 #include "AbilitySystemComponent.generated.h"
 
+#define UE_API GAMEPLAYABILITIES_API
+
 class AGameplayAbilityTargetActor;
 class AHUD;
 class FDebugDisplayInfo;
@@ -103,8 +105,8 @@ enum class EConsiderPending : uint8
 ENUM_CLASS_FLAGS(EConsiderPending)
 
 /** The core ActorComponent for interfacing with the GameplayAbilities System */
-UCLASS(ClassGroup=AbilitySystem, hidecategories=(Object,LOD,Lighting,Transform,Sockets,TextureStreaming), editinlinenew, meta=(BlueprintSpawnableComponent))
-class GAMEPLAYABILITIES_API UAbilitySystemComponent : public UGameplayTasksComponent, public IGameplayTagAssetInterface, public IAbilitySystemReplicationProxyInterface
+UCLASS(ClassGroup=AbilitySystem, hidecategories=(Object,LOD,Lighting,Transform,Sockets,TextureStreaming), editinlinenew, meta=(BlueprintSpawnableComponent), MinimalAPI)
+class UAbilitySystemComponent : public UGameplayTasksComponent, public IGameplayTagAssetInterface, public IAbilitySystemReplicationProxyInterface
 {
 	GENERATED_UCLASS_BODY()
 
@@ -160,17 +162,17 @@ class GAMEPLAYABILITIES_API UAbilitySystemComponent : public UGameplayTasksCompo
 	 * 
 	 * @return true if Attribute is valid and this ability system component contains an attribute set that contains Attribute. Returns false otherwise.
 	 */
-	bool HasAttributeSetForAttribute(FGameplayAttribute Attribute) const;
+	UE_API bool HasAttributeSetForAttribute(FGameplayAttribute Attribute) const;
 
 	/** Initializes starting attributes from a data table. Not well supported, a gameplay effect with curve table references may be a better solution */
-	const UAttributeSet* InitStats(TSubclassOf<class UAttributeSet> Attributes, const UDataTable* DataTable);
+	UE_API const UAttributeSet* InitStats(TSubclassOf<class UAttributeSet> Attributes, const UDataTable* DataTable);
 
 	UFUNCTION(BlueprintCallable, Category="Skills", meta=(DisplayName="InitStats", ScriptName="InitStats"))
-	void K2_InitStats(TSubclassOf<class UAttributeSet> Attributes, const UDataTable* DataTable);
+	UE_API void K2_InitStats(TSubclassOf<class UAttributeSet> Attributes, const UDataTable* DataTable);
 		
 	/** Returns a list of all attributes for this abilty system component */
 	UFUNCTION(BlueprintPure, Category="Gameplay Attributes")
-	void GetAllAttributes(TArray<FGameplayAttribute>& OutAttributes);
+	UE_API void GetAllAttributes(TArray<FGameplayAttribute>& OutAttributes);
 
 	/**
 	 * Returns a reference to the Attribute Set instance, if one exists in this component
@@ -179,7 +181,7 @@ class GAMEPLAYABILITIES_API UAbilitySystemComponent : public UGameplayTasksCompo
 	 * @param bFound Set to true if an instance of the Attribute Set exists
 	 */
 	UFUNCTION(BlueprintCallable, BlueprintPure = false, Category = "Gameplay Attributes", meta=(DeterminesOutputType = AttributeSetClass))
-	const UAttributeSet* GetAttributeSet(TSubclassOf<UAttributeSet> AttributeSetClass) const;
+	UE_API const UAttributeSet* GetAttributeSet(TSubclassOf<UAttributeSet> AttributeSetClass) const;
 
 	/**
 	 * Returns the current value of the given gameplay attribute, or zero if the attribute is not found.
@@ -189,28 +191,28 @@ class GAMEPLAYABILITIES_API UAbilitySystemComponent : public UGameplayTasksCompo
 	 * @param bFound Set to true if the attribute exists in this component
 	 */
 	UFUNCTION(BlueprintPure, Category = "Gameplay Attributes")
-	float GetGameplayAttributeValue(FGameplayAttribute Attribute, bool& bFound) const;
+	UE_API float GetGameplayAttributeValue(FGameplayAttribute Attribute, bool& bFound) const;
 
 	UPROPERTY(EditAnywhere, Category="AttributeTest")
 	TArray<FAttributeDefaults>	DefaultStartingData;
 
 	/** Remove all current AttributeSets and register the ones in the passed array. Note that it's better to call Add/Remove directly when possible. */
-	void SetSpawnedAttributes(const TArray<UAttributeSet*>& NewAttributeSet);
+	UE_API void SetSpawnedAttributes(const TArray<UAttributeSet*>& NewAttributeSet);
 
 	UE_DEPRECATED(5.1, "This function will be made private. Use Add/Remove SpawnedAttributes instead")
-	TArray<TObjectPtr<UAttributeSet>>& GetSpawnedAttributes_Mutable();
+	UE_API TArray<TObjectPtr<UAttributeSet>>& GetSpawnedAttributes_Mutable();
 
 	/** Access the spawned attributes list when you don't intend to modify the list. */
-	const TArray<UAttributeSet*>& GetSpawnedAttributes() const;
+	UE_API const TArray<UAttributeSet*>& GetSpawnedAttributes() const;
 
 	/** Add a new attribute set */
-	void AddSpawnedAttribute(UAttributeSet* Attribute);
+	UE_API void AddSpawnedAttribute(UAttributeSet* Attribute);
 
 	/** Remove an existing attribute set */
-	void RemoveSpawnedAttribute(UAttributeSet* Attribute);
+	UE_API void RemoveSpawnedAttribute(UAttributeSet* Attribute);
 
 	/** Remove all attribute sets */
-	void RemoveAllSpawnedAttributes();
+	UE_API void RemoveAllSpawnedAttributes();
 
 
 	/** The linked Anim Instance that this component will play montages in. Use NAME_None for the main anim instance. */
@@ -219,10 +221,10 @@ class GAMEPLAYABILITIES_API UAbilitySystemComponent : public UGameplayTasksCompo
 
 
 	/** Sets the base value of an attribute. Existing active modifiers are NOT cleared and will act upon the new base value. */
-	void SetNumericAttributeBase(const FGameplayAttribute &Attribute, float NewBaseValue);
+	UE_API void SetNumericAttributeBase(const FGameplayAttribute &Attribute, float NewBaseValue);
 
 	/** Gets the base value of an attribute. That is, the value of the attribute with no stateful modifiers */
-	float GetNumericAttributeBase(const FGameplayAttribute &Attribute) const;
+	UE_API float GetNumericAttributeBase(const FGameplayAttribute &Attribute) const;
 
 	/**
 	 *	Applies an in-place mod to the given attribute. This correctly update the attribute's aggregator, updates the attribute set property,
@@ -231,36 +233,36 @@ class GAMEPLAYABILITIES_API UAbilitySystemComponent : public UGameplayTasksCompo
 	 *	No GameplayEffectSpec is created or is applied!
 	 *	This should only be used in cases where applying a real GameplayEffectSpec is too slow or not possible.
 	 */
-	void ApplyModToAttribute(const FGameplayAttribute &Attribute, TEnumAsByte<EGameplayModOp::Type> ModifierOp, float ModifierMagnitude);
+	UE_API void ApplyModToAttribute(const FGameplayAttribute &Attribute, TEnumAsByte<EGameplayModOp::Type> ModifierOp, float ModifierMagnitude);
 
 	/**
 	 *  Applies an inplace mod to the given attribute. Unlike ApplyModToAttribute this function will run on the client or server.
 	 *  This may result in problems related to prediction and will not roll back properly.
 	 */
-	void ApplyModToAttributeUnsafe(const FGameplayAttribute &Attribute, TEnumAsByte<EGameplayModOp::Type> ModifierOp, float ModifierMagnitude);
+	UE_API void ApplyModToAttributeUnsafe(const FGameplayAttribute &Attribute, TEnumAsByte<EGameplayModOp::Type> ModifierOp, float ModifierMagnitude);
 
 	/** Returns current (final) value of an attribute */
-	float GetNumericAttribute(const FGameplayAttribute &Attribute) const;
-	float GetNumericAttributeChecked(const FGameplayAttribute &Attribute) const;
+	UE_API float GetNumericAttribute(const FGameplayAttribute &Attribute) const;
+	UE_API float GetNumericAttributeChecked(const FGameplayAttribute &Attribute) const;
 
 	/** Returns an attribute value, after applying tag filters */
-	float GetFilteredAttributeValue(const FGameplayAttribute& Attribute, const FGameplayTagRequirements& SourceTags, const FGameplayTagContainer& TargetTags, const TArray<FActiveGameplayEffectHandle>& HandlesToIgnore = TArray<FActiveGameplayEffectHandle>());
+	UE_API float GetFilteredAttributeValue(const FGameplayAttribute& Attribute, const FGameplayTagRequirements& SourceTags, const FGameplayTagContainer& TargetTags, const TArray<FActiveGameplayEffectHandle>& HandlesToIgnore = TArray<FActiveGameplayEffectHandle>());
 
 	// ----------------------------------------------------------------------------------------------------------------
 	//	Replication
 	// ----------------------------------------------------------------------------------------------------------------
 
 	/** Forces avatar actor to update it's replication. Useful for things like needing to replication for movement / locations reasons. */
-	virtual void ForceAvatarReplication();
+	UE_API virtual void ForceAvatarReplication();
 
 	/** When true, we will not replicate active gameplay effects for this ability system component, so attributes and tags */
-	virtual void SetReplicationMode(EGameplayEffectReplicationMode NewReplicationMode);
+	UE_API virtual void SetReplicationMode(EGameplayEffectReplicationMode NewReplicationMode);
 
 	/** How gameplay effects are replicated */
 	EGameplayEffectReplicationMode ReplicationMode;
 
 	/** Who to route replication through if ReplicationProxyEnabled (if this returns null, when ReplicationProxyEnabled, we wont replicate)  */
-	virtual IAbilitySystemReplicationProxyInterface* GetReplicationInterface();
+	UE_API virtual IAbilitySystemReplicationProxyInterface* GetReplicationInterface();
 
 	/** Current prediction key, set with FScopedPredictionWindow */
 	FPredictionKey	ScopedPredictionKey;
@@ -278,19 +280,19 @@ class GAMEPLAYABILITIES_API UAbilitySystemComponent : public UGameplayTasksCompo
 	}
 
 	/** Returns true if this is running on the server or has a valid prediction key */
-	bool HasAuthorityOrPredictionKey(const FGameplayAbilityActivationInfo* ActivationInfo) const;
+	UE_API bool HasAuthorityOrPredictionKey(const FGameplayAbilityActivationInfo* ActivationInfo) const;
 
 	/** Returns true if this component's actor has authority */
-	virtual bool IsOwnerActorAuthoritative() const;
+	UE_API virtual bool IsOwnerActorAuthoritative() const;
 
 	/** Returns true if this component should record montage replication info. */
-	virtual bool ShouldRecordMontageReplication() const;
+	UE_API virtual bool ShouldRecordMontageReplication() const;
 
 	/** Replicate that an ability has ended/canceled, to the client or server as appropriate */
-	virtual void ReplicateEndOrCancelAbility(FGameplayAbilitySpecHandle Handle, FGameplayAbilityActivationInfo ActivationInfo, UGameplayAbility* Ability, bool bWasCanceled);
+	UE_API virtual void ReplicateEndOrCancelAbility(FGameplayAbilitySpecHandle Handle, FGameplayAbilityActivationInfo ActivationInfo, UGameplayAbility* Ability, bool bWasCanceled);
 
 	/** Force cancels the ability and does not replicate this to the other side. This should be called when the ability is cancelled by the other side */
-	virtual void ForceCancelAbilityDueToReplication(UGameplayAbility* Instance);
+	UE_API virtual void ForceCancelAbilityDueToReplication(UGameplayAbility* Instance);
 
 	/** A pending activation that cannot be activated yet, will be rechecked at a later point */
 	struct FPendingAbilityInfo
@@ -323,18 +325,18 @@ class GAMEPLAYABILITIES_API UAbilitySystemComponent : public UGameplayTasksCompo
 
 	/** Applies a previously created gameplay effect spec to a target */
 	UFUNCTION(BlueprintCallable, Category = GameplayEffects, meta = (DisplayName = "ApplyGameplayEffectSpecToTarget", ScriptName = "ApplyGameplayEffectSpecToTarget"))
-	FActiveGameplayEffectHandle BP_ApplyGameplayEffectSpecToTarget(const FGameplayEffectSpecHandle& SpecHandle, UAbilitySystemComponent* Target);
+	UE_API FActiveGameplayEffectHandle BP_ApplyGameplayEffectSpecToTarget(const FGameplayEffectSpecHandle& SpecHandle, UAbilitySystemComponent* Target);
 
-	virtual FActiveGameplayEffectHandle ApplyGameplayEffectSpecToTarget(const FGameplayEffectSpec& GameplayEffect, UAbilitySystemComponent *Target, FPredictionKey PredictionKey=FPredictionKey());
+	UE_API virtual FActiveGameplayEffectHandle ApplyGameplayEffectSpecToTarget(const FGameplayEffectSpec& GameplayEffect, UAbilitySystemComponent *Target, FPredictionKey PredictionKey=FPredictionKey());
 
 	/** Applies a previously created gameplay effect spec to this component */
 	UFUNCTION(BlueprintCallable, Category = GameplayEffects, meta = (DisplayName = "ApplyGameplayEffectSpecToSelf", ScriptName = "ApplyGameplayEffectSpecToSelf"))
-	FActiveGameplayEffectHandle BP_ApplyGameplayEffectSpecToSelf(const FGameplayEffectSpecHandle& SpecHandle);
+	UE_API FActiveGameplayEffectHandle BP_ApplyGameplayEffectSpecToSelf(const FGameplayEffectSpecHandle& SpecHandle);
 
-	virtual FActiveGameplayEffectHandle ApplyGameplayEffectSpecToSelf(const FGameplayEffectSpec& GameplayEffect, FPredictionKey PredictionKey = FPredictionKey());
+	UE_API virtual FActiveGameplayEffectHandle ApplyGameplayEffectSpecToSelf(const FGameplayEffectSpec& GameplayEffect, FPredictionKey PredictionKey = FPredictionKey());
 
 	/** Gets the FActiveGameplayEffect based on the passed in Handle */
-	const UGameplayEffect* GetGameplayEffectDefForHandle(FActiveGameplayEffectHandle Handle);
+	UE_API const UGameplayEffect* GetGameplayEffectDefForHandle(FActiveGameplayEffectHandle Handle);
 
 	/** 
 	* Removes the specified GameplayEffect by Handle.
@@ -344,7 +346,7 @@ class GAMEPLAYABILITIES_API UAbilitySystemComponent : public UGameplayTasksCompo
 	* @return Whether the effect was successfully removed.
 	*/
 	UFUNCTION(BlueprintCallable, BlueprintAuthorityOnly, Category = GameplayEffects)
-	virtual bool RemoveActiveGameplayEffect(FActiveGameplayEffectHandle Handle, int32 StacksToRemove=-1);
+	UE_API virtual bool RemoveActiveGameplayEffect(FActiveGameplayEffectHandle Handle, int32 StacksToRemove=-1);
 
 	/** 
 	 * Remove active gameplay effects whose backing definition are the specified gameplay effect class
@@ -354,15 +356,15 @@ class GAMEPLAYABILITIES_API UAbilitySystemComponent : public UGameplayTasksCompo
 	 * @param StacksToRemove					Number of stacks to remove, -1 means remove all
 	 */
 	UFUNCTION(BlueprintCallable, BlueprintAuthorityOnly, Category = GameplayEffects)
-	virtual void RemoveActiveGameplayEffectBySourceEffect(TSubclassOf<UGameplayEffect> GameplayEffect, UAbilitySystemComponent* InstigatorAbilitySystemComponent, int32 StacksToRemove = -1);
+	UE_API virtual void RemoveActiveGameplayEffectBySourceEffect(TSubclassOf<UGameplayEffect> GameplayEffect, UAbilitySystemComponent* InstigatorAbilitySystemComponent, int32 StacksToRemove = -1);
 
 	/** Get an outgoing GameplayEffectSpec that is ready to be applied to other things. */
 	UFUNCTION(BlueprintCallable, Category = GameplayEffects)
-	virtual FGameplayEffectSpecHandle MakeOutgoingSpec(TSubclassOf<UGameplayEffect> GameplayEffectClass, float Level, FGameplayEffectContextHandle Context) const;
+	UE_API virtual FGameplayEffectSpecHandle MakeOutgoingSpec(TSubclassOf<UGameplayEffect> GameplayEffectClass, float Level, FGameplayEffectContextHandle Context) const;
 
 	/** Create an EffectContext for the owner of this AbilitySystemComponent */
 	UFUNCTION(BlueprintCallable, Category = GameplayEffects)
-	virtual FGameplayEffectContextHandle MakeEffectContext() const;
+	UE_API virtual FGameplayEffectContextHandle MakeEffectContext() const;
 
 	/**
 	 * Get the count of the specified source effect on the ability system component. For non-stacking effects, this is the sum of all active instances.
@@ -374,7 +376,7 @@ class GAMEPLAYABILITIES_API UAbilitySystemComponent : public UGameplayTasksCompo
 	 * @return Count of the specified source effect
 	 */
 	UFUNCTION(BlueprintCallable, BlueprintPure, Category=GameplayEffects)
-	int32 GetGameplayEffectCount(TSubclassOf<UGameplayEffect> SourceGameplayEffect, UAbilitySystemComponent* OptionalInstigatorFilterComponent, bool bEnforceOnGoingCheck = true) const;
+	UE_API int32 GetGameplayEffectCount(TSubclassOf<UGameplayEffect> SourceGameplayEffect, UAbilitySystemComponent* OptionalInstigatorFilterComponent, bool bEnforceOnGoingCheck = true) const;
 
 	/**
 	 * Get the count of the specified source effect on the ability system component. For non-stacking effects, this is the sum of all active instances.
@@ -386,81 +388,81 @@ class GAMEPLAYABILITIES_API UAbilitySystemComponent : public UGameplayTasksCompo
 	 * @return Count of the specified source effect
 	 */
 	UFUNCTION(BlueprintCallable, BlueprintPure, Category = GameplayEffects)
-	int32 GetGameplayEffectCount_IfLoaded(TSoftClassPtr<UGameplayEffect> SoftSourceGameplayEffect, UAbilitySystemComponent* OptionalInstigatorFilterComponent, bool bEnforceOnGoingCheck = true) const;
+	UE_API int32 GetGameplayEffectCount_IfLoaded(TSoftClassPtr<UGameplayEffect> SoftSourceGameplayEffect, UAbilitySystemComponent* OptionalInstigatorFilterComponent, bool bEnforceOnGoingCheck = true) const;
 
 	/** Returns the sum of StackCount of all gameplay effects that pass query */
-	int32 GetAggregatedStackCount(const FGameplayEffectQuery& Query) const;
+	UE_API int32 GetAggregatedStackCount(const FGameplayEffectQuery& Query) const;
 
 	/** This only exists so it can be hooked up to a multicast delegate */
 	UE_DEPRECATED_FORGAME(5.5, "This shouldn't be public. Use RemoveActiveGameplayEffect (which won't allow non-authority to remove) or expose RemoveActiveGameplayEffect_AllowClientRemoval if your intent is to break that contract.")
-	void RemoveActiveGameplayEffect_NoReturn(FActiveGameplayEffectHandle Handle, int32 StacksToRemove = -1);
+	UE_API void RemoveActiveGameplayEffect_NoReturn(FActiveGameplayEffectHandle Handle, int32 StacksToRemove = -1);
 
 	/** Called for predictively added gameplay cue. Needs to remove tag count and possible invoke OnRemove event if misprediction */
-	virtual void OnPredictiveGameplayCueCatchup(FGameplayTag Tag);
+	UE_API virtual void OnPredictiveGameplayCueCatchup(FGameplayTag Tag);
 
 	/** Returns the total duration of a gameplay effect */
-	float GetGameplayEffectDuration(FActiveGameplayEffectHandle Handle) const;
+	UE_API float GetGameplayEffectDuration(FActiveGameplayEffectHandle Handle) const;
 
 	/** Called whenever the server time replicates via the game state to keep our cooldown timers in sync with the server */
-	virtual void RecomputeGameplayEffectStartTimes(const float WorldTime, const float ServerWorldTime);
+	UE_API virtual void RecomputeGameplayEffectStartTimes(const float WorldTime, const float ServerWorldTime);
 
 	/** Return start time and total duration of a gameplay effect */
-	void GetGameplayEffectStartTimeAndDuration(FActiveGameplayEffectHandle Handle, float& StartEffectTime, float& Duration) const;
+	UE_API void GetGameplayEffectStartTimeAndDuration(FActiveGameplayEffectHandle Handle, float& StartEffectTime, float& Duration) const;
 
 	/** Dynamically update the set-by-caller magnitude for an active gameplay effect */
 	UFUNCTION(BlueprintCallable, BlueprintAuthorityOnly, Category = GameplayEffects)
-	virtual void UpdateActiveGameplayEffectSetByCallerMagnitude(FActiveGameplayEffectHandle ActiveHandle, UPARAM(meta=(Categories = "SetByCaller"))FGameplayTag SetByCallerTag, float NewValue);
+	UE_API virtual void UpdateActiveGameplayEffectSetByCallerMagnitude(FActiveGameplayEffectHandle ActiveHandle, UPARAM(meta=(Categories = "SetByCaller"))FGameplayTag SetByCallerTag, float NewValue);
 
 	/** Dynamically update multiple set-by-caller magnitudes for an active gameplay effect */
 	UFUNCTION(BlueprintCallable, BlueprintAuthorityOnly, Category = GameplayEffects)
-	virtual void UpdateActiveGameplayEffectSetByCallerMagnitudes(FActiveGameplayEffectHandle ActiveHandle, const TMap<FGameplayTag, float>& NewSetByCallerValues);
+	UE_API virtual void UpdateActiveGameplayEffectSetByCallerMagnitudes(FActiveGameplayEffectHandle ActiveHandle, const TMap<FGameplayTag, float>& NewSetByCallerValues);
 
 	/** Updates the level of an already applied gameplay effect. The intention is that this is 'seemless' and doesnt behave like removing/reapplying */
 	UFUNCTION(BlueprintCallable, BlueprintAuthorityOnly, Category = GameplayEffects)
-	virtual void SetActiveGameplayEffectLevel(FActiveGameplayEffectHandle ActiveHandle, int32 NewLevel);
+	UE_API virtual void SetActiveGameplayEffectLevel(FActiveGameplayEffectHandle ActiveHandle, int32 NewLevel);
 
 	/** Updates the level of an already applied gameplay effect. The intention is that this is 'seemless' and doesnt behave like removing/reapplying */
 	UFUNCTION(BlueprintCallable, BlueprintAuthorityOnly, Category = GameplayEffects)
-	virtual void SetActiveGameplayEffectLevelUsingQuery(FGameplayEffectQuery Query, int32 NewLevel);
+	UE_API virtual void SetActiveGameplayEffectLevelUsingQuery(FGameplayEffectQuery Query, int32 NewLevel);
 
 	/** Inhibit an active gameplay effect so that it is disabled, but not removed */
 	UE_DEPRECATED(5.4, "Use SetActiveGameplayEffectInhibit with a MoveTemp(ActiveGEHandle) so it's clear the Handle is no longer valid. Check (then use) the returned FActiveGameplayEffectHandle to continue your operation.")
-	virtual void InhibitActiveGameplayEffect(FActiveGameplayEffectHandle ActiveGEHandle, bool bInhibit, bool bInvokeGameplayCueEvents);
+	UE_API virtual void InhibitActiveGameplayEffect(FActiveGameplayEffectHandle ActiveGEHandle, bool bInhibit, bool bInvokeGameplayCueEvents);
 
 	/**
 	 * (Un-)Inhibit an Active Gameplay Effect so it may be disabled (and perform some disabling actions, such as removing tags).
 	 * An inhibited Active Gameplay Effect will lay dormant for re-enabling on some condition (usually tags).  When it's uninhibited, it will reapply some of its functionality.
 	 * NOTE:  The passed-in ActiveGEHandle is no longer valid.  Use the returned FActiveGameplayEffectHandle to determine if the Active GE Handle is still active.
 	 */
-	virtual FActiveGameplayEffectHandle SetActiveGameplayEffectInhibit(FActiveGameplayEffectHandle&& ActiveGEHandle, bool bInhibit, bool bInvokeGameplayCueEvents);
+	UE_API virtual FActiveGameplayEffectHandle SetActiveGameplayEffectInhibit(FActiveGameplayEffectHandle&& ActiveGEHandle, bool bInhibit, bool bInvokeGameplayCueEvents);
 
 	/**
 	 * Raw accessor to ask the magnitude of a gameplay effect, not necessarily always correct. How should outside code (UI, etc) ask things like 'how much is this gameplay effect modifying my damage by'
 	 * (most likely we want to catch this on the backend - when damage is applied we can get a full dump/history of how the number got to where it is. But still we may need polling methods like below (how much would my damage be)
 	 */
 	UFUNCTION(BlueprintCallable, Category = GameplayEffects)
-	float GetGameplayEffectMagnitude(FActiveGameplayEffectHandle Handle, FGameplayAttribute Attribute) const;
+	UE_API float GetGameplayEffectMagnitude(FActiveGameplayEffectHandle Handle, FGameplayAttribute Attribute) const;
 
 	/** Returns current stack count of an already applied GE */
-	int32 GetCurrentStackCount(FActiveGameplayEffectHandle Handle) const;
+	UE_API int32 GetCurrentStackCount(FActiveGameplayEffectHandle Handle) const;
 
 	/** Returns current stack count of an already applied GE, but given the ability spec handle that was granted by the GE */
-	int32 GetCurrentStackCount(FGameplayAbilitySpecHandle Handle) const;
+	UE_API int32 GetCurrentStackCount(FGameplayAbilitySpecHandle Handle) const;
 
 	/** Returns debug string describing active gameplay effect */
-	FString GetActiveGEDebugString(FActiveGameplayEffectHandle Handle) const;
+	UE_API FString GetActiveGEDebugString(FActiveGameplayEffectHandle Handle) const;
 
 	/** Gets the GE Handle of the GE that granted the passed in Ability */
-	FActiveGameplayEffectHandle FindActiveGameplayEffectHandle(FGameplayAbilitySpecHandle Handle) const;
+	UE_API FActiveGameplayEffectHandle FindActiveGameplayEffectHandle(FGameplayAbilitySpecHandle Handle) const;
 
 	/** Returns const pointer to the actual active gameplay effect structure */
-	const FActiveGameplayEffect* GetActiveGameplayEffect(const FActiveGameplayEffectHandle Handle) const;
+	UE_API const FActiveGameplayEffect* GetActiveGameplayEffect(const FActiveGameplayEffectHandle Handle) const;
 
 	/** Returns all active gameplay effects on this ASC */
-	const FActiveGameplayEffectsContainer& GetActiveGameplayEffects() const;
+	UE_API const FActiveGameplayEffectsContainer& GetActiveGameplayEffects() const;
 
 	/** Returns a const pointer to the gameplay effect CDO associated with an active handle. */
-	const UGameplayEffect* GetGameplayEffectCDO(const FActiveGameplayEffectHandle Handle) const;
+	UE_API const UGameplayEffect* GetGameplayEffectCDO(const FActiveGameplayEffectHandle Handle) const;
 
 	/**
 	 * Get the source tags from the gameplay spec represented by the specified handle, if possible
@@ -469,7 +471,7 @@ class GAMEPLAYABILITIES_API UAbilitySystemComponent : public UGameplayTasksCompo
 	 * 
 	 * @return Source tags from the gameplay spec represented by the handle, if possible
 	 */
-	const FGameplayTagContainer* GetGameplayEffectSourceTagsFromHandle(FActiveGameplayEffectHandle Handle) const;
+	UE_API const FGameplayTagContainer* GetGameplayEffectSourceTagsFromHandle(FActiveGameplayEffectHandle Handle) const;
 
 	/**
 	 * Get the target tags from the gameplay spec represented by the specified handle, if possible
@@ -478,14 +480,14 @@ class GAMEPLAYABILITIES_API UAbilitySystemComponent : public UGameplayTasksCompo
 	 * 
 	 * @return Target tags from the gameplay spec represented by the handle, if possible
 	 */
-	const FGameplayTagContainer* GetGameplayEffectTargetTagsFromHandle(FActiveGameplayEffectHandle Handle) const;
+	UE_API const FGameplayTagContainer* GetGameplayEffectTargetTagsFromHandle(FActiveGameplayEffectHandle Handle) const;
 
 	/**
 	 * Populate the specified capture spec with the data necessary to capture an attribute from the component
 	 * 
 	 * @param OutCaptureSpec	[OUT] Capture spec to populate with captured data
 	 */
-	void CaptureAttributeForGameplayEffect(OUT FGameplayEffectAttributeCaptureSpec& OutCaptureSpec);
+	UE_API void CaptureAttributeForGameplayEffect(OUT FGameplayEffectAttributeCaptureSpec& OutCaptureSpec);
 	
 	// ----------------------------------------------------------------------------------------------------------------
 	//  Callbacks / Notifies
@@ -493,19 +495,19 @@ class GAMEPLAYABILITIES_API UAbilitySystemComponent : public UGameplayTasksCompo
 	// ----------------------------------------------------------------------------------------------------------------
 
 	/** Called when a specific attribute aggregator value changes, gameplay effects refresh their values when this happens */
-	void OnAttributeAggregatorDirty(FAggregator* Aggregator, FGameplayAttribute Attribute, bool FromRecursiveCall=false);
+	UE_API void OnAttributeAggregatorDirty(FAggregator* Aggregator, FGameplayAttribute Attribute, bool FromRecursiveCall=false);
 
 	/** Called when attribute magnitudes change, to forward information to dependent gameplay effects */
-	void OnMagnitudeDependencyChange(FActiveGameplayEffectHandle Handle, const FAggregator* ChangedAggregator);
+	UE_API void OnMagnitudeDependencyChange(FActiveGameplayEffectHandle Handle, const FAggregator* ChangedAggregator);
 
 	/** This ASC has successfully applied a GE to something (potentially itself) */
-	void OnGameplayEffectAppliedToTarget(UAbilitySystemComponent* Target, const FGameplayEffectSpec& SpecApplied, FActiveGameplayEffectHandle ActiveHandle);
-	void OnGameplayEffectAppliedToSelf(UAbilitySystemComponent* Source, const FGameplayEffectSpec& SpecApplied, FActiveGameplayEffectHandle ActiveHandle);
-	void OnPeriodicGameplayEffectExecuteOnTarget(UAbilitySystemComponent* Target, const FGameplayEffectSpec& SpecExecuted, FActiveGameplayEffectHandle ActiveHandle);
-	void OnPeriodicGameplayEffectExecuteOnSelf(UAbilitySystemComponent* Source, const FGameplayEffectSpec& SpecExecuted, FActiveGameplayEffectHandle ActiveHandle);
+	UE_API void OnGameplayEffectAppliedToTarget(UAbilitySystemComponent* Target, const FGameplayEffectSpec& SpecApplied, FActiveGameplayEffectHandle ActiveHandle);
+	UE_API void OnGameplayEffectAppliedToSelf(UAbilitySystemComponent* Source, const FGameplayEffectSpec& SpecApplied, FActiveGameplayEffectHandle ActiveHandle);
+	UE_API void OnPeriodicGameplayEffectExecuteOnTarget(UAbilitySystemComponent* Target, const FGameplayEffectSpec& SpecExecuted, FActiveGameplayEffectHandle ActiveHandle);
+	UE_API void OnPeriodicGameplayEffectExecuteOnSelf(UAbilitySystemComponent* Source, const FGameplayEffectSpec& SpecExecuted, FActiveGameplayEffectHandle ActiveHandle);
 
 	/** Called when the duration of a gameplay effect has changed */
-	virtual void OnGameplayEffectDurationChange(struct FActiveGameplayEffect& ActiveEffect);
+	UE_API virtual void OnGameplayEffectDurationChange(struct FActiveGameplayEffect& ActiveEffect);
 
 	/** Called on server whenever a GE is applied to self. This includes instant and duration based GEs. */
 	FOnGameplayEffectAppliedDelegate OnGameplayEffectAppliedDelegateToSelf;
@@ -526,10 +528,10 @@ class GAMEPLAYABILITIES_API UAbilitySystemComponent : public UGameplayTasksCompo
 	FImmunityBlockGE OnImmunityBlockGameplayEffectDelegate;
 
 	/** Register for when an attribute value changes, should be replaced by GetGameplayAttributeValueChangeDelegate */
-	FOnGameplayAttributeChange& RegisterGameplayAttributeEvent(FGameplayAttribute Attribute);
+	UE_API FOnGameplayAttributeChange& RegisterGameplayAttributeEvent(FGameplayAttribute Attribute);
 
 	/** Register for when an attribute value changes */
-	FOnGameplayAttributeValueChange& GetGameplayAttributeValueChangeDelegate(FGameplayAttribute Attribute);
+	UE_API FOnGameplayAttributeValueChange& GetGameplayAttributeValueChangeDelegate(FGameplayAttribute Attribute);
 
 	/** A generic callback anytime an ability is activated (started) */
 	FGenericAbilityDelegate AbilityActivatedCallbacks;
@@ -553,19 +555,19 @@ class GAMEPLAYABILITIES_API UAbilitySystemComponent : public UGameplayTasksCompo
 	TArray<FGameplayEffectApplicationQuery> GameplayEffectApplicationQueries;
 
 	/** Call notify callbacks above */
-	virtual void NotifyAbilityCommit(UGameplayAbility* Ability);
-	virtual void NotifyAbilityActivated(const FGameplayAbilitySpecHandle Handle, UGameplayAbility* Ability);
-	virtual void NotifyAbilityFailed(const FGameplayAbilitySpecHandle Handle, UGameplayAbility* Ability, const FGameplayTagContainer& FailureReason);
+	UE_API virtual void NotifyAbilityCommit(UGameplayAbility* Ability);
+	UE_API virtual void NotifyAbilityActivated(const FGameplayAbilitySpecHandle Handle, UGameplayAbility* Ability);
+	UE_API virtual void NotifyAbilityFailed(const FGameplayAbilitySpecHandle Handle, UGameplayAbility* Ability, const FGameplayTagContainer& FailureReason);
 
 	/** Called when any gameplay effects are removed */
-	FOnGivenActiveGameplayEffectRemoved& OnAnyGameplayEffectRemovedDelegate();
+	UE_API FOnGivenActiveGameplayEffectRemoved& OnAnyGameplayEffectRemovedDelegate();
 
 	/** Returns delegate structure that allows binding to several gameplay effect changes */
-	FActiveGameplayEffectEvents* GetActiveEffectEventSet(FActiveGameplayEffectHandle Handle);
-	FOnActiveGameplayEffectRemoved_Info* OnGameplayEffectRemoved_InfoDelegate(FActiveGameplayEffectHandle Handle);
-	FOnActiveGameplayEffectStackChange* OnGameplayEffectStackChangeDelegate(FActiveGameplayEffectHandle Handle);
-	FOnActiveGameplayEffectTimeChange* OnGameplayEffectTimeChangeDelegate(FActiveGameplayEffectHandle Handle);
-	FOnActiveGameplayEffectInhibitionChanged* OnGameplayEffectInhibitionChangedDelegate(FActiveGameplayEffectHandle Handle);
+	UE_API FActiveGameplayEffectEvents* GetActiveEffectEventSet(FActiveGameplayEffectHandle Handle);
+	UE_API FOnActiveGameplayEffectRemoved_Info* OnGameplayEffectRemoved_InfoDelegate(FActiveGameplayEffectHandle Handle);
+	UE_API FOnActiveGameplayEffectStackChange* OnGameplayEffectStackChangeDelegate(FActiveGameplayEffectHandle Handle);
+	UE_API FOnActiveGameplayEffectTimeChange* OnGameplayEffectTimeChangeDelegate(FActiveGameplayEffectHandle Handle);
+	UE_API FOnActiveGameplayEffectInhibitionChanged* OnGameplayEffectInhibitionChangedDelegate(FActiveGameplayEffectHandle Handle);
 
 	// ----------------------------------------------------------------------------------------------------------------
 	//  Gameplay tag operations
@@ -682,7 +684,7 @@ class GAMEPLAYABILITIES_API UAbilitySystemComponent : public UGameplayTasksCompo
 	 * @param GameplayTag The gameplay tag to query
 	 */
 	UFUNCTION(BlueprintPure, Category = "Gameplay Tags")
-	int32 GetGameplayTagCount(FGameplayTag GameplayTag) const;
+	UE_API int32 GetGameplayTagCount(FGameplayTag GameplayTag) const;
 
 	/**
 	 *  Allows GameCode to add loose gameplaytags which are not backed by a GameplayEffect. Tags added using 
@@ -739,25 +741,25 @@ class GAMEPLAYABILITIES_API UAbilitySystemComponent : public UGameplayTasksCompo
 	}
 
 	/** Allow events to be registered for specific gameplay tags being added or removed */
-	FOnGameplayEffectTagCountChanged& RegisterGameplayTagEvent(FGameplayTag Tag, EGameplayTagEventType::Type EventType=EGameplayTagEventType::NewOrRemoved);
+	UE_API FOnGameplayEffectTagCountChanged& RegisterGameplayTagEvent(FGameplayTag Tag, EGameplayTagEventType::Type EventType=EGameplayTagEventType::NewOrRemoved);
 
 	/** Unregister previously added events */
-	bool UnregisterGameplayTagEvent(FDelegateHandle DelegateHandle, FGameplayTag Tag, EGameplayTagEventType::Type EventType=EGameplayTagEventType::NewOrRemoved);
+	UE_API bool UnregisterGameplayTagEvent(FDelegateHandle DelegateHandle, FGameplayTag Tag, EGameplayTagEventType::Type EventType=EGameplayTagEventType::NewOrRemoved);
 
 	/** Register a tag event and immediately call it */
-	FDelegateHandle RegisterAndCallGameplayTagEvent(FGameplayTag Tag, FOnGameplayEffectTagCountChanged::FDelegate Delegate, EGameplayTagEventType::Type EventType=EGameplayTagEventType::NewOrRemoved);
+	UE_API FDelegateHandle RegisterAndCallGameplayTagEvent(FGameplayTag Tag, FOnGameplayEffectTagCountChanged::FDelegate Delegate, EGameplayTagEventType::Type EventType=EGameplayTagEventType::NewOrRemoved);
 
 	/** Returns multicast delegate that is invoked whenever a tag is added or removed (but not if just count is increased. Only for 'new' and 'removed' events) */
-	FOnGameplayEffectTagCountChanged& RegisterGenericGameplayTagEvent();
+	UE_API FOnGameplayEffectTagCountChanged& RegisterGenericGameplayTagEvent();
 
 	/** Executes a gameplay event. Returns the number of successful ability activations triggered by the event */
-	virtual int32 HandleGameplayEvent(FGameplayTag EventTag, const FGameplayEventData* Payload);
+	UE_API virtual int32 HandleGameplayEvent(FGameplayTag EventTag, const FGameplayEventData* Payload);
 
 	/** Adds a new delegate to call when gameplay events happen. It will only be called if it matches any tags in passed filter container */
-	FDelegateHandle AddGameplayEventTagContainerDelegate(const FGameplayTagContainer& TagFilter, const FGameplayEventTagMulticastDelegate::FDelegate& Delegate);
+	UE_API FDelegateHandle AddGameplayEventTagContainerDelegate(const FGameplayTagContainer& TagFilter, const FGameplayEventTagMulticastDelegate::FDelegate& Delegate);
 
 	/** Remotes previously registered delegate */
-	void RemoveGameplayEventTagContainerDelegate(const FGameplayTagContainer& TagFilter, FDelegateHandle DelegateHandle);
+	UE_API void RemoveGameplayEventTagContainerDelegate(const FGameplayTagContainer& TagFilter, FDelegateHandle DelegateHandle);
 
 	/** Callbacks bound to Gameplay tags, these only activate if the exact tag is used. To handle tag hierarchies use AddGameplayEventContainerDelegate */
 	TMap<FGameplayTag, FGameplayEventMulticastDelegate> GenericGameplayEventCallbacks;
@@ -774,11 +776,11 @@ class GAMEPLAYABILITIES_API UAbilitySystemComponent : public UGameplayTasksCompo
 	UPROPERTY(meta = (SystemGameplayAttribute = "true"))
 	float IncomingDuration;
 
-	static FProperty* GetOutgoingDurationProperty();
-	static FProperty* GetIncomingDurationProperty();
+	static UE_API FProperty* GetOutgoingDurationProperty();
+	static UE_API FProperty* GetIncomingDurationProperty();
 
-	static const FGameplayEffectAttributeCaptureDefinition& GetOutgoingDurationCapture();
-	static const FGameplayEffectAttributeCaptureDefinition& GetIncomingDurationCapture();
+	static UE_API const FGameplayEffectAttributeCaptureDefinition& GetOutgoingDurationCapture();
+	static UE_API const FGameplayEffectAttributeCaptureDefinition& GetIncomingDurationCapture();
 
 	// ----------------------------------------------------------------------------------------------------------------
 	//  Additional Helper Functions
@@ -786,13 +788,13 @@ class GAMEPLAYABILITIES_API UAbilitySystemComponent : public UGameplayTasksCompo
 
 	/** Apply a gameplay effect to passed in target */
 	UFUNCTION(BlueprintCallable, Category = GameplayEffects, meta=(DisplayName = "ApplyGameplayEffectToTarget", ScriptName = "ApplyGameplayEffectToTarget"))
-	FActiveGameplayEffectHandle BP_ApplyGameplayEffectToTarget(TSubclassOf<UGameplayEffect> GameplayEffectClass, UAbilitySystemComponent *Target, float Level, FGameplayEffectContextHandle Context);
-	FActiveGameplayEffectHandle ApplyGameplayEffectToTarget(UGameplayEffect *GameplayEffect, UAbilitySystemComponent *Target, float Level = UGameplayEffect::INVALID_LEVEL, FGameplayEffectContextHandle Context = FGameplayEffectContextHandle(), FPredictionKey PredictionKey = FPredictionKey());
+	UE_API FActiveGameplayEffectHandle BP_ApplyGameplayEffectToTarget(TSubclassOf<UGameplayEffect> GameplayEffectClass, UAbilitySystemComponent *Target, float Level, FGameplayEffectContextHandle Context);
+	UE_API FActiveGameplayEffectHandle ApplyGameplayEffectToTarget(UGameplayEffect *GameplayEffect, UAbilitySystemComponent *Target, float Level = UGameplayEffect::INVALID_LEVEL, FGameplayEffectContextHandle Context = FGameplayEffectContextHandle(), FPredictionKey PredictionKey = FPredictionKey());
 
 	/** Apply a gameplay effect to self */
 	UFUNCTION(BlueprintCallable, Category = GameplayEffects, meta=(DisplayName = "ApplyGameplayEffectToSelf", ScriptName = "ApplyGameplayEffectToSelf"))
-	FActiveGameplayEffectHandle BP_ApplyGameplayEffectToSelf(TSubclassOf<UGameplayEffect> GameplayEffectClass, float Level, FGameplayEffectContextHandle EffectContext);
-	FActiveGameplayEffectHandle ApplyGameplayEffectToSelf(const UGameplayEffect *GameplayEffect, float Level, const FGameplayEffectContextHandle& EffectContext, FPredictionKey PredictionKey = FPredictionKey());
+	UE_API FActiveGameplayEffectHandle BP_ApplyGameplayEffectToSelf(TSubclassOf<UGameplayEffect> GameplayEffectClass, float Level, FGameplayEffectContextHandle EffectContext);
+	UE_API FActiveGameplayEffectHandle ApplyGameplayEffectToSelf(const UGameplayEffect *GameplayEffect, float Level, const FGameplayEffectContextHandle& EffectContext, FPredictionKey PredictionKey = FPredictionKey());
 
 	/** Returns the number of gameplay effects that are currently active on this ability system component */
 	int32 GetNumActiveGameplayEffects() const
@@ -825,50 +827,50 @@ class GAMEPLAYABILITIES_API UAbilitySystemComponent : public UGameplayTasksCompo
 	}
 
 	/** Gets time remaining for all effects that match query */
-	TArray<float> GetActiveEffectsTimeRemaining(const FGameplayEffectQuery& Query) const;
+	UE_API TArray<float> GetActiveEffectsTimeRemaining(const FGameplayEffectQuery& Query) const;
 
 	/** Gets total duration for all effects that match query */
-	TArray<float> GetActiveEffectsDuration(const FGameplayEffectQuery& Query) const;
+	UE_API TArray<float> GetActiveEffectsDuration(const FGameplayEffectQuery& Query) const;
 
 	/** Gets both time remaining and total duration  for all effects that match query */
-	TArray<TPair<float,float>> GetActiveEffectsTimeRemainingAndDuration(const FGameplayEffectQuery& Query) const;
+	UE_API TArray<TPair<float,float>> GetActiveEffectsTimeRemainingAndDuration(const FGameplayEffectQuery& Query) const;
 
 	/** Returns list of active effects, for a query */
 	UFUNCTION(BlueprintCallable, BlueprintPure=false, Category = "GameplayEffects", meta=(DisplayName = "Get Active Gameplay Effects for Query"))
-	TArray<FActiveGameplayEffectHandle> GetActiveEffects(const FGameplayEffectQuery& Query) const;
+	UE_API TArray<FActiveGameplayEffectHandle> GetActiveEffects(const FGameplayEffectQuery& Query) const;
 
 	/** Returns list of active effects that have all of the passed in tags */
 	UFUNCTION(BlueprintCallable, BlueprintPure = false, Category = "GameplayEffects")
-	TArray<FActiveGameplayEffectHandle> GetActiveEffectsWithAllTags(FGameplayTagContainer Tags) const;
+	UE_API TArray<FActiveGameplayEffectHandle> GetActiveEffectsWithAllTags(FGameplayTagContainer Tags) const;
 
 	/** This will give the world time that all effects matching this query will be finished. If multiple effects match, it returns the one that returns last */
-	float GetActiveEffectsEndTime(const FGameplayEffectQuery& Query) const;
-	float GetActiveEffectsEndTimeWithInstigators(const FGameplayEffectQuery& Query, TArray<AActor*>& Instigators) const;
+	UE_API float GetActiveEffectsEndTime(const FGameplayEffectQuery& Query) const;
+	UE_API float GetActiveEffectsEndTimeWithInstigators(const FGameplayEffectQuery& Query, TArray<AActor*>& Instigators) const;
 
 	/** Returns end time and total duration */
-	bool GetActiveEffectsEndTimeAndDuration(const FGameplayEffectQuery& Query, float& EndTime, float& Duration) const;
+	UE_API bool GetActiveEffectsEndTimeAndDuration(const FGameplayEffectQuery& Query, float& EndTime, float& Duration) const;
 
 	/** Modify the start time of a gameplay effect, to deal with timers being out of sync originally */
-	virtual void ModifyActiveEffectStartTime(FActiveGameplayEffectHandle Handle, float StartTimeDiff);
+	UE_API virtual void ModifyActiveEffectStartTime(FActiveGameplayEffectHandle Handle, float StartTimeDiff);
 
 	/** Removes all active effects that contain any of the tags in Tags */
 	UFUNCTION(BlueprintCallable, Category = GameplayEffects)
-	int32 RemoveActiveEffectsWithTags(FGameplayTagContainer Tags);
+	UE_API int32 RemoveActiveEffectsWithTags(FGameplayTagContainer Tags);
 
 	/** Removes all active effects with captured source tags that contain any of the tags in Tags */
 	UFUNCTION(BlueprintCallable, Category = GameplayEffects)
-	int32 RemoveActiveEffectsWithSourceTags(FGameplayTagContainer Tags);
+	UE_API int32 RemoveActiveEffectsWithSourceTags(FGameplayTagContainer Tags);
 
 	/** Removes all active effects that apply any of the tags in Tags */
 	UFUNCTION(BlueprintCallable, Category = GameplayEffects)
-	int32 RemoveActiveEffectsWithAppliedTags(FGameplayTagContainer Tags);
+	UE_API int32 RemoveActiveEffectsWithAppliedTags(FGameplayTagContainer Tags);
 
 	/** Removes all active effects that grant any of the tags in Tags */
 	UFUNCTION(BlueprintCallable, Category = GameplayEffects)
-	int32 RemoveActiveEffectsWithGrantedTags(FGameplayTagContainer Tags);
+	UE_API int32 RemoveActiveEffectsWithGrantedTags(FGameplayTagContainer Tags);
 
 	/** Removes all active effects that match given query. StacksToRemove=-1 will remove all stacks. */
-	virtual int32 RemoveActiveEffects(const FGameplayEffectQuery& Query, int32 StacksToRemove = -1);
+	UE_API virtual int32 RemoveActiveEffects(const FGameplayEffectQuery& Query, int32 StacksToRemove = -1);
 
 	// ----------------------------------------------------------------------------------------------------------------
 	//	GameplayCues
@@ -876,59 +878,59 @@ class GAMEPLAYABILITIES_API UAbilitySystemComponent : public UGameplayTasksCompo
 	
 	// Do not call these functions directly, call the wrappers on GameplayCueManager instead
 	UFUNCTION(NetMulticast, unreliable)
-	void NetMulticast_InvokeGameplayCueExecuted_FromSpec(const FGameplayEffectSpecForRPC Spec, FPredictionKey PredictionKey) override;
+	UE_API void NetMulticast_InvokeGameplayCueExecuted_FromSpec(const FGameplayEffectSpecForRPC Spec, FPredictionKey PredictionKey) override;
 
 	UFUNCTION(NetMulticast, unreliable)
-	void NetMulticast_InvokeGameplayCueExecuted(const FGameplayTag GameplayCueTag, FPredictionKey PredictionKey, FGameplayEffectContextHandle EffectContext) override;
+	UE_API void NetMulticast_InvokeGameplayCueExecuted(const FGameplayTag GameplayCueTag, FPredictionKey PredictionKey, FGameplayEffectContextHandle EffectContext) override;
 
 	UFUNCTION(NetMulticast, unreliable)
-	void NetMulticast_InvokeGameplayCuesExecuted(const FGameplayTagContainer GameplayCueTags, FPredictionKey PredictionKey, FGameplayEffectContextHandle EffectContext) override;
+	UE_API void NetMulticast_InvokeGameplayCuesExecuted(const FGameplayTagContainer GameplayCueTags, FPredictionKey PredictionKey, FGameplayEffectContextHandle EffectContext) override;
 
 	UFUNCTION(NetMulticast, unreliable)
-	void NetMulticast_InvokeGameplayCueExecuted_WithParams(const FGameplayTag GameplayCueTag, FPredictionKey PredictionKey, FGameplayCueParameters GameplayCueParameters) override;
+	UE_API void NetMulticast_InvokeGameplayCueExecuted_WithParams(const FGameplayTag GameplayCueTag, FPredictionKey PredictionKey, FGameplayCueParameters GameplayCueParameters) override;
 
 	UFUNCTION(NetMulticast, unreliable)
-	void NetMulticast_InvokeGameplayCuesExecuted_WithParams(const FGameplayTagContainer GameplayCueTags, FPredictionKey PredictionKey, FGameplayCueParameters GameplayCueParameters) override;
+	UE_API void NetMulticast_InvokeGameplayCuesExecuted_WithParams(const FGameplayTagContainer GameplayCueTags, FPredictionKey PredictionKey, FGameplayCueParameters GameplayCueParameters) override;
 
 	UFUNCTION(NetMulticast, unreliable)
-	void NetMulticast_InvokeGameplayCueAdded(const FGameplayTag GameplayCueTag, FPredictionKey PredictionKey, FGameplayEffectContextHandle EffectContext) override;
+	UE_API void NetMulticast_InvokeGameplayCueAdded(const FGameplayTag GameplayCueTag, FPredictionKey PredictionKey, FGameplayEffectContextHandle EffectContext) override;
 
 	UFUNCTION(NetMulticast, unreliable)
-	void NetMulticast_InvokeGameplayCueAdded_WithParams(const FGameplayTag GameplayCueTag, FPredictionKey PredictionKey, FGameplayCueParameters Parameters) override;
+	UE_API void NetMulticast_InvokeGameplayCueAdded_WithParams(const FGameplayTag GameplayCueTag, FPredictionKey PredictionKey, FGameplayCueParameters Parameters) override;
 	
 	UFUNCTION(NetMulticast, unreliable)
-	void NetMulticast_InvokeGameplayCueAddedAndWhileActive_FromSpec(const FGameplayEffectSpecForRPC& Spec, FPredictionKey PredictionKey) override;
+	UE_API void NetMulticast_InvokeGameplayCueAddedAndWhileActive_FromSpec(const FGameplayEffectSpecForRPC& Spec, FPredictionKey PredictionKey) override;
 
 	UFUNCTION(NetMulticast, unreliable)
-	void NetMulticast_InvokeGameplayCueAddedAndWhileActive_WithParams(const FGameplayTag GameplayCueTag, FPredictionKey PredictionKey, FGameplayCueParameters GameplayCueParameters) override;
+	UE_API void NetMulticast_InvokeGameplayCueAddedAndWhileActive_WithParams(const FGameplayTag GameplayCueTag, FPredictionKey PredictionKey, FGameplayCueParameters GameplayCueParameters) override;
 
 	UFUNCTION(NetMulticast, unreliable)
-	void NetMulticast_InvokeGameplayCuesAddedAndWhileActive_WithParams(const FGameplayTagContainer GameplayCueTags, FPredictionKey PredictionKey, FGameplayCueParameters GameplayCueParameters) override;
+	UE_API void NetMulticast_InvokeGameplayCuesAddedAndWhileActive_WithParams(const FGameplayTagContainer GameplayCueTags, FPredictionKey PredictionKey, FGameplayCueParameters GameplayCueParameters) override;
 
 	/** GameplayCues can also come on their own. These take an optional effect context to pass through hit result, etc */
-	void ExecuteGameplayCue(const FGameplayTag GameplayCueTag, FGameplayEffectContextHandle EffectContext = FGameplayEffectContextHandle());
-	void ExecuteGameplayCue(const FGameplayTag GameplayCueTag, const FGameplayCueParameters& GameplayCueParameters);
+	UE_API void ExecuteGameplayCue(const FGameplayTag GameplayCueTag, FGameplayEffectContextHandle EffectContext = FGameplayEffectContextHandle());
+	UE_API void ExecuteGameplayCue(const FGameplayTag GameplayCueTag, const FGameplayCueParameters& GameplayCueParameters);
 
 	/** Add a persistent gameplay cue */
-	void AddGameplayCue(const FGameplayTag GameplayCueTag, FGameplayEffectContextHandle EffectContext = FGameplayEffectContextHandle());
-	void AddGameplayCue(const FGameplayTag GameplayCueTag, const FGameplayCueParameters& GameplayCueParameters);
+	UE_API void AddGameplayCue(const FGameplayTag GameplayCueTag, FGameplayEffectContextHandle EffectContext = FGameplayEffectContextHandle());
+	UE_API void AddGameplayCue(const FGameplayTag GameplayCueTag, const FGameplayCueParameters& GameplayCueParameters);
 
 	/** Add gameplaycue for minimal replication mode. Should only be called in paths that would replicate gameplaycues in other ways (through GE for example) if not in minimal replication mode */
-	void AddGameplayCue_MinimalReplication(const FGameplayTag GameplayCueTag, FGameplayEffectContextHandle EffectContext = FGameplayEffectContextHandle());
+	UE_API void AddGameplayCue_MinimalReplication(const FGameplayTag GameplayCueTag, FGameplayEffectContextHandle EffectContext = FGameplayEffectContextHandle());
 
 	/** Remove a persistent gameplay cue */
-	void RemoveGameplayCue(const FGameplayTag GameplayCueTag);
+	UE_API void RemoveGameplayCue(const FGameplayTag GameplayCueTag);
 	
 	/** Remove gameplaycue for minimal replication mode. Should only be called in paths that would replicate gameplaycues in other ways (through GE for example) if not in minimal replication mode */
-	void RemoveGameplayCue_MinimalReplication(const FGameplayTag GameplayCueTag);
+	UE_API void RemoveGameplayCue_MinimalReplication(const FGameplayTag GameplayCueTag);
 	
 	/** Removes any GameplayCue added on its own, i.e. not as part of a GameplayEffect. */
-	void RemoveAllGameplayCues();
+	UE_API void RemoveAllGameplayCues();
 
 	/** Handles gameplay cue events from external sources */
-	void InvokeGameplayCueEvent(const FGameplayEffectSpecForRPC& Spec, EGameplayCueEvent::Type EventType);
-	void InvokeGameplayCueEvent(const FGameplayTag GameplayCueTag, EGameplayCueEvent::Type EventType, FGameplayEffectContextHandle EffectContext = FGameplayEffectContextHandle());
-	void InvokeGameplayCueEvent(const FGameplayTag GameplayCueTag, EGameplayCueEvent::Type EventType, const FGameplayCueParameters& GameplayCueParameters);
+	UE_API void InvokeGameplayCueEvent(const FGameplayEffectSpecForRPC& Spec, EGameplayCueEvent::Type EventType);
+	UE_API void InvokeGameplayCueEvent(const FGameplayTag GameplayCueTag, EGameplayCueEvent::Type EventType, FGameplayEffectContextHandle EffectContext = FGameplayEffectContextHandle());
+	UE_API void InvokeGameplayCueEvent(const FGameplayTag GameplayCueTag, EGameplayCueEvent::Type EventType, const FGameplayCueParameters& GameplayCueParameters);
 
 	/** Allows polling to see if a GameplayCue is active. We expect most GameplayCue handling to be event based, but some cases we may need to check if a GameplayCue is active (Animation Blueprint for example) */
 	UFUNCTION(BlueprintCallable, Category="GameplayCue", meta=(GameplayTagFilter="GameplayCue"))
@@ -938,17 +940,17 @@ class GAMEPLAYABILITIES_API UAbilitySystemComponent : public UGameplayTasksCompo
 	}
 
 	/** Will initialize gameplay cue parameters with this ASC's Owner (Instigator) and AvatarActor (EffectCauser) */
-	virtual void InitDefaultGameplayCueParameters(FGameplayCueParameters& Parameters);
+	UE_API virtual void InitDefaultGameplayCueParameters(FGameplayCueParameters& Parameters);
 
 	/** Are we ready to invoke gameplaycues yet? */
-	virtual bool IsReadyForGameplayCues();
+	UE_API virtual bool IsReadyForGameplayCues();
 
 	/** Handle GameplayCues that may have been deferred while doing the NetDeltaSerialize and waiting for the avatar actor to get loaded */
-	virtual void HandleDeferredGameplayCues(const FActiveGameplayEffectsContainer* GameplayEffectsContainer);
+	UE_API virtual void HandleDeferredGameplayCues(const FActiveGameplayEffectsContainer* GameplayEffectsContainer);
 
 	/** Invokes the WhileActive event for all GCs on active, non inhibited, GEs. This would typically be used on "respawn" or something where the mesh/avatar has changed */
 	UE_DEPRECATED(5.4, "ReinvokeActiveGameplayCues was unused and had logic inconsistent with predicting Gameplay Effects.  You can implement it in your own project if desired.")
-	virtual void ReinvokeActiveGameplayCues();
+	UE_API virtual void ReinvokeActiveGameplayCues();
 
 	/**
 	 *	GameplayAbilities
@@ -972,7 +974,7 @@ class GAMEPLAYABILITIES_API UAbilitySystemComponent : public UGameplayTasksCompo
 	 * 
 	 * @param AbilitySpec FGameplayAbilitySpec containing information about the ability class, level and input ID to bind it to.
 	 */
-	FGameplayAbilitySpecHandle GiveAbility(const FGameplayAbilitySpec& AbilitySpec);
+	UE_API FGameplayAbilitySpecHandle GiveAbility(const FGameplayAbilitySpec& AbilitySpec);
 
 	/*
 	 * Grants an ability and attempts to activate it exactly one time, which will cause it to be removed.
@@ -981,7 +983,7 @@ class GAMEPLAYABILITIES_API UAbilitySystemComponent : public UGameplayTasksCompo
 	 * @param AbilitySpec FGameplayAbilitySpec containing information about the ability class, level and input ID to bind it to.
 	 * @param GameplayEventData Optional activation event data. If provided, Activate Ability From Event will be called instead of ActivateAbility, passing the Event Data
 	 */
-	FGameplayAbilitySpecHandle GiveAbilityAndActivateOnce(FGameplayAbilitySpec& AbilitySpec, const FGameplayEventData* GameplayEventData = nullptr);
+	UE_API FGameplayAbilitySpecHandle GiveAbilityAndActivateOnce(FGameplayAbilitySpec& AbilitySpec, const FGameplayEventData* GameplayEventData = nullptr);
 
 	/**
 	 * Grants a Gameplay Ability and returns its handle.
@@ -992,7 +994,7 @@ class GAMEPLAYABILITIES_API UAbilitySystemComponent : public UGameplayTasksCompo
 	 * @param InputID Input ID value to bind ability activation to.
 	 */
 	UFUNCTION(BlueprintCallable, BlueprintAuthorityOnly, Category = "Gameplay Abilities", meta = (DisplayName = "Give Ability", ScriptName = "GiveAbility"))
-	FGameplayAbilitySpecHandle K2_GiveAbility(TSubclassOf<UGameplayAbility> AbilityClass, int32 Level = 0, int32 InputID = -1);
+	UE_API FGameplayAbilitySpecHandle K2_GiveAbility(TSubclassOf<UGameplayAbility> AbilityClass, int32 Level = 0, int32 InputID = -1);
 
 	/**
 	 * Grants a Gameplay Ability, activates it once, and removes it.
@@ -1003,11 +1005,11 @@ class GAMEPLAYABILITIES_API UAbilitySystemComponent : public UGameplayTasksCompo
 	 * @param InputID Input ID value to bind ability activation to.
 	 */
 	UFUNCTION(BlueprintCallable, BlueprintAuthorityOnly, Category = "Gameplay Abilities", meta = (DisplayName = "Give Ability And Activate Once", ScriptName = "GiveAbilityAndActivateOnce"))
-	FGameplayAbilitySpecHandle K2_GiveAbilityAndActivateOnce(TSubclassOf<UGameplayAbility> AbilityClass, int32 Level = 0, int32 InputID = -1);
+	UE_API FGameplayAbilitySpecHandle K2_GiveAbilityAndActivateOnce(TSubclassOf<UGameplayAbility> AbilityClass, int32 Level = 0, int32 InputID = -1);
 
 	/** Wipes all 'given' abilities. This will be ignored if the actor is not authoritative. */
 	UFUNCTION(BlueprintCallable, BlueprintAuthorityOnly, Category="Gameplay Abilities")
-	void ClearAllAbilities();
+	UE_API void ClearAllAbilities();
 
 	/**
 	 * Clears all abilities bound to a given Input ID
@@ -1016,7 +1018,7 @@ class GAMEPLAYABILITIES_API UAbilitySystemComponent : public UGameplayTasksCompo
 	 * @param InputID The numeric Input ID of the abilities to remove
 	 */
 	UFUNCTION(BlueprintCallable, BlueprintAuthorityOnly, Category = "Gameplay Abilities")
-	void ClearAllAbilitiesWithInputID(int32 InputID = 0);
+	UE_API void ClearAllAbilitiesWithInputID(int32 InputID = 0);
 
 	/** 
 	 * Removes the specified ability.
@@ -1025,10 +1027,10 @@ class GAMEPLAYABILITIES_API UAbilitySystemComponent : public UGameplayTasksCompo
 	 * @param Handle Ability Spec Handle of the ability we want to remove
 	 */
 	UFUNCTION(BlueprintCallable, BlueprintAuthorityOnly, Category = "Gameplay Abilities")
-	void ClearAbility(const FGameplayAbilitySpecHandle& Handle);
+	UE_API void ClearAbility(const FGameplayAbilitySpecHandle& Handle);
 	
 	/** Sets an ability spec to remove when its finished. If the spec is not currently active, it terminates it immediately. Also clears InputID of the Spec. */
-	void SetRemoveAbilityOnEnd(FGameplayAbilitySpecHandle AbilitySpecHandle);
+	UE_API void SetRemoveAbilityOnEnd(FGameplayAbilitySpecHandle AbilitySpecHandle);
 
 	/** 
 	 * Gets all Activatable Gameplay Abilities that match all tags in GameplayTagContainer AND for which
@@ -1039,7 +1041,7 @@ class GAMEPLAYABILITIES_API UAbilitySystemComponent : public UGameplayTasksCompo
 	 * For example, Behavior Trees can use various decorators to test an ability fetched using this mechanism as well
 	 * as the Task to execute the ability without needing to know that there even is more than one such ability.
 	 */
-	void GetActivatableGameplayAbilitySpecsByAllMatchingTags(const FGameplayTagContainer& GameplayTagContainer, TArray < struct FGameplayAbilitySpec* >& MatchingGameplayAbilities, bool bOnlyAbilitiesThatSatisfyTagRequirements = true) const;
+	UE_API void GetActivatableGameplayAbilitySpecsByAllMatchingTags(const FGameplayTagContainer& GameplayTagContainer, TArray < struct FGameplayAbilitySpec* >& MatchingGameplayAbilities, bool bOnlyAbilitiesThatSatisfyTagRequirements = true) const;
 
 	/** 
 	 * Attempts to activate every gameplay ability that matches the given tag and DoesAbilitySatisfyTagRequirements().
@@ -1047,7 +1049,7 @@ class GAMEPLAYABILITIES_API UAbilitySystemComponent : public UGameplayTasksCompo
 	 * If bAllowRemoteActivation is true, it will remotely activate local/server abilities, if false it will only try to locally activate abilities.
 	 */
 	UFUNCTION(BlueprintCallable, Category = "Abilities")
-	bool TryActivateAbilitiesByTag(const FGameplayTagContainer& GameplayTagContainer, bool bAllowRemoteActivation = true);
+	UE_API bool TryActivateAbilitiesByTag(const FGameplayTagContainer& GameplayTagContainer, bool bAllowRemoteActivation = true);
 
 	/**
 	 * Attempts to activate the ability that is passed in. This will check costs and requirements before doing so.
@@ -1055,7 +1057,7 @@ class GAMEPLAYABILITIES_API UAbilitySystemComponent : public UGameplayTasksCompo
 	 * If bAllowRemoteActivation is true, it will remotely activate local/server abilities, if false it will only try to locally activate the ability
 	 */
 	UFUNCTION(BlueprintCallable, Category = "Abilities")
-	bool TryActivateAbilityByClass(TSubclassOf<UGameplayAbility> InAbilityToActivate, bool bAllowRemoteActivation = true);
+	UE_API bool TryActivateAbilityByClass(TSubclassOf<UGameplayAbility> InAbilityToActivate, bool bAllowRemoteActivation = true);
 
 	/** 
 	 * Attempts to activate the given ability, will check costs and requirements before doing so.
@@ -1063,31 +1065,31 @@ class GAMEPLAYABILITIES_API UAbilitySystemComponent : public UGameplayTasksCompo
 	 * If bAllowRemoteActivation is true, it will remotely activate local/server abilities, if false it will only try to locally activate the ability
 	 */
 	UFUNCTION(BlueprintCallable, Category = "Abilities")
-	bool TryActivateAbility(FGameplayAbilitySpecHandle AbilityToActivate, bool bAllowRemoteActivation = true);
+	UE_API bool TryActivateAbility(FGameplayAbilitySpecHandle AbilityToActivate, bool bAllowRemoteActivation = true);
 
-	bool HasActivatableTriggeredAbility(FGameplayTag Tag);
+	UE_API bool HasActivatableTriggeredAbility(FGameplayTag Tag);
 
 	/** Triggers an ability from a gameplay event, will only trigger on local/server depending on execution flags */
-	bool TriggerAbilityFromGameplayEvent(FGameplayAbilitySpecHandle AbilityToTrigger, FGameplayAbilityActorInfo* ActorInfo, FGameplayTag Tag, const FGameplayEventData* Payload, UAbilitySystemComponent& Component);
+	UE_API bool TriggerAbilityFromGameplayEvent(FGameplayAbilitySpecHandle AbilityToTrigger, FGameplayAbilityActorInfo* ActorInfo, FGameplayTag Tag, const FGameplayEventData* Payload, UAbilitySystemComponent& Component);
 
 	// ----------------------------------------------------------------------------------------------------------------
 	// Ability Cancelling/Interrupts
 	// ----------------------------------------------------------------------------------------------------------------
 
 	/** Cancels the specified ability CDO. */
-	void CancelAbility(UGameplayAbility* Ability);	
+	UE_API void CancelAbility(UGameplayAbility* Ability);	
 
 	/** Cancels the ability indicated by passed in spec handle. If handle is not found among reactivated abilities nothing happens. */
-	void CancelAbilityHandle(const FGameplayAbilitySpecHandle& AbilityHandle);
+	UE_API void CancelAbilityHandle(const FGameplayAbilitySpecHandle& AbilityHandle);
 
 	/** Cancel all abilities with the specified tags. Will not cancel the Ignore instance */
-	void CancelAbilities(const FGameplayTagContainer* WithTags=nullptr, const FGameplayTagContainer* WithoutTags=nullptr, UGameplayAbility* Ignore=nullptr);
+	UE_API void CancelAbilities(const FGameplayTagContainer* WithTags=nullptr, const FGameplayTagContainer* WithoutTags=nullptr, UGameplayAbility* Ignore=nullptr);
 
 	/** Cancels all abilities regardless of tags. Will not cancel the ignore instance */
-	void CancelAllAbilities(UGameplayAbility* Ignore=nullptr);
+	UE_API void CancelAllAbilities(UGameplayAbility* Ignore=nullptr);
 
 	/** Cancels all abilities and kills any remaining instanced abilities */
-	virtual void DestroyActiveState();
+	UE_API virtual void DestroyActiveState();
 
 	/** 
 	 * Called from ability activation or native code, will apply the correct ability blocking tags and cancel existing abilities. Subclasses can override the behavior 
@@ -1099,24 +1101,24 @@ class GAMEPLAYABILITIES_API UAbilitySystemComponent : public UGameplayTasksCompo
 	 * @param bExecuteCancelTags If true will cancel abilities matching tags
 	 * @param CancelTags what tags to cancel
 	 */
-	virtual void ApplyAbilityBlockAndCancelTags(const FGameplayTagContainer& AbilityTags, UGameplayAbility* RequestingAbility, bool bEnableBlockTags, const FGameplayTagContainer& BlockTags, bool bExecuteCancelTags, const FGameplayTagContainer& CancelTags);
+	UE_API virtual void ApplyAbilityBlockAndCancelTags(const FGameplayTagContainer& AbilityTags, UGameplayAbility* RequestingAbility, bool bEnableBlockTags, const FGameplayTagContainer& BlockTags, bool bExecuteCancelTags, const FGameplayTagContainer& CancelTags);
 
 	/** Called when an ability is cancellable or not. Doesn't do anything by default, can be overridden to tie into gameplay events */
 	virtual void HandleChangeAbilityCanBeCanceled(const FGameplayTagContainer& AbilityTags, UGameplayAbility* RequestingAbility, bool bCanBeCanceled) {}
 
 	/** Returns true if any passed in tags are blocked */
-	virtual bool AreAbilityTagsBlocked(const FGameplayTagContainer& Tags) const;
+	UE_API virtual bool AreAbilityTagsBlocked(const FGameplayTagContainer& Tags) const;
 
 	/** Block or cancel blocking for specific ability tags */
-	void BlockAbilitiesWithTags(const FGameplayTagContainer& Tags);
-	void UnBlockAbilitiesWithTags(const FGameplayTagContainer& Tags);
+	UE_API void BlockAbilitiesWithTags(const FGameplayTagContainer& Tags);
+	UE_API void UnBlockAbilitiesWithTags(const FGameplayTagContainer& Tags);
 
 	/** Checks if the ability system is currently blocking InputID. Returns true if InputID is blocked, false otherwise.  */
-	bool IsAbilityInputBlocked(int32 InputID) const;
+	UE_API bool IsAbilityInputBlocked(int32 InputID) const;
 
 	/** Block or cancel blocking for specific input IDs */
-	void BlockAbilityByInputID(int32 InputID);
-	void UnBlockAbilityByInputID(int32 InputID);
+	UE_API void BlockAbilityByInputID(int32 InputID);
+	UE_API void UnBlockAbilityByInputID(int32 InputID);
 
 	// ----------------------------------------------------------------------------------------------------------------
 	// Functions meant to be called from GameplayAbility and subclasses, but not meant for general use
@@ -1138,11 +1140,11 @@ class GAMEPLAYABILITIES_API UAbilitySystemComponent : public UGameplayTasksCompo
 	float GetAbilityLastActivatedTime() const { return AbilityLastActivatedTime; }
 
 	/** Returns an ability spec from a handle. If modifying call MarkAbilitySpecDirty. Treat the return value as ephemeral as the pointer will potentially be invalidated on any subsequent call into AbilitySystemComponent. */
-	FGameplayAbilitySpec* FindAbilitySpecFromHandle(FGameplayAbilitySpecHandle Handle, EConsiderPending ConsiderPending = EConsiderPending::PendingRemove) const;
+	UE_API FGameplayAbilitySpec* FindAbilitySpecFromHandle(FGameplayAbilitySpecHandle Handle, EConsiderPending ConsiderPending = EConsiderPending::PendingRemove) const;
 	
 	/** Returns an ability spec from a GE handle. If modifying call MarkAbilitySpecDirty */
 	UE_DEPRECATED(5.3, "FindAbilitySpecFromGEHandle was never accurate because a GameplayEffect can grant multiple GameplayAbilities. It now returns nullptr.")
-	FGameplayAbilitySpec* FindAbilitySpecFromGEHandle(FActiveGameplayEffectHandle Handle) const;
+	UE_API FGameplayAbilitySpec* FindAbilitySpecFromGEHandle(FActiveGameplayEffectHandle Handle) const;
 
 	/**
 	* Returns all ability spec handles granted from a GE handle. Only the server may call this function.
@@ -1150,13 +1152,13 @@ class GAMEPLAYABILITIES_API UAbilitySystemComponent : public UGameplayTasksCompo
 	* @param Handle - The handle of the Active Gameplay Effect which granted the abilities we are looking for
 	* @param ConsiderPending - Are we returning AbilitySpecs that are pending for addition/removal?
 	*/
-	TArray<const FGameplayAbilitySpec*> FindAbilitySpecsFromGEHandle(const FScopedAbilityListLock& ScopeLock, FActiveGameplayEffectHandle Handle, EConsiderPending ConsiderPending = EConsiderPending::PendingRemove) const;
+	UE_API TArray<const FGameplayAbilitySpec*> FindAbilitySpecsFromGEHandle(const FScopedAbilityListLock& ScopeLock, FActiveGameplayEffectHandle Handle, EConsiderPending ConsiderPending = EConsiderPending::PendingRemove) const;
 
 	/** Returns an ability spec corresponding to given ability class. If modifying call MarkAbilitySpecDirty */
-	FGameplayAbilitySpec* FindAbilitySpecFromClass(TSubclassOf<UGameplayAbility> InAbilityClass) const;
+	UE_API FGameplayAbilitySpec* FindAbilitySpecFromClass(TSubclassOf<UGameplayAbility> InAbilityClass) const;
 
 	/** Returns an ability spec from a handle. If modifying call MarkAbilitySpecDirty */
-	FGameplayAbilitySpec* FindAbilitySpecFromInputID(int32 InputID) const;
+	UE_API FGameplayAbilitySpec* FindAbilitySpecFromInputID(int32 InputID) const;
 
 	/**
 	 * Returns all abilities with the given InputID
@@ -1164,12 +1166,12 @@ class GAMEPLAYABILITIES_API UAbilitySystemComponent : public UGameplayTasksCompo
 	 * @param InputID The Input ID to match
 	 * @param OutAbilitySpecs Array of pointers to matching specs
 	 */
-	virtual void FindAllAbilitySpecsFromInputID(int32 InputID, TArray<const FGameplayAbilitySpec*>& OutAbilitySpecs) const;
+	UE_API virtual void FindAllAbilitySpecsFromInputID(int32 InputID, TArray<const FGameplayAbilitySpec*>& OutAbilitySpecs) const;
 
 	/**
 	 * Build a simple FGameplayAbilitySpec from class, level and optional Input ID
 	 */
-	virtual FGameplayAbilitySpec BuildAbilitySpecFromClass(TSubclassOf<UGameplayAbility> AbilityClass, int32 Level = 0, int32 InputID = -1);
+	UE_API virtual FGameplayAbilitySpec BuildAbilitySpecFromClass(TSubclassOf<UGameplayAbility> AbilityClass, int32 Level = 0, int32 InputID = -1);
 
 	/**
 	 * Returns an array with all granted ability handles
@@ -1178,7 +1180,7 @@ class GAMEPLAYABILITIES_API UAbilitySystemComponent : public UGameplayTasksCompo
 	 * @param OutAbilityHandles This array will be filled with the granted Ability Spec Handles
 	 */
 	UFUNCTION(BlueprintCallable, BlueprintPure = false, Category = "Gameplay Abilities")
-	void GetAllAbilities(TArray<FGameplayAbilitySpecHandle>& OutAbilityHandles) const;
+	UE_API void GetAllAbilities(TArray<FGameplayAbilitySpecHandle>& OutAbilityHandles) const;
 
 	/**
 	 * Returns an array with all abilities that match the provided tags
@@ -1188,7 +1190,7 @@ class GAMEPLAYABILITIES_API UAbilitySystemComponent : public UGameplayTasksCompo
 	 * @param bExactMatch If true, tags must be matched exactly. Otherwise, abilities matching any of the tags will be returned
 	 */
 	UFUNCTION(BlueprintCallable, BlueprintPure = false, Category = "Gameplay Abilities")
-	void FindAllAbilitiesWithTags(TArray<FGameplayAbilitySpecHandle>& OutAbilityHandles, FGameplayTagContainer Tags, bool bExactMatch = true) const;
+	UE_API void FindAllAbilitiesWithTags(TArray<FGameplayAbilitySpecHandle>& OutAbilityHandles, FGameplayTagContainer Tags, bool bExactMatch = true) const;
 
 	/**
 	 * Returns an array with all abilities that match the provided Gameplay Tag Query
@@ -1197,7 +1199,7 @@ class GAMEPLAYABILITIES_API UAbilitySystemComponent : public UGameplayTasksCompo
 	 * @param Query Gameplay Tag Query to match
 	 */
 	UFUNCTION(BlueprintCallable, BlueprintPure = false, Category = "Gameplay Abilities")
-	void FindAllAbilitiesMatchingQuery(TArray<FGameplayAbilitySpecHandle>& OutAbilityHandles, FGameplayTagQuery Query) const;
+	UE_API void FindAllAbilitiesMatchingQuery(TArray<FGameplayAbilitySpecHandle>& OutAbilityHandles, FGameplayTagQuery Query) const;
 
 	/**
 	 * Returns an array with all abilities bound to an Input ID value
@@ -1206,28 +1208,28 @@ class GAMEPLAYABILITIES_API UAbilitySystemComponent : public UGameplayTasksCompo
 	 * @param InputID The Input ID to match
 	 */
 	UFUNCTION(BlueprintCallable, BlueprintPure = false, Category = "Gameplay Abilities")
-	void FindAllAbilitiesWithInputID(TArray<FGameplayAbilitySpecHandle>& OutAbilityHandles, int32 InputID = 0) const;
+	UE_API void FindAllAbilitiesWithInputID(TArray<FGameplayAbilitySpecHandle>& OutAbilityHandles, int32 InputID = 0) const;
 
 	/** Retrieves the EffectContext of the GameplayEffect of the active GameplayEffect. */
-	FGameplayEffectContextHandle GetEffectContextFromActiveGEHandle(FActiveGameplayEffectHandle Handle);
+	UE_API FGameplayEffectContextHandle GetEffectContextFromActiveGEHandle(FActiveGameplayEffectHandle Handle);
 
 	/** Call to mark that an ability spec has been modified */
-	void MarkAbilitySpecDirty(FGameplayAbilitySpec& Spec, bool WasAddOrRemove=false);
+	UE_API void MarkAbilitySpecDirty(FGameplayAbilitySpec& Spec, bool WasAddOrRemove=false);
 
 	/** Attempts to activate the given ability, will only work if called from the correct client/server context */
-	bool InternalTryActivateAbility(FGameplayAbilitySpecHandle AbilityToActivate, FPredictionKey InPredictionKey = FPredictionKey(), UGameplayAbility ** OutInstancedAbility = nullptr, FOnGameplayAbilityEnded::FDelegate* OnGameplayAbilityEndedDelegate = nullptr, const FGameplayEventData* TriggerEventData = nullptr);
+	UE_API bool InternalTryActivateAbility(FGameplayAbilitySpecHandle AbilityToActivate, FPredictionKey InPredictionKey = FPredictionKey(), UGameplayAbility ** OutInstancedAbility = nullptr, FOnGameplayAbilityEnded::FDelegate* OnGameplayAbilityEndedDelegate = nullptr, const FGameplayEventData* TriggerEventData = nullptr);
 
 	/** Failure tags used by InternalTryActivateAbility (E.g., this stores the  FailureTags of the last call to InternalTryActivateAbility */
 	FGameplayTagContainer InternalTryActivateAbilityFailureTags;
 
 	/** Called from the ability to let the component know it is ended */
-	virtual void NotifyAbilityEnded(FGameplayAbilitySpecHandle Handle, UGameplayAbility* Ability, bool bWasCancelled);
+	UE_API virtual void NotifyAbilityEnded(FGameplayAbilitySpecHandle Handle, UGameplayAbility* Ability, bool bWasCancelled);
 
-	void ClearAbilityReplicatedDataCache(FGameplayAbilitySpecHandle Handle, const FGameplayAbilityActivationInfo& ActivationInfo);
+	UE_API void ClearAbilityReplicatedDataCache(FGameplayAbilitySpecHandle Handle, const FGameplayAbilityActivationInfo& ActivationInfo);
 
 	/** Called from FScopedAbilityListLock */
-	void IncrementAbilityListLock();
-	void DecrementAbilityListLock();
+	UE_API void IncrementAbilityListLock();
+	UE_API void DecrementAbilityListLock();
 
 	// ----------------------------------------------------------------------------------------------------------------
 	// Debugging
@@ -1262,86 +1264,86 @@ class GAMEPLAYABILITIES_API UAbilitySystemComponent : public UGameplayTasksCompo
 		int32 GameFlags; // arbitrary flags for games to set/read in Debug_Internal
 	};
 
-	static void OnShowDebugInfo(AHUD* HUD, UCanvas* Canvas, const FDebugDisplayInfo& DisplayInfo, float& YL, float& YPos);
+	static UE_API void OnShowDebugInfo(AHUD* HUD, UCanvas* Canvas, const FDebugDisplayInfo& DisplayInfo, float& YL, float& YPos);
 
-	virtual void DisplayDebug(class UCanvas* Canvas, const class FDebugDisplayInfo& DebugDisplay, float& YL, float& YPos);
-	virtual void PrintDebug();
+	UE_API virtual void DisplayDebug(class UCanvas* Canvas, const class FDebugDisplayInfo& DebugDisplay, float& YL, float& YPos);
+	UE_API virtual void PrintDebug();
 
-	void AccumulateScreenPos(FAbilitySystemComponentDebugInfo& Info);
-	virtual void Debug_Internal(struct FAbilitySystemComponentDebugInfo& Info);
-	void DebugLine(struct FAbilitySystemComponentDebugInfo& Info, FString Str, float XOffset, float YOffset, int32 MinTextRowsToAdvance = 0);
-	FString CleanupName(FString Str);
+	UE_API void AccumulateScreenPos(FAbilitySystemComponentDebugInfo& Info);
+	UE_API virtual void Debug_Internal(struct FAbilitySystemComponentDebugInfo& Info);
+	UE_API void DebugLine(struct FAbilitySystemComponentDebugInfo& Info, FString Str, float XOffset, float YOffset, int32 MinTextRowsToAdvance = 0);
+	UE_API FString CleanupName(FString Str);
 
 	/** Print a debug list of all gameplay effects */
-	void PrintAllGameplayEffects() const;
+	UE_API void PrintAllGameplayEffects() const;
 
 	/** Ask the server to send ability system debug information back to the client, via ClientPrintDebug_Response  */
 	UFUNCTION(Server, reliable, WithValidation)
-	void ServerPrintDebug_Request();
+	UE_API void ServerPrintDebug_Request();
 
 	/** Same as ServerPrintDebug_Request but this includes the client debug strings so that the server can embed them in replays */
 	UFUNCTION(Server, reliable, WithValidation)
-	void ServerPrintDebug_RequestWithStrings(const TArray<FString>& Strings);
+	UE_API void ServerPrintDebug_RequestWithStrings(const TArray<FString>& Strings);
 
 	/** Virtual function games can override to do their own stuff when either ServerPrintDebug function runs on the server */
-	virtual void OnServerPrintDebug_Request();
+	UE_API virtual void OnServerPrintDebug_Request();
 
 	/** Determines whether to call ServerPrintDebug_Request or ServerPrintDebug_RequestWithStrings.   */
-	virtual bool ShouldSendClientDebugStringsToServer() const;
+	UE_API virtual bool ShouldSendClientDebugStringsToServer() const;
 
 	UFUNCTION(Client, reliable)
-	void ClientPrintDebug_Response(const TArray<FString>& Strings, int32 GameFlags);
-	virtual void OnClientPrintDebug_Response(const TArray<FString>& Strings, int32 GameFlags);
+	UE_API void ClientPrintDebug_Response(const TArray<FString>& Strings, int32 GameFlags);
+	UE_API virtual void OnClientPrintDebug_Response(const TArray<FString>& Strings, int32 GameFlags);
 
 #if ENABLE_VISUAL_LOG
-	void ClearDebugInstantEffects();
+	UE_API void ClearDebugInstantEffects();
 	
-	virtual void GrabDebugSnapshot(FVisualLogEntry* Snapshot) const override;
+	UE_API virtual void GrabDebugSnapshot(FVisualLogEntry* Snapshot) const override;
 #endif // ENABLE_VISUAL_LOG
 
 	UE_DEPRECATED(4.26, "This will be made private in future engine versions. Use SetClientDebugStrings, GetClientDebugStrings, or GetClientDebugStrings_Mutable instead.")
 	UPROPERTY(ReplicatedUsing=OnRep_ClientDebugString)
 	TArray<FString>	ClientDebugStrings;
 
-	void SetClientDebugStrings(TArray<FString>&& NewClientDebugStrings);
-	TArray<FString>& GetClientDebugStrings_Mutable();
-	const TArray<FString>& GetClientDebugStrings() const;
+	UE_API void SetClientDebugStrings(TArray<FString>&& NewClientDebugStrings);
+	UE_API TArray<FString>& GetClientDebugStrings_Mutable();
+	UE_API const TArray<FString>& GetClientDebugStrings() const;
 
 	UE_DEPRECATED(4.26, "This will be made private in future engine versions. Use SetServerDebugStrings, GetServerDebugStrings, or GetServerDebugStrings_Mutable instead.")
 	UPROPERTY(ReplicatedUsing=OnRep_ServerDebugString)
 	TArray<FString>	ServerDebugStrings;
 
-	void SetServerDebugStrings(TArray<FString>&& NewServerDebugStrings);
-	TArray<FString>& GetServerDebugStrings_Mutable();
-	const TArray<FString>& GetServerDebugStrings() const;
+	UE_API void SetServerDebugStrings(TArray<FString>&& NewServerDebugStrings);
+	UE_API TArray<FString>& GetServerDebugStrings_Mutable();
+	UE_API const TArray<FString>& GetServerDebugStrings() const;
 
 	UFUNCTION()
-	virtual void OnRep_ClientDebugString();
+	UE_API virtual void OnRep_ClientDebugString();
 
 	UFUNCTION()
-	virtual void OnRep_ServerDebugString();
+	UE_API virtual void OnRep_ServerDebugString();
 
 	// ----------------------------------------------------------------------------------------------------------------
 	// Batching client->server RPCs
 	// This is a WIP feature to batch up client->server communication. It is opt in and not complete. It only batches the below functions. Other Server RPCs are not safe to call during a batch window. Only opt in if you know what you are doing!
 	// ----------------------------------------------------------------------------------------------------------------	
 
-	void CallServerTryActivateAbility(FGameplayAbilitySpecHandle AbilityToActivate, bool InputPressed, FPredictionKey PredictionKey);
-	void CallServerSetReplicatedTargetData(FGameplayAbilitySpecHandle AbilityHandle, FPredictionKey AbilityOriginalPredictionKey, const FGameplayAbilityTargetDataHandle& ReplicatedTargetDataHandle, FGameplayTag ApplicationTag, FPredictionKey CurrentPredictionKey);
-	void CallServerEndAbility(FGameplayAbilitySpecHandle AbilityToEnd, FGameplayAbilityActivationInfo ActivationInfo, FPredictionKey PredictionKey);
+	UE_API void CallServerTryActivateAbility(FGameplayAbilitySpecHandle AbilityToActivate, bool InputPressed, FPredictionKey PredictionKey);
+	UE_API void CallServerSetReplicatedTargetData(FGameplayAbilitySpecHandle AbilityHandle, FPredictionKey AbilityOriginalPredictionKey, const FGameplayAbilityTargetDataHandle& ReplicatedTargetDataHandle, FGameplayTag ApplicationTag, FPredictionKey CurrentPredictionKey);
+	UE_API void CallServerEndAbility(FGameplayAbilitySpecHandle AbilityToEnd, FGameplayAbilityActivationInfo ActivationInfo, FPredictionKey PredictionKey);
 
 	virtual bool ShouldDoServerAbilityRPCBatch() const { return false; }
-	virtual void BeginServerAbilityRPCBatch(FGameplayAbilitySpecHandle AbilityHandle);
-	virtual void EndServerAbilityRPCBatch(FGameplayAbilitySpecHandle AbilityHandle);
+	UE_API virtual void BeginServerAbilityRPCBatch(FGameplayAbilitySpecHandle AbilityHandle);
+	UE_API virtual void EndServerAbilityRPCBatch(FGameplayAbilitySpecHandle AbilityHandle);
 
 	/** Accumulated client side data that is batched out to server on EndServerAbilityRPCBatch */
 	TArray<FServerAbilityRPCBatch, TInlineAllocator<1> > LocalServerAbilityRPCBatchData;
 
 	UFUNCTION(Server, reliable, WithValidation)
-	void	ServerAbilityRPCBatch(FServerAbilityRPCBatch BatchInfo);
+	UE_API void	ServerAbilityRPCBatch(FServerAbilityRPCBatch BatchInfo);
 
 	// Overridable function for sub classes
-	virtual void ServerAbilityRPCBatch_Internal(FServerAbilityRPCBatch& BatchInfo);
+	UE_API virtual void ServerAbilityRPCBatch_Internal(FServerAbilityRPCBatch& BatchInfo);
 
 	// ----------------------------------------------------------------------------------------------------------------
 	//	Input handling/targeting
@@ -1353,11 +1355,11 @@ class GAMEPLAYABILITIES_API UAbilitySystemComponent : public UGameplayTasksCompo
 	 * This should not be used to game mechanics like silences or disables. Those should be done through gameplay effects.
 	 */
 	UFUNCTION(BlueprintCallable, Category="Abilities")
-	bool GetUserAbilityActivationInhibited() const;
+	UE_API bool GetUserAbilityActivationInhibited() const;
 	
 	/** Disable or Enable a local user from being able to activate abilities. This should only be used for input/UI etc related inhibition. Do not use for game mechanics. */
 	UFUNCTION(BlueprintCallable, Category="Abilities")
-	virtual void SetUserAbilityActivationInhibited(bool NewInhibit);
+	UE_API virtual void SetUserAbilityActivationInhibited(bool NewInhibit);
 
 	/** Rather activation is currently inhibited */
 	UPROPERTY()
@@ -1380,17 +1382,17 @@ class GAMEPLAYABILITIES_API UAbilitySystemComponent : public UGameplayTasksCompo
 	TArray<TObjectPtr<AGameplayAbilityTargetActor>> SpawnedTargetActors;
 
 	/** Bind to an input component with some default action names */
-	virtual void BindToInputComponent(UInputComponent* InputComponent);
+	UE_API virtual void BindToInputComponent(UInputComponent* InputComponent);
 
 	/** Bind to an input component with customized bindings */
-	virtual void BindAbilityActivationToInputComponent(UInputComponent* InputComponent, FGameplayAbilityInputBinds BindInfo);
+	UE_API virtual void BindAbilityActivationToInputComponent(UInputComponent* InputComponent, FGameplayAbilityInputBinds BindInfo);
 
 	/** Initializes BlockedAbilityBindings variable */
-	virtual void SetBlockAbilityBindingsArray(FGameplayAbilityInputBinds BindInfo);
+	UE_API virtual void SetBlockAbilityBindingsArray(FGameplayAbilityInputBinds BindInfo);
 
 	/** Called to handle ability bind input */
-	virtual void AbilityLocalInputPressed(int32 InputID);
-	virtual void AbilityLocalInputReleased(int32 InputID);
+	UE_API virtual void AbilityLocalInputPressed(int32 InputID);
+	UE_API virtual void AbilityLocalInputReleased(int32 InputID);
 
 	/*
 	 * Sends a local player Input Pressed event with the provided Input ID, notifying any bound abilities
@@ -1398,30 +1400,30 @@ class GAMEPLAYABILITIES_API UAbilitySystemComponent : public UGameplayTasksCompo
 	 * @param InputID The Input ID to match
 	 */
 	UFUNCTION(BlueprintCallable, Category = "Gameplay Abilities")
-	void PressInputID(int32 InputID);
+	UE_API void PressInputID(int32 InputID);
 
 	/**
 	 * Sends a local player Input Released event with the provided Input ID, notifying any bound abilities
 	 * @param InputID The Input ID to match
 	 */
 	UFUNCTION(BlueprintCallable, Category = "Gameplay Abilities")
-	void ReleaseInputID(int32 InputID);
+	UE_API void ReleaseInputID(int32 InputID);
 
 	/** Handle confirm/cancel for target actors */
-	virtual void LocalInputConfirm();
-	virtual void LocalInputCancel();
+	UE_API virtual void LocalInputConfirm();
+	UE_API virtual void LocalInputCancel();
 
 	/**
 	 * Sends a local player Input Confirm event, notifying abilities
 	 */
 	UFUNCTION(BlueprintCallable, Category = "Gameplay Abilities")
-	void InputConfirm();
+	UE_API void InputConfirm();
 
 	/**
 	 * Sends a local player Input Cancel event, notifying abilities
 	 */
 	UFUNCTION(BlueprintCallable, Category = "Gameplay Abilities")
-	void InputCancel();
+	UE_API void InputCancel();
 	
 	/** InputID for binding GenericConfirm/Cancel events */
 	int32 GenericConfirmInputID;
@@ -1438,71 +1440,71 @@ class GAMEPLAYABILITIES_API UAbilitySystemComponent : public UGameplayTasksCompo
 
 	/** Any active targeting actors will be told to stop and return current targeting data */
 	UFUNCTION(BlueprintCallable, Category = "Abilities")
-	virtual void TargetConfirm();
+	UE_API virtual void TargetConfirm();
 
 	/** Any active targeting actors will be stopped and canceled, not returning any targeting data */
 	UFUNCTION(BlueprintCallable, Category = "Abilities")
-	virtual void TargetCancel();
+	UE_API virtual void TargetCancel();
 
 	// ----------------------------------------------------------------------------------------------------------------
 	//	AnimMontage Support
 	// ----------------------------------------------------------------------------------------------------------------	
 
 	/** Plays a montage and handles replication and prediction based on passed in ability/activation info */
-	virtual float PlayMontage(UGameplayAbility* AnimatingAbility, FGameplayAbilityActivationInfo ActivationInfo, UAnimMontage* Montage, float InPlayRate, FName StartSectionName = NAME_None, float StartTimeSeconds = 0.0f);
+	UE_API virtual float PlayMontage(UGameplayAbility* AnimatingAbility, FGameplayAbilityActivationInfo ActivationInfo, UAnimMontage* Montage, float InPlayRate, FName StartSectionName = NAME_None, float StartTimeSeconds = 0.0f);
 
-	virtual UAnimMontage* PlaySlotAnimationAsDynamicMontage(UGameplayAbility* AnimatingAbility, FGameplayAbilityActivationInfo ActivationInfo, UAnimSequenceBase* AnimAsset, FName SlotName, float BlendInTime, float BlendOutTime, float InPlayRate = 1.f, float StartTimeSeconds = 0.0f);
+	UE_API virtual UAnimMontage* PlaySlotAnimationAsDynamicMontage(UGameplayAbility* AnimatingAbility, FGameplayAbilityActivationInfo ActivationInfo, UAnimSequenceBase* AnimAsset, FName SlotName, float BlendInTime, float BlendOutTime, float InPlayRate = 1.f, float StartTimeSeconds = 0.0f);
 
-	virtual UAnimMontage* PlaySlotAnimationAsDynamicMontage_WithFractionalLoops(UGameplayAbility* AnimatingAbility, FGameplayAbilityActivationInfo ActivationInfo, UAnimSequenceBase* AnimAsset, FName SlotName, float BlendInTime, float BlendOutTime, float InPlayRate = 1.f, float StartTimeSeconds = 0.0f, float PlayCount = 1.0f);
+	UE_API virtual UAnimMontage* PlaySlotAnimationAsDynamicMontage_WithFractionalLoops(UGameplayAbility* AnimatingAbility, FGameplayAbilityActivationInfo ActivationInfo, UAnimSequenceBase* AnimAsset, FName SlotName, float BlendInTime, float BlendOutTime, float InPlayRate = 1.f, float StartTimeSeconds = 0.0f, float PlayCount = 1.0f);
 
 	/** Plays a montage without updating replication/prediction structures. Used by simulated proxies when replication tells them to play a montage. */
-	virtual float PlayMontageSimulated(UAnimMontage* Montage, float InPlayRate, FName StartSectionName = NAME_None);
+	UE_API virtual float PlayMontageSimulated(UAnimMontage* Montage, float InPlayRate, FName StartSectionName = NAME_None);
 	
-	virtual UAnimMontage* PlaySlotAnimationAsDynamicMontageSimulated(UAnimSequenceBase* AnimAsset, FName SlotName, float BlendInTime, float BlendOutTime, float InPlayRate = 1.f);
+	UE_API virtual UAnimMontage* PlaySlotAnimationAsDynamicMontageSimulated(UAnimSequenceBase* AnimAsset, FName SlotName, float BlendInTime, float BlendOutTime, float InPlayRate = 1.f);
 
-	virtual UAnimMontage* PlaySlotAnimationAsDynamicMontage_WithFractionalLoopsSimulated(UAnimSequenceBase* AnimAsset, FName SlotName, float BlendInTime, float BlendOutTime, float InPlayRate = 1.f, float PlayCount = 1.0f);
+	UE_API virtual UAnimMontage* PlaySlotAnimationAsDynamicMontage_WithFractionalLoopsSimulated(UAnimSequenceBase* AnimAsset, FName SlotName, float BlendInTime, float BlendOutTime, float InPlayRate = 1.f, float PlayCount = 1.0f);
 
 	/** Stops whatever montage is currently playing. Expectation is caller should only be stopping it if they are the current animating ability (or have good reason not to check) */
-	virtual void CurrentMontageStop(float OverrideBlendOutTime = -1.0f);
+	UE_API virtual void CurrentMontageStop(float OverrideBlendOutTime = -1.0f);
 
 	/** Stops current montage if it's the one given as the Montage param */
-	virtual void StopMontageIfCurrent(const UAnimMontage& Montage, float OverrideBlendOutTime = -1.0f);
+	UE_API virtual void StopMontageIfCurrent(const UAnimMontage& Montage, float OverrideBlendOutTime = -1.0f);
 
 	/** Clear the animating ability that is passed in, if it's still currently animating */
-	virtual void ClearAnimatingAbility(UGameplayAbility* Ability);
+	UE_API virtual void ClearAnimatingAbility(UGameplayAbility* Ability);
 
 	/** Jumps current montage to given section. Expectation is caller should only be stopping it if they are the current animating ability (or have good reason not to check) */
-	virtual void CurrentMontageJumpToSection(FName SectionName);
+	UE_API virtual void CurrentMontageJumpToSection(FName SectionName);
 
 	/** Sets current montages next section name. Expectation is caller should only be stopping it if they are the current animating ability (or have good reason not to check) */
-	virtual void CurrentMontageSetNextSectionName(FName FromSectionName, FName ToSectionName);
+	UE_API virtual void CurrentMontageSetNextSectionName(FName FromSectionName, FName ToSectionName);
 
 	/** Sets current montage's play rate */
-	virtual void CurrentMontageSetPlayRate(float InPlayRate);
+	UE_API virtual void CurrentMontageSetPlayRate(float InPlayRate);
 
 	/** Returns true if the passed in ability is the current animating ability */
-	bool IsAnimatingAbility(UGameplayAbility* Ability) const;
+	UE_API bool IsAnimatingAbility(UGameplayAbility* Ability) const;
 
 	/** Returns the current animating ability */
-	UGameplayAbility* GetAnimatingAbility() const;
+	UE_API UGameplayAbility* GetAnimatingAbility() const;
 
 	/** Returns montage that is currently playing */
-	UAnimMontage* GetCurrentMontage() const;
+	UE_API UAnimMontage* GetCurrentMontage() const;
 
 	/** Get SectionID of currently playing AnimMontage */
-	int32 GetCurrentMontageSectionID() const;
+	UE_API int32 GetCurrentMontageSectionID() const;
 
 	/** Get SectionName of currently playing AnimMontage */
-	FName GetCurrentMontageSectionName() const;
+	UE_API FName GetCurrentMontageSectionName() const;
 
 	/** Get length in time of current section */
-	float GetCurrentMontageSectionLength() const;
+	UE_API float GetCurrentMontageSectionLength() const;
 
 	/** Returns amount of time left in current section */
-	float GetCurrentMontageSectionTimeLeft() const;
+	UE_API float GetCurrentMontageSectionTimeLeft() const;
 
 	/** Method to set the replication method for the position in the montage */
-	void SetMontageRepAnimPositionMethod(ERepAnimPositionMethod InMethod);
+	UE_API void SetMontageRepAnimPositionMethod(ERepAnimPositionMethod InMethod);
 
 	// ----------------------------------------------------------------------------------------------------------------
 	//	Actor interaction
@@ -1520,23 +1522,23 @@ private:
 
 public:
 
-	void SetOwnerActor(AActor* NewOwnerActor);
+	UE_API void SetOwnerActor(AActor* NewOwnerActor);
 	AActor* GetOwnerActor() const { return OwnerActor; }
 
-	void SetAvatarActor_Direct(AActor* NewAvatarActor);
+	UE_API void SetAvatarActor_Direct(AActor* NewAvatarActor);
 	AActor* GetAvatarActor_Direct() const { return AvatarActor; }
 	
 	UFUNCTION()
-	void OnRep_OwningActor();
+	UE_API void OnRep_OwningActor();
 
 	UFUNCTION()
-	void OnAvatarActorDestroyed(AActor* InActor);
+	UE_API void OnAvatarActorDestroyed(AActor* InActor);
 
 	UFUNCTION()
-	void OnOwnerActorDestroyed(AActor* InActor);
+	UE_API void OnOwnerActorDestroyed(AActor* InActor);
 
 	UFUNCTION()
-	void OnSpawnedAttributesEndPlayed(AActor* InActor, EEndPlayReason::Type EndPlayReason);
+	UE_API void OnSpawnedAttributesEndPlayed(AActor* InActor, EEndPlayReason::Type EndPlayReason);
 
 	/** Cached off data about the owning actor that abilities will need to frequently access (movement component, mesh component, anim instance, etc) */
 	TSharedPtr<FGameplayAbilityActorInfo>	AbilityActorInfo;
@@ -1546,16 +1548,16 @@ public:
 	 *      OwnerActor is the actor that logically owns this component.
 	 *		AvatarActor is what physical actor in the world we are acting on. Usually a Pawn but it could be a Tower, Building, Turret, etc, may be the same as Owner
 	 */
-	virtual void InitAbilityActorInfo(AActor* InOwnerActor, AActor* InAvatarActor);
+	UE_API virtual void InitAbilityActorInfo(AActor* InOwnerActor, AActor* InAvatarActor);
 
 	/** Returns avatar actor to be used for a specific task, normally GetAvatarActor */
-	virtual AActor* GetGameplayTaskAvatar(const UGameplayTask* Task) const override;
+	UE_API virtual AActor* GetGameplayTaskAvatar(const UGameplayTask* Task) const override;
 
 	/** Returns the avatar actor for this component */
-	AActor* GetAvatarActor() const;
+	UE_API AActor* GetAvatarActor() const;
 
 	/** Changes the avatar actor, leaves the owner actor the same */
-	void SetAvatarActor(AActor* InAvatarActor);
+	UE_API void SetAvatarActor(AActor* InAvatarActor);
 
 	/** called when the ASC's AbilityActorInfo has a PlayerController set. */
 	virtual void OnPlayerControllerSet() { }
@@ -1563,13 +1565,13 @@ public:
 	/**
 	* This is called when the actor that is initialized to this system dies, this will clear that actor from this system and FGameplayAbilityActorInfo
 	*/
-	virtual void ClearActorInfo();
+	UE_API virtual void ClearActorInfo();
 
 	/**
 	 *	This will refresh the Ability's ActorInfo structure based on the current ActorInfo. That is, AvatarActor will be the same but we will look for new
 	 *	AnimInstance, MovementComponent, PlayerController, etc.
 	 */	
-	void RefreshAbilityActorInfo();
+	UE_API void RefreshAbilityActorInfo();
 
 	// ----------------------------------------------------------------------------------------------------------------
 	//	Synchronization RPCs
@@ -1578,107 +1580,107 @@ public:
 	
 	/** Replicates the Generic Replicated Event to the server. */
 	UFUNCTION(Server, reliable, WithValidation)
-	void ServerSetReplicatedEvent(EAbilityGenericReplicatedEvent::Type EventType, FGameplayAbilitySpecHandle AbilityHandle, FPredictionKey AbilityOriginalPredictionKey, FPredictionKey CurrentPredictionKey);
+	UE_API void ServerSetReplicatedEvent(EAbilityGenericReplicatedEvent::Type EventType, FGameplayAbilitySpecHandle AbilityHandle, FPredictionKey AbilityOriginalPredictionKey, FPredictionKey CurrentPredictionKey);
 
 	/** Replicates the Generic Replicated Event to the server with payload. */
 	UFUNCTION(Server, reliable, WithValidation)
-	void ServerSetReplicatedEventWithPayload(EAbilityGenericReplicatedEvent::Type EventType, FGameplayAbilitySpecHandle AbilityHandle, FPredictionKey AbilityOriginalPredictionKey, FPredictionKey CurrentPredictionKey, FVector_NetQuantize100 VectorPayload);
+	UE_API void ServerSetReplicatedEventWithPayload(EAbilityGenericReplicatedEvent::Type EventType, FGameplayAbilitySpecHandle AbilityHandle, FPredictionKey AbilityOriginalPredictionKey, FPredictionKey CurrentPredictionKey, FVector_NetQuantize100 VectorPayload);
 
 	/** Replicates the Generic Replicated Event to the client. */
 	UFUNCTION(Client, reliable)
-	void ClientSetReplicatedEvent(EAbilityGenericReplicatedEvent::Type EventType, FGameplayAbilitySpecHandle AbilityHandle, FPredictionKey AbilityOriginalPredictionKey);
+	UE_API void ClientSetReplicatedEvent(EAbilityGenericReplicatedEvent::Type EventType, FGameplayAbilitySpecHandle AbilityHandle, FPredictionKey AbilityOriginalPredictionKey);
 
 	/** Calls local callbacks that are registered with the given Generic Replicated Event */
-	bool InvokeReplicatedEvent(EAbilityGenericReplicatedEvent::Type EventType, FGameplayAbilitySpecHandle AbilityHandle, FPredictionKey AbilityOriginalPredictionKey, FPredictionKey CurrentPredictionKey = FPredictionKey());
+	UE_API bool InvokeReplicatedEvent(EAbilityGenericReplicatedEvent::Type EventType, FGameplayAbilitySpecHandle AbilityHandle, FPredictionKey AbilityOriginalPredictionKey, FPredictionKey CurrentPredictionKey = FPredictionKey());
 
 	/** Calls local callbacks that are registered with the given Generic Replicated Event */
-	bool InvokeReplicatedEventWithPayload(EAbilityGenericReplicatedEvent::Type EventType, FGameplayAbilitySpecHandle AbilityHandle, FPredictionKey AbilityOriginalPredictionKey, FPredictionKey CurrentPredictionKey, FVector_NetQuantize100 VectorPayload);
+	UE_API bool InvokeReplicatedEventWithPayload(EAbilityGenericReplicatedEvent::Type EventType, FGameplayAbilitySpecHandle AbilityHandle, FPredictionKey AbilityOriginalPredictionKey, FPredictionKey CurrentPredictionKey, FVector_NetQuantize100 VectorPayload);
 	
 	/** Replicates targeting data to the server */
 	UFUNCTION(Server, reliable, WithValidation)
-	void ServerSetReplicatedTargetData(FGameplayAbilitySpecHandle AbilityHandle, FPredictionKey AbilityOriginalPredictionKey, const FGameplayAbilityTargetDataHandle& ReplicatedTargetDataHandle, FGameplayTag ApplicationTag, FPredictionKey CurrentPredictionKey);
+	UE_API void ServerSetReplicatedTargetData(FGameplayAbilitySpecHandle AbilityHandle, FPredictionKey AbilityOriginalPredictionKey, const FGameplayAbilityTargetDataHandle& ReplicatedTargetDataHandle, FGameplayTag ApplicationTag, FPredictionKey CurrentPredictionKey);
 
 	/** Replicates to the server that targeting has been cancelled */
 	UFUNCTION(Server, reliable, WithValidation)
-	void ServerSetReplicatedTargetDataCancelled(FGameplayAbilitySpecHandle AbilityHandle, FPredictionKey AbilityOriginalPredictionKey, FPredictionKey CurrentPredictionKey);
+	UE_API void ServerSetReplicatedTargetDataCancelled(FGameplayAbilitySpecHandle AbilityHandle, FPredictionKey AbilityOriginalPredictionKey, FPredictionKey CurrentPredictionKey);
 
 	/** Sets the current target data and calls applicable callbacks */
-	virtual void ConfirmAbilityTargetData(FGameplayAbilitySpecHandle AbilityHandle, FPredictionKey AbilityOriginalPredictionKey, const FGameplayAbilityTargetDataHandle& TargetData, const FGameplayTag& ApplicationTag);
+	UE_API virtual void ConfirmAbilityTargetData(FGameplayAbilitySpecHandle AbilityHandle, FPredictionKey AbilityOriginalPredictionKey, const FGameplayAbilityTargetDataHandle& TargetData, const FGameplayTag& ApplicationTag);
 
 	/** Cancels the ability target data and calls callbacks */
-	virtual void CancelAbilityTargetData(FGameplayAbilitySpecHandle AbilityHandle, FPredictionKey AbilityOriginalPredictionKey);
+	UE_API virtual void CancelAbilityTargetData(FGameplayAbilitySpecHandle AbilityHandle, FPredictionKey AbilityOriginalPredictionKey);
 
 	/** Deletes all cached ability client data (Was: ConsumeAbilityTargetData)*/
-	void ConsumeAllReplicatedData(FGameplayAbilitySpecHandle AbilityHandle, FPredictionKey AbilityOriginalPredictionKey);
+	UE_API void ConsumeAllReplicatedData(FGameplayAbilitySpecHandle AbilityHandle, FPredictionKey AbilityOriginalPredictionKey);
 	/** Consumes cached TargetData from client (only TargetData) */
-	void ConsumeClientReplicatedTargetData(FGameplayAbilitySpecHandle AbilityHandle, FPredictionKey AbilityOriginalPredictionKey);
+	UE_API void ConsumeClientReplicatedTargetData(FGameplayAbilitySpecHandle AbilityHandle, FPredictionKey AbilityOriginalPredictionKey);
 
 	/** Consumes the given Generic Replicated Event (unsets it). */
-	void ConsumeGenericReplicatedEvent(EAbilityGenericReplicatedEvent::Type EventType, FGameplayAbilitySpecHandle AbilityHandle, FPredictionKey AbilityOriginalPredictionKey);
+	UE_API void ConsumeGenericReplicatedEvent(EAbilityGenericReplicatedEvent::Type EventType, FGameplayAbilitySpecHandle AbilityHandle, FPredictionKey AbilityOriginalPredictionKey);
 
 	/** Gets replicated data of the given Generic Replicated Event. */
-	FAbilityReplicatedData GetReplicatedDataOfGenericReplicatedEvent(EAbilityGenericReplicatedEvent::Type EventType, FGameplayAbilitySpecHandle AbilityHandle, FPredictionKey AbilityOriginalPredictionKey);
+	UE_API FAbilityReplicatedData GetReplicatedDataOfGenericReplicatedEvent(EAbilityGenericReplicatedEvent::Type EventType, FGameplayAbilitySpecHandle AbilityHandle, FPredictionKey AbilityOriginalPredictionKey);
 	
 	/** Calls any Replicated delegates that have been sent (TargetData or Generic Replicated Events). Note this can be dangerous if multiple places in an ability register events and then call this function. */
-	void CallAllReplicatedDelegatesIfSet(FGameplayAbilitySpecHandle AbilityHandle, FPredictionKey AbilityOriginalPredictionKey);
+	UE_API void CallAllReplicatedDelegatesIfSet(FGameplayAbilitySpecHandle AbilityHandle, FPredictionKey AbilityOriginalPredictionKey);
 
 	/** Calls the TargetData Confirm/Cancel events if they have been sent. */
-	bool CallReplicatedTargetDataDelegatesIfSet(FGameplayAbilitySpecHandle AbilityHandle, FPredictionKey AbilityOriginalPredictionKey);
+	UE_API bool CallReplicatedTargetDataDelegatesIfSet(FGameplayAbilitySpecHandle AbilityHandle, FPredictionKey AbilityOriginalPredictionKey);
 
 	/** Calls a given Generic Replicated Event delegate if the event has already been sent */
-	bool CallReplicatedEventDelegateIfSet(EAbilityGenericReplicatedEvent::Type EventType, FGameplayAbilitySpecHandle AbilityHandle, FPredictionKey AbilityOriginalPredictionKey);
+	UE_API bool CallReplicatedEventDelegateIfSet(EAbilityGenericReplicatedEvent::Type EventType, FGameplayAbilitySpecHandle AbilityHandle, FPredictionKey AbilityOriginalPredictionKey);
 
 	/** Calls passed in delegate if the Client Event has already been sent. If not, it adds the delegate to our multicast callback that will fire when it does. */
-	bool CallOrAddReplicatedDelegate(EAbilityGenericReplicatedEvent::Type EventType, FGameplayAbilitySpecHandle AbilityHandle, FPredictionKey AbilityOriginalPredictionKey, FSimpleMulticastDelegate::FDelegate Delegate);
+	UE_API bool CallOrAddReplicatedDelegate(EAbilityGenericReplicatedEvent::Type EventType, FGameplayAbilitySpecHandle AbilityHandle, FPredictionKey AbilityOriginalPredictionKey, FSimpleMulticastDelegate::FDelegate Delegate);
 
 	/** Returns TargetDataSet delegate for a given Ability/PredictionKey pair */
-	FAbilityTargetDataSetDelegate& AbilityTargetDataSetDelegate(FGameplayAbilitySpecHandle AbilityHandle, FPredictionKey AbilityOriginalPredictionKey);
+	UE_API FAbilityTargetDataSetDelegate& AbilityTargetDataSetDelegate(FGameplayAbilitySpecHandle AbilityHandle, FPredictionKey AbilityOriginalPredictionKey);
 
 	/** Returns TargetData Cancelled delegate for a given Ability/PredictionKey pair */
-	FSimpleMulticastDelegate& AbilityTargetDataCancelledDelegate(FGameplayAbilitySpecHandle AbilityHandle, FPredictionKey AbilityOriginalPredictionKey);
+	UE_API FSimpleMulticastDelegate& AbilityTargetDataCancelledDelegate(FGameplayAbilitySpecHandle AbilityHandle, FPredictionKey AbilityOriginalPredictionKey);
 
 	/** Returns Generic Replicated Event for a given Ability/PredictionKey pair */
-	FSimpleMulticastDelegate& AbilityReplicatedEventDelegate(EAbilityGenericReplicatedEvent::Type EventType, FGameplayAbilitySpecHandle AbilityHandle, FPredictionKey AbilityOriginalPredictionKey);
+	UE_API FSimpleMulticastDelegate& AbilityReplicatedEventDelegate(EAbilityGenericReplicatedEvent::Type EventType, FGameplayAbilitySpecHandle AbilityHandle, FPredictionKey AbilityOriginalPredictionKey);
 
 	/** Direct Input state replication. These will be called if bReplicateInputDirectly is true on the ability and is generally not a good thing to use. (Instead, prefer to use Generic Replicated Events). */
 	UFUNCTION(Server, reliable, WithValidation)
-	void ServerSetInputPressed(FGameplayAbilitySpecHandle AbilityHandle);
+	UE_API void ServerSetInputPressed(FGameplayAbilitySpecHandle AbilityHandle);
 
 	UFUNCTION(Server, reliable, WithValidation)
-	void ServerSetInputReleased(FGameplayAbilitySpecHandle AbilityHandle);
+	UE_API void ServerSetInputReleased(FGameplayAbilitySpecHandle AbilityHandle);
 
 	/** Called on local player always. Called on server only if bReplicateInputDirectly is set on the GameplayAbility. */
-	virtual void AbilitySpecInputPressed(FGameplayAbilitySpec& Spec);
+	UE_API virtual void AbilitySpecInputPressed(FGameplayAbilitySpec& Spec);
 
 	/** Called on local player always. Called on server only if bReplicateInputDirectly is set on the GameplayAbility. */
-	virtual void AbilitySpecInputReleased(FGameplayAbilitySpec& Spec);
+	UE_API virtual void AbilitySpecInputReleased(FGameplayAbilitySpec& Spec);
 
 	// ----------------------------------------------------------------------------------------------------------------
 	//  Component overrides
 	// ----------------------------------------------------------------------------------------------------------------
 
-	virtual void InitializeComponent() override;
-	virtual void UninitializeComponent() override;
-	virtual void OnComponentDestroyed(bool bDestroyingHierarchy) override;
-	virtual bool GetShouldTick() const override;
-	virtual void TickComponent(float DeltaTime, enum ELevelTick TickType, FActorComponentTickFunction *ThisTickFunction) override;
-	virtual void GetSubobjectsWithStableNamesForNetworking(TArray<UObject*>& Objs) override;
-	virtual bool ReplicateSubobjects(class UActorChannel *Channel, class FOutBunch *Bunch, FReplicationFlags *RepFlags) override;
+	UE_API virtual void InitializeComponent() override;
+	UE_API virtual void UninitializeComponent() override;
+	UE_API virtual void OnComponentDestroyed(bool bDestroyingHierarchy) override;
+	UE_API virtual bool GetShouldTick() const override;
+	UE_API virtual void TickComponent(float DeltaTime, enum ELevelTick TickType, FActorComponentTickFunction *ThisTickFunction) override;
+	UE_API virtual void GetSubobjectsWithStableNamesForNetworking(TArray<UObject*>& Objs) override;
+	UE_API virtual bool ReplicateSubobjects(class UActorChannel *Channel, class FOutBunch *Bunch, FReplicationFlags *RepFlags) override;
 	/** Force owning actor to update it's replication, to make sure that gameplay cues get sent down quickly. Override to change how aggressive this is */
-	virtual void ForceReplication() override;
-	virtual void PreNetReceive() override;
-	virtual void PostNetReceive() override;
-	virtual void OnRegister() override;
-	virtual void OnUnregister() override;
-	virtual void ReadyForReplication() override;
-	virtual void BeginPlay() override;
+	UE_API virtual void ForceReplication() override;
+	UE_API virtual void PreNetReceive() override;
+	UE_API virtual void PostNetReceive() override;
+	UE_API virtual void OnRegister() override;
+	UE_API virtual void OnUnregister() override;
+	UE_API virtual void ReadyForReplication() override;
+	UE_API virtual void BeginPlay() override;
 
 protected:
-	virtual void GetLifetimeReplicatedProps(TArray<FLifetimeProperty>& OutLifetimeProps) const override;
-	virtual void GetReplicatedCustomConditionState(FCustomPropertyConditionState& OutActiveState) const override;
+	UE_API virtual void GetLifetimeReplicatedProps(TArray<FLifetimeProperty>& OutLifetimeProps) const override;
+	UE_API virtual void GetReplicatedCustomConditionState(FCustomPropertyConditionState& OutActiveState) const override;
 
-	void UpdateActiveGameplayEffectsReplicationCondition();
-	void UpdateMinimalReplicationGameplayCuesCondition();
+	UE_API void UpdateActiveGameplayEffectsReplicationCondition();
+	UE_API void UpdateMinimalReplicationGameplayCuesCondition();
 
 	/**
 	 *	The abilities we can activate. 
@@ -1707,31 +1709,31 @@ protected:
 	/** Full list of all instance-per-execution gameplay abilities associated with this component */
 	PRAGMA_DISABLE_DEPRECATION_WARNINGS
 	const TArray<UGameplayAbility*>& GetReplicatedInstancedAbilities() const { return AllReplicatedInstancedAbilities; }
-	PRAGMA_ENABLE_DEPRECATION_WARNINGS
+	UE_API PRAGMA_ENABLE_DEPRECATION_WARNINGS
 
 	/** Add a gameplay ability associated to this component */
 	void AddReplicatedInstancedAbility(UGameplayAbility* GameplayAbility);
 
 	/** Remove a gameplay ability associated to this component */
-	void RemoveReplicatedInstancedAbility(UGameplayAbility* GameplayAbility);
+	UE_API void RemoveReplicatedInstancedAbility(UGameplayAbility* GameplayAbility);
 
 	/** Unregister all the gameplay abilities of this component */
-	void RemoveAllReplicatedInstancedAbilities();
+	UE_API void RemoveAllReplicatedInstancedAbilities();
 
 	/** Will be called from GiveAbility or from OnRep. Initializes events (triggers and inputs) with the given ability */
-	virtual void OnGiveAbility(FGameplayAbilitySpec& AbilitySpec);
+	UE_API virtual void OnGiveAbility(FGameplayAbilitySpec& AbilitySpec);
 
 	/** Will be called from RemoveAbility or from OnRep. Unbinds inputs with the given ability */
-	virtual void OnRemoveAbility(FGameplayAbilitySpec& AbilitySpec);
+	UE_API virtual void OnRemoveAbility(FGameplayAbilitySpec& AbilitySpec);
 
 	/** Called from ClearAbility, ClearAllAbilities or OnRep. Clears any triggers that should no longer exist. */
-	void CheckForClearedAbilities();
+	UE_API void CheckForClearedAbilities();
 
 	/** Cancel a specific ability spec */
-	virtual void CancelAbilitySpec(FGameplayAbilitySpec& Spec, UGameplayAbility* Ignore);
+	UE_API virtual void CancelAbilitySpec(FGameplayAbilitySpec& Spec, UGameplayAbility* Ignore);
 
 	/** Creates a new instance of an ability, storing it in the spec */
-	virtual UGameplayAbility* CreateNewInstanceOfAbility(FGameplayAbilitySpec& Spec, const UGameplayAbility* Ability);
+	UE_API virtual UGameplayAbility* CreateNewInstanceOfAbility(FGameplayAbilitySpec& Spec, const UGameplayAbility* Ability);
 
 	/** Indicates how many levels of ABILITY_SCOPE_LOCK() we are in. The ability list may not be modified while AbilityScopeLockCount > 0. */
 	int32 AbilityScopeLockCount;
@@ -1746,67 +1748,67 @@ protected:
 	float AbilityLastActivatedTime;
 
 	UFUNCTION()
-	virtual void OnRep_ActivateAbilities();
+	UE_API virtual void OnRep_ActivateAbilities();
 
 	UFUNCTION(Server, reliable, WithValidation)
-	void	ServerTryActivateAbility(FGameplayAbilitySpecHandle AbilityToActivate, bool InputPressed, FPredictionKey PredictionKey);
+	UE_API void	ServerTryActivateAbility(FGameplayAbilitySpecHandle AbilityToActivate, bool InputPressed, FPredictionKey PredictionKey);
 
 	UFUNCTION(Server, reliable, WithValidation)
-	void	ServerTryActivateAbilityWithEventData(FGameplayAbilitySpecHandle AbilityToActivate, bool InputPressed, FPredictionKey PredictionKey, FGameplayEventData TriggerEventData);
+	UE_API void	ServerTryActivateAbilityWithEventData(FGameplayAbilitySpecHandle AbilityToActivate, bool InputPressed, FPredictionKey PredictionKey, FGameplayEventData TriggerEventData);
 
 	UFUNCTION(Client, reliable)
-	void	ClientTryActivateAbility(FGameplayAbilitySpecHandle AbilityToActivate);
+	UE_API void	ClientTryActivateAbility(FGameplayAbilitySpecHandle AbilityToActivate);
 
 	/** Called by ServerEndAbility and ClientEndAbility; avoids code duplication. */
-	void	RemoteEndOrCancelAbility(FGameplayAbilitySpecHandle AbilityToEnd, FGameplayAbilityActivationInfo ActivationInfo, bool bWasCanceled);
+	UE_API void	RemoteEndOrCancelAbility(FGameplayAbilitySpecHandle AbilityToEnd, FGameplayAbilityActivationInfo ActivationInfo, bool bWasCanceled);
 
 	UFUNCTION(Server, reliable, WithValidation)
-	void	ServerEndAbility(FGameplayAbilitySpecHandle AbilityToEnd, FGameplayAbilityActivationInfo ActivationInfo, FPredictionKey PredictionKey);
+	UE_API void	ServerEndAbility(FGameplayAbilitySpecHandle AbilityToEnd, FGameplayAbilityActivationInfo ActivationInfo, FPredictionKey PredictionKey);
 
 	UFUNCTION(Client, reliable)
-	void	ClientEndAbility(FGameplayAbilitySpecHandle AbilityToEnd, FGameplayAbilityActivationInfo ActivationInfo);
+	UE_API void	ClientEndAbility(FGameplayAbilitySpecHandle AbilityToEnd, FGameplayAbilityActivationInfo ActivationInfo);
 
 	UFUNCTION(Server, reliable, WithValidation)
-	void    ServerCancelAbility(FGameplayAbilitySpecHandle AbilityToCancel, FGameplayAbilityActivationInfo ActivationInfo);
+	UE_API void    ServerCancelAbility(FGameplayAbilitySpecHandle AbilityToCancel, FGameplayAbilityActivationInfo ActivationInfo);
 
 	UFUNCTION(Client, reliable)
-	void    ClientCancelAbility(FGameplayAbilitySpecHandle AbilityToCancel, FGameplayAbilityActivationInfo ActivationInfo);
+	UE_API void    ClientCancelAbility(FGameplayAbilitySpecHandle AbilityToCancel, FGameplayAbilityActivationInfo ActivationInfo);
 
 	UFUNCTION(Client, Reliable)
-	void	ClientActivateAbilityFailed(FGameplayAbilitySpecHandle AbilityToActivate, int16 PredictionKey);
+	UE_API void	ClientActivateAbilityFailed(FGameplayAbilitySpecHandle AbilityToActivate, int16 PredictionKey);
 	int32	ClientActivateAbilityFailedCountRecent;
 	float	ClientActivateAbilityFailedStartTime;
 
 	
-	void	OnClientActivateAbilityCaughtUp(FGameplayAbilitySpecHandle AbilityToActivate, FPredictionKey::KeyType PredictionKey);
+	UE_API void	OnClientActivateAbilityCaughtUp(FGameplayAbilitySpecHandle AbilityToActivate, FPredictionKey::KeyType PredictionKey);
 
 	UFUNCTION(Client, Reliable)
-	void	ClientActivateAbilitySucceed(FGameplayAbilitySpecHandle AbilityToActivate, FPredictionKey PredictionKey);
+	UE_API void	ClientActivateAbilitySucceed(FGameplayAbilitySpecHandle AbilityToActivate, FPredictionKey PredictionKey);
 
 	UFUNCTION(Client, Reliable)
-	void	ClientActivateAbilitySucceedWithEventData(FGameplayAbilitySpecHandle AbilityToActivate, FPredictionKey PredictionKey, FGameplayEventData TriggerEventData);
+	UE_API void	ClientActivateAbilitySucceedWithEventData(FGameplayAbilitySpecHandle AbilityToActivate, FPredictionKey PredictionKey, FGameplayEventData TriggerEventData);
 
 	/** Implementation of ServerTryActivateAbility */
-	virtual void InternalServerTryActivateAbility(FGameplayAbilitySpecHandle AbilityToActivate, bool InputPressed, const FPredictionKey& PredictionKey, const FGameplayEventData* TriggerEventData);
+	UE_API virtual void InternalServerTryActivateAbility(FGameplayAbilitySpecHandle AbilityToActivate, bool InputPressed, const FPredictionKey& PredictionKey, const FGameplayEventData* TriggerEventData);
 
 	/** Called when a prediction key that played a montage is rejected */
-	void OnPredictiveMontageRejected(UAnimMontage* PredictiveMontage);
+	UE_API void OnPredictiveMontageRejected(UAnimMontage* PredictiveMontage);
 
 	/** Copy LocalAnimMontageInfo into RepAnimMontageInfo */
-	void AnimMontage_UpdateReplicatedData();
-	void AnimMontage_UpdateReplicatedData(FGameplayAbilityRepAnimMontage& OutRepAnimMontageInfo);
+	UE_API void AnimMontage_UpdateReplicatedData();
+	UE_API void AnimMontage_UpdateReplicatedData(FGameplayAbilityRepAnimMontage& OutRepAnimMontageInfo);
 
 	/** Copy over playing flags for duplicate animation data */
-	void AnimMontage_UpdateForcedPlayFlags(FGameplayAbilityRepAnimMontage& OutRepAnimMontageInfo);
+	UE_API void AnimMontage_UpdateForcedPlayFlags(FGameplayAbilityRepAnimMontage& OutRepAnimMontageInfo);
 
 	UE_DEPRECATED(4.26, "This will be made private in future engine versions. Use SetRepAnimMontageInfo, GetRepAnimMontageInfo, or GetRepAnimMontageInfo_Mutable instead.")
 	/** Data structure for replicating montage info to simulated clients */
 	UPROPERTY(ReplicatedUsing=OnRep_ReplicatedAnimMontage)
 	FGameplayAbilityRepAnimMontage RepAnimMontageInfo;
 
-	void SetRepAnimMontageInfo(const FGameplayAbilityRepAnimMontage& NewRepAnimMontageInfo);
-	FGameplayAbilityRepAnimMontage& GetRepAnimMontageInfo_Mutable();
-	const FGameplayAbilityRepAnimMontage& GetRepAnimMontageInfo() const;
+	UE_API void SetRepAnimMontageInfo(const FGameplayAbilityRepAnimMontage& NewRepAnimMontageInfo);
+	UE_API FGameplayAbilityRepAnimMontage& GetRepAnimMontageInfo_Mutable();
+	UE_API const FGameplayAbilityRepAnimMontage& GetRepAnimMontageInfo() const;
 
 	/** Cached value of rather this is a simulated actor */
 	UPROPERTY()
@@ -1821,22 +1823,22 @@ protected:
 	FGameplayAbilityLocalAnimMontage LocalAnimMontageInfo;
 
 	UFUNCTION()
-	virtual void OnRep_ReplicatedAnimMontage();
+	UE_API virtual void OnRep_ReplicatedAnimMontage();
 
 	/** Returns true if we are ready to handle replicated montage information */
-	virtual bool IsReadyForReplicatedMontage();
+	UE_API virtual bool IsReadyForReplicatedMontage();
 
 	/** RPC function called from CurrentMontageSetNextSectionName, replicates to other clients */
 	UFUNCTION(reliable, server, WithValidation)
-	void ServerCurrentMontageSetNextSectionName(UAnimSequenceBase* ClientAnimation, float ClientPosition, FName SectionName, FName NextSectionName);
+	UE_API void ServerCurrentMontageSetNextSectionName(UAnimSequenceBase* ClientAnimation, float ClientPosition, FName SectionName, FName NextSectionName);
 
 	/** RPC function called from CurrentMontageJumpToSection, replicates to other clients */
 	UFUNCTION(reliable, server, WithValidation)
-	void ServerCurrentMontageJumpToSectionName(UAnimSequenceBase* ClientAnimation, FName SectionName);
+	UE_API void ServerCurrentMontageJumpToSectionName(UAnimSequenceBase* ClientAnimation, FName SectionName);
 
 	/** RPC function called from CurrentMontageSetPlayRate, replicates to other clients */
 	UFUNCTION(reliable, server, WithValidation)
-	void ServerCurrentMontageSetPlayRate(UAnimSequenceBase* ClientAnimation, float InPlayRate);
+	UE_API void ServerCurrentMontageSetPlayRate(UAnimSequenceBase* ClientAnimation, float InPlayRate);
 
 	/** Abilities that are triggered from a gameplay event */
 	TMap<FGameplayTag, TArray<FGameplayAbilitySpecHandle > > GameplayEventTriggeredAbilities;
@@ -1845,34 +1847,34 @@ protected:
 	TMap<FGameplayTag, TArray<FGameplayAbilitySpecHandle > > OwnedTagTriggeredAbilities;
 
 	/** Callback that is called when an owned tag bound to an ability changes */
-	virtual void MonitoredTagChanged(const FGameplayTag Tag, int32 NewCount);
+	UE_API virtual void MonitoredTagChanged(const FGameplayTag Tag, int32 NewCount);
 
 	/** Returns true if the specified ability should be activated from an event in this network mode */
-	bool HasNetworkAuthorityToActivateTriggeredAbility(const FGameplayAbilitySpec &Spec) const;
+	UE_API bool HasNetworkAuthorityToActivateTriggeredAbility(const FGameplayAbilitySpec &Spec) const;
 
 	UE_DEPRECATED(5.3, "Use OnImmunityBlockGameplayEffectDelegate directly.  It is trigger from a UImmunityGameplayEffectComponent.  You can create your own GameplayEffectComponent if you need different functionality.")
-	virtual void OnImmunityBlockGameplayEffect(const FGameplayEffectSpec& Spec, const FActiveGameplayEffect* ImmunityGE);
+	UE_API virtual void OnImmunityBlockGameplayEffect(const FGameplayEffectSpec& Spec, const FActiveGameplayEffect* ImmunityGE);
 
 	// Internal gameplay cue functions
-	virtual void AddGameplayCue_Internal(const FGameplayTag GameplayCueTag, FGameplayEffectContextHandle& EffectContext, FActiveGameplayCueContainer& GameplayCueContainer);
-	virtual void AddGameplayCue_Internal(const FGameplayTag GameplayCueTag, const FGameplayCueParameters& GameplayCueParameters, FActiveGameplayCueContainer& GameplayCueContainer);
-	virtual void RemoveGameplayCue_Internal(const FGameplayTag GameplayCueTag, FActiveGameplayCueContainer& GameplayCueContainer);
+	UE_API virtual void AddGameplayCue_Internal(const FGameplayTag GameplayCueTag, FGameplayEffectContextHandle& EffectContext, FActiveGameplayCueContainer& GameplayCueContainer);
+	UE_API virtual void AddGameplayCue_Internal(const FGameplayTag GameplayCueTag, const FGameplayCueParameters& GameplayCueParameters, FActiveGameplayCueContainer& GameplayCueContainer);
+	UE_API virtual void RemoveGameplayCue_Internal(const FGameplayTag GameplayCueTag, FActiveGameplayCueContainer& GameplayCueContainer);
 
 	/** Actually pushes the final attribute value to the attribute set's property. Should not be called by outside code since this does not go through the attribute aggregator system. */
-	void SetNumericAttribute_Internal(const FGameplayAttribute &Attribute, float& NewFloatValue);
+	UE_API void SetNumericAttribute_Internal(const FGameplayAttribute &Attribute, float& NewFloatValue);
 
-	bool HasNetworkAuthorityToApplyGameplayEffect(FPredictionKey PredictionKey) const;
+	UE_API bool HasNetworkAuthorityToApplyGameplayEffect(FPredictionKey PredictionKey) const;
 
-	void ExecutePeriodicEffect(FActiveGameplayEffectHandle	Handle);
+	UE_API void ExecutePeriodicEffect(FActiveGameplayEffectHandle	Handle);
 
-	void ExecuteGameplayEffect(FGameplayEffectSpec &Spec, FPredictionKey PredictionKey);
+	UE_API void ExecuteGameplayEffect(FGameplayEffectSpec &Spec, FPredictionKey PredictionKey);
 
-	void CheckDurationExpired(FActiveGameplayEffectHandle Handle);
+	UE_API void CheckDurationExpired(FActiveGameplayEffectHandle Handle);
 		
-	TArray<TObjectPtr<UGameplayTask>>&	GetAbilityActiveTasks(UGameplayAbility* Ability);
+	UE_API TArray<TObjectPtr<UGameplayTask>>&	GetAbilityActiveTasks(UGameplayAbility* Ability);
 
 	/** A version of RemoveActiveGameplayEffect that allows us to remove it even when we don't have authority */
-	void RemoveActiveGameplayEffect_AllowClientRemoval(FActiveGameplayEffectHandle Handle, int32 StacksToRemove = -1);
+	UE_API void RemoveActiveGameplayEffect_AllowClientRemoval(FActiveGameplayEffectHandle Handle, int32 StacksToRemove = -1);
 
 	/** Contains all of the gameplay effects that are currently active on this component */
 	UPROPERTY(Replicated)
@@ -1894,11 +1896,11 @@ protected:
 	UPROPERTY(Transient, Replicated)
 	TArray<uint8> BlockedAbilityBindings;
 
-	void SetBlockedAbilityBindings(const TArray<uint8>& NewBlockedAbilityBindings);
-	TArray<uint8>& GetBlockedAbilityBindings_Mutable();
-	const TArray<uint8>& GetBlockedAbilityBindings() const;
+	UE_API void SetBlockedAbilityBindings(const TArray<uint8>& NewBlockedAbilityBindings);
+	UE_API TArray<uint8>& GetBlockedAbilityBindings_Mutable();
+	UE_API const TArray<uint8>& GetBlockedAbilityBindings() const;
 
-	void DebugCyclicAggregatorBroadcasts(struct FAggregator* Aggregator);
+	UE_API void DebugCyclicAggregatorBroadcasts(struct FAggregator* Aggregator);
 	
 	/** Acceleration map for all gameplay tags (OwnedGameplayTags from GEs and explicit GameplayCueTags) */
 	FGameplayTagCountContainer GameplayTagCountContainer;
@@ -1907,24 +1909,24 @@ protected:
 	UPROPERTY(Replicated)
 	FMinimalReplicationTagCountMap MinimalReplicationTags;
 
-	void SetMinimalReplicationTags(const FMinimalReplicationTagCountMap& NewMinimalReplicationTags);
-	FMinimalReplicationTagCountMap& GetMinimalReplicationTags_Mutable();
-	const FMinimalReplicationTagCountMap& GetMinimalReplicationTags() const;
+	UE_API void SetMinimalReplicationTags(const FMinimalReplicationTagCountMap& NewMinimalReplicationTags);
+	UE_API FMinimalReplicationTagCountMap& GetMinimalReplicationTags_Mutable();
+	UE_API const FMinimalReplicationTagCountMap& GetMinimalReplicationTags() const;
 
-	FMinimalReplicationTagCountMap& GetReplicatedLooseTags_Mutable();
-	const FMinimalReplicationTagCountMap& GetReplicatedLooseTags() const;
+	UE_API FMinimalReplicationTagCountMap& GetReplicatedLooseTags_Mutable();
+	UE_API const FMinimalReplicationTagCountMap& GetReplicatedLooseTags() const;
 
-	void ResetTagMap();
+	UE_API void ResetTagMap();
 
-	void NotifyTagMap_StackCountChange(const FGameplayTagContainer& Container);
+	UE_API void NotifyTagMap_StackCountChange(const FGameplayTagContainer& Container);
 
 	virtual void OnTagUpdated(const FGameplayTag& Tag, bool TagExists) {};
 	
-	const UAttributeSet*	GetAttributeSubobject(const TSubclassOf<UAttributeSet> AttributeClass) const;
-	const UAttributeSet*	GetAttributeSubobjectChecked(const TSubclassOf<UAttributeSet> AttributeClass) const;
-	const UAttributeSet*	GetOrCreateAttributeSubobject(TSubclassOf<UAttributeSet> AttributeClass);
+	UE_API const UAttributeSet*	GetAttributeSubobject(const TSubclassOf<UAttributeSet> AttributeClass) const;
+	UE_API const UAttributeSet*	GetAttributeSubobjectChecked(const TSubclassOf<UAttributeSet> AttributeClass) const;
+	UE_API const UAttributeSet*	GetOrCreateAttributeSubobject(TSubclassOf<UAttributeSet> AttributeClass);
 
-	void UpdateTagMap_Internal(const FGameplayTagContainer& Container, int32 CountDelta);
+	UE_API void UpdateTagMap_Internal(const FGameplayTagContainer& Container, int32 CountDelta);
 
 	friend struct FActiveGameplayEffect;
 	friend struct FActiveGameplayEffectAction;
@@ -1942,7 +1944,7 @@ protected:
 private:
 
 	// Needs to be called when modifying the SpawnedAttributes array for changes to be replicated
-	void SetSpawnedAttributesListDirty();
+	UE_API void SetSpawnedAttributesListDirty();
 
     // Private accessor to the AllReplicatedInstancedAbilities array until the deprecation tag on it is removed and we can reference the array directly again.
 	PRAGMA_DISABLE_DEPRECATION_WARNINGS
@@ -1956,7 +1958,7 @@ private:
 	TArray<TObjectPtr<UAttributeSet>>	SpawnedAttributes;
 
 	UFUNCTION()
-	void OnRep_SpawnedAttributes(const TArray<UAttributeSet*>& PreviousSpawnedAttributes);
+	UE_API void OnRep_SpawnedAttributes(const TArray<UAttributeSet*>& PreviousSpawnedAttributes);
 
 	FDelegateHandle MonitoredTagChangedDelegateHandle;
 	FTimerHandle    OnRep_ActivateAbilitiesTimerHandle;
@@ -1969,7 +1971,7 @@ private:
 public:
 
 	/** Caches the flags that indicate whether this component has network authority. */
-	void CacheIsNetSimulated();
+	UE_API void CacheIsNetSimulated();
 
 	/** PredictionKeys, see more info in GameplayPrediction.h. This has to come *last* in all replicated properties on the AbilitySystemComponent to ensure OnRep/callback order. */
 	UPROPERTY(Replicated, Transient)
@@ -2002,6 +2004,8 @@ protected:
 	TArray<FAbilityListLockActiveChange*> AbilityListLockActiveChanges;
 
 private:
-	float PlayMontageInternal(UGameplayAbility* AnimatingAbility, FGameplayAbilityActivationInfo ActivationInfo, UAnimMontage* Montage, float InPlayRate, FName StartSectionName, float StartTimeSeconds, TFunctionRef<void(FGameplayAbilityRepAnimMontage&)> MutateRepAnimMontageFunction);
+	UE_API float PlayMontageInternal(UGameplayAbility* AnimatingAbility, FGameplayAbilityActivationInfo ActivationInfo, UAnimMontage* Montage, float InPlayRate, FName StartSectionName, float StartTimeSeconds, TFunctionRef<void(FGameplayAbilityRepAnimMontage&)> MutateRepAnimMontageFunction);
 	
 };
+
+#undef UE_API
