[GameplayAbilities]

diff --git a/Engine/Plugins/Runtime/GameplayAbilities/Source/GameplayAbilities/Public/GameplayCueNotify_Looping.h b/Engine/Plugins/Runtime/GameplayAbilities/Source/GameplayAbilities/Public/GameplayCueNotify_Looping.h
--- a/Engine/Plugins/Runtime/GameplayAbilities/Source/GameplayAbilities/Public/GameplayCueNotify_Looping.h
+++ b/Engine/Plugins/Runtime/GameplayAbilities/Source/GameplayAbilities/Public/GameplayCueNotify_Looping.h
@@ -7,6 +7,8 @@
 #include "GameplayCueNotifyTypes.h"
 #include "GameplayCueNotify_Looping.generated.h"
 
+#define UE_API GAMEPLAYABILITIES_API
+
 
 /**
  * AGameplayCueNotify_Looping
@@ -14,41 +16,41 @@
  *	This is an instanced gameplay cue notify for continuous looping effects.
  *	The game is responsible for defining the start/stop by adding/removing the gameplay cue.
  */
-UCLASS(Blueprintable, notplaceable, Category = "GameplayCueNotify", Meta = (ShowWorldContextPin, DisplayName = "GCN Looping", ShortTooltip = "A GameplayCueNotify that has a duration that is driven by the game."))
-class GAMEPLAYABILITIES_API AGameplayCueNotify_Looping : public AGameplayCueNotify_Actor
+UCLASS(Blueprintable, notplaceable, Category = "GameplayCueNotify", Meta = (ShowWorldContextPin, DisplayName = "GCN Looping", ShortTooltip = "A GameplayCueNotify that has a duration that is driven by the game."), MinimalAPI)
+class AGameplayCueNotify_Looping : public AGameplayCueNotify_Actor
 {
 	GENERATED_BODY()
 
 public:
 
-	AGameplayCueNotify_Looping();
+	UE_API AGameplayCueNotify_Looping();
 
 protected:
 
-	virtual void EndPlay(const EEndPlayReason::Type EndPlayReason) override;
-	virtual bool Recycle() override;
+	UE_API virtual void EndPlay(const EEndPlayReason::Type EndPlayReason) override;
+	UE_API virtual bool Recycle() override;
 
-	virtual bool OnActive_Implementation(AActor* Target, const FGameplayCueParameters& Parameters) override;
-	virtual bool WhileActive_Implementation(AActor* Target, const FGameplayCueParameters& Parameters) override;
-	virtual bool OnExecute_Implementation(AActor* Target, const FGameplayCueParameters& Parameters) override;
-	virtual bool OnRemove_Implementation(AActor* Target, const FGameplayCueParameters& Parameters) override;
+	UE_API virtual bool OnActive_Implementation(AActor* Target, const FGameplayCueParameters& Parameters) override;
+	UE_API virtual bool WhileActive_Implementation(AActor* Target, const FGameplayCueParameters& Parameters) override;
+	UE_API virtual bool OnExecute_Implementation(AActor* Target, const FGameplayCueParameters& Parameters) override;
+	UE_API virtual bool OnRemove_Implementation(AActor* Target, const FGameplayCueParameters& Parameters) override;
 
 	UFUNCTION(BlueprintImplementableEvent)
-	void OnApplication(AActor* Target, const FGameplayCueParameters& Parameters, const FGameplayCueNotify_SpawnResult& SpawnResults);
+	UE_API void OnApplication(AActor* Target, const FGameplayCueParameters& Parameters, const FGameplayCueNotify_SpawnResult& SpawnResults);
 
 	UFUNCTION(BlueprintImplementableEvent)
-	void OnLoopingStart(AActor* Target, const FGameplayCueParameters& Parameters, const FGameplayCueNotify_SpawnResult& SpawnResults);
+	UE_API void OnLoopingStart(AActor* Target, const FGameplayCueParameters& Parameters, const FGameplayCueNotify_SpawnResult& SpawnResults);
 
 	UFUNCTION(BlueprintImplementableEvent)
-	void OnRecurring(AActor* Target, const FGameplayCueParameters& Parameters, const FGameplayCueNotify_SpawnResult& SpawnResults);
+	UE_API void OnRecurring(AActor* Target, const FGameplayCueParameters& Parameters, const FGameplayCueNotify_SpawnResult& SpawnResults);
 
 	UFUNCTION(BlueprintImplementableEvent)
-	void OnRemoval(AActor* Target, const FGameplayCueParameters& Parameters, const FGameplayCueNotify_SpawnResult& SpawnResults);
+	UE_API void OnRemoval(AActor* Target, const FGameplayCueParameters& Parameters, const FGameplayCueNotify_SpawnResult& SpawnResults);
 
-	void RemoveLoopingEffects();
+	UE_API void RemoveLoopingEffects();
 
 #if WITH_EDITOR
-	virtual EDataValidationResult IsDataValid(class FDataValidationContext& Context) const override;
+	UE_API virtual EDataValidationResult IsDataValid(class FDataValidationContext& Context) const override;
 #endif // #if WITH_EDITOR
 
 protected:
@@ -95,3 +97,5 @@ protected:
 
 	bool bLoopingEffectsRemoved;
 };
+
+#undef UE_API
diff --git a/Engine/Plugins/Runtime/GameplayAbilities/Source/GameplayAbilities/Public/GameplayCueNotify_Static.h b/Engine/Plugins/Runtime/GameplayAbilities/Source/GameplayAbilities/Public/GameplayCueNotify_Static.h
--- a/Engine/Plugins/Runtime/GameplayAbilities/Source/GameplayAbilities/Public/GameplayCueNotify_Static.h
+++ b/Engine/Plugins/Runtime/GameplayAbilities/Source/GameplayAbilities/Public/GameplayCueNotify_Static.h
@@ -9,51 +9,53 @@
 #include "GameplayEffectTypes.h"
 #include "GameplayCueNotify_Static.generated.h"
 
+#define UE_API GAMEPLAYABILITIES_API
+
 /**
  *	A non instantiated UObject that acts as a handler for a GameplayCue. These are useful for one-off "burst" effects.
  */
 
-UCLASS(Blueprintable, meta = (ShowWorldContextPin), hidecategories = (Replication))
-class GAMEPLAYABILITIES_API UGameplayCueNotify_Static : public UObject
+UCLASS(Blueprintable, meta = (ShowWorldContextPin), hidecategories = (Replication), MinimalAPI)
+class UGameplayCueNotify_Static : public UObject
 {
 	GENERATED_UCLASS_BODY()
 
 	/** Does this GameplayCueNotify handle this type of GameplayCueEvent? */
-	virtual bool HandlesEvent(EGameplayCueEvent::Type EventType) const;
+	UE_API virtual bool HandlesEvent(EGameplayCueEvent::Type EventType) const;
 
-	virtual void OnOwnerDestroyed();
+	UE_API virtual void OnOwnerDestroyed();
 
-	virtual void PostInitProperties() override;
+	UE_API virtual void PostInitProperties() override;
 
-	virtual void Serialize(FArchive& Ar) override;
+	UE_API virtual void Serialize(FArchive& Ar) override;
 
-	virtual void HandleGameplayCue(AActor* MyTarget, EGameplayCueEvent::Type EventType, const FGameplayCueParameters& Parameters);
+	UE_API virtual void HandleGameplayCue(AActor* MyTarget, EGameplayCueEvent::Type EventType, const FGameplayCueParameters& Parameters);
 
-	UWorld* GetWorld() const override;
+	UE_API UWorld* GetWorld() const override;
 
 #if WITH_EDITOR
-	virtual void PostEditChangeProperty(FPropertyChangedEvent& PropertyChangedEvent) override;
+	UE_API virtual void PostEditChangeProperty(FPropertyChangedEvent& PropertyChangedEvent) override;
 #endif // WITH_EDITOR
 
 	/** Generic Event Graph event that will get called for every event type */
 	UFUNCTION(BlueprintImplementableEvent, Category = "GameplayCueNotify", DisplayName = "HandleGameplayCue", meta=(ScriptName = "HandleGameplayCue"))
-	void K2_HandleGameplayCue(AActor* MyTarget, EGameplayCueEvent::Type EventType, const FGameplayCueParameters& Parameters) const;
+	UE_API void K2_HandleGameplayCue(AActor* MyTarget, EGameplayCueEvent::Type EventType, const FGameplayCueParameters& Parameters) const;
 
 	/** Called when a GameplayCue is executed, this is used for instant effects or periodic ticks */
 	UFUNCTION(BlueprintNativeEvent, BlueprintPure, Category = "GameplayCueNotify")
-	bool OnExecute(AActor* MyTarget, const FGameplayCueParameters& Parameters) const;
+	UE_API bool OnExecute(AActor* MyTarget, const FGameplayCueParameters& Parameters) const;
 
 	/** Called when a GameplayCue with duration is first activated, this will only be called if the client witnessed the activation */
 	UFUNCTION(BlueprintNativeEvent, BlueprintPure, Category = "GameplayCueNotify")
-	bool OnActive(AActor* MyTarget, const FGameplayCueParameters& Parameters) const;
+	UE_API bool OnActive(AActor* MyTarget, const FGameplayCueParameters& Parameters) const;
 
 	/** Called when a GameplayCue with duration is first seen as active, even if it wasn't actually just applied (Join in progress, etc) */
 	UFUNCTION(BlueprintNativeEvent, BlueprintPure, Category = "GameplayCueNotify")
-	bool WhileActive(AActor* MyTarget, const FGameplayCueParameters& Parameters) const;
+	UE_API bool WhileActive(AActor* MyTarget, const FGameplayCueParameters& Parameters) const;
 
 	/** Called when a GameplayCue with duration is removed */
 	UFUNCTION(BlueprintNativeEvent, BlueprintPure, Category = "GameplayCueNotify")
-	bool OnRemove(AActor* MyTarget, const FGameplayCueParameters& Parameters) const;
+	UE_API bool OnRemove(AActor* MyTarget, const FGameplayCueParameters& Parameters) const;
 
 	/** Tag this notify is activated by */
 	UPROPERTY(EditDefaultsOnly, Category = GameplayCue, meta=(Categories="GameplayCue"))
@@ -68,5 +70,7 @@ class GAMEPLAYABILITIES_API UGameplayCueNotify_Static : public UObject
 	bool IsOverride;
 
 private:
-	virtual void DeriveGameplayCueTagFromAssetName();
+	UE_API virtual void DeriveGameplayCueTagFromAssetName();
 };
+
+#undef UE_API
diff --git a/Engine/Plugins/Runtime/GameplayAbilities/Source/GameplayAbilities/Public/GameplayCueSet.h b/Engine/Plugins/Runtime/GameplayAbilities/Source/GameplayAbilities/Public/GameplayCueSet.h
--- a/Engine/Plugins/Runtime/GameplayAbilities/Source/GameplayAbilities/Public/GameplayCueSet.h
+++ b/Engine/Plugins/Runtime/GameplayAbilities/Source/GameplayAbilities/Public/GameplayCueSet.h
@@ -10,6 +10,8 @@
 #include "GameplayEffectTypes.h"
 #include "GameplayCueSet.generated.h"
 
+#define UE_API GAMEPLAYABILITIES_API
+
 USTRUCT()
 struct FGameplayCueNotifyData
 {
@@ -47,45 +49,45 @@ struct FGameplayCueReferencePair
 /**
  *	A set of gameplay cue actors to handle gameplay cue events
  */
-UCLASS()
-class GAMEPLAYABILITIES_API UGameplayCueSet : public UDataAsset
+UCLASS(MinimalAPI)
+class UGameplayCueSet : public UDataAsset
 {
 	GENERATED_UCLASS_BODY()
 
 	/** Handles the cue event by spawning the cue actor. Returns true if the event was handled. */
-	virtual bool HandleGameplayCue(AActor* TargetActor, FGameplayTag GameplayCueTag, EGameplayCueEvent::Type EventType, const FGameplayCueParameters& Parameters);
+	UE_API virtual bool HandleGameplayCue(AActor* TargetActor, FGameplayTag GameplayCueTag, EGameplayCueEvent::Type EventType, const FGameplayCueParameters& Parameters);
 
 	/** Adds a list of cues to the set */
-	virtual void AddCues(const TArray<FGameplayCueReferencePair>& CuesToAdd);
+	UE_API virtual void AddCues(const TArray<FGameplayCueReferencePair>& CuesToAdd);
 
 	/** Removes all cues from the set matching any of the supplied tags */
-	virtual void RemoveCuesByTags(const FGameplayTagContainer& TagsToRemove);
+	UE_API virtual void RemoveCuesByTags(const FGameplayTagContainer& TagsToRemove);
 
 	/** Removes all cues from the set matching the supplied string refs */
-	virtual void RemoveCuesByStringRefs(const TArray<FSoftObjectPath>& CuesToRemove);
+	UE_API virtual void RemoveCuesByStringRefs(const TArray<FSoftObjectPath>& CuesToRemove);
 
 	/** Nulls reference to the loaded class. Note this doesn't remove the entire cue from the internal data structure, just the hard ref to the loaded class */
-	virtual void RemoveLoadedClass(UClass* Class);
+	UE_API virtual void RemoveLoadedClass(UClass* Class);
 
 	/** Returns filenames of everything we know about (loaded or not) */
-	virtual void GetFilenames(TArray<FString>& Filenames) const;
+	UE_API virtual void GetFilenames(TArray<FString>& Filenames) const;
 
 	/** Extracts all soft object paths pointing to Cues */
-	virtual void GetSoftObjectPaths(TArray<FSoftObjectPath>& List) const;
+	UE_API virtual void GetSoftObjectPaths(TArray<FSoftObjectPath>& List) const;
 
 #if WITH_EDITOR
 
-	void CopyCueDataToSetForEditorPreview(FGameplayTag Tag, UGameplayCueSet* DestinationSet);
+	UE_API void CopyCueDataToSetForEditorPreview(FGameplayTag Tag, UGameplayCueSet* DestinationSet);
 
 	/** Updates an existing cue */
-	virtual void UpdateCueByStringRefs(const FSoftObjectPath& CueToRemove, FString NewPath);
+	UE_API virtual void UpdateCueByStringRefs(const FSoftObjectPath& CueToRemove, FString NewPath);
 
 #endif
 
 	/** Removes all cues from the set */
-	virtual void Empty();
+	UE_API virtual void Empty();
 
-	virtual void PrintCues() const;
+	UE_API virtual void PrintCues() const;
 	
 	UPROPERTY(EditAnywhere, Category=CueSet)
 	TArray<FGameplayCueNotifyData> GameplayCueData;
@@ -93,9 +95,11 @@ class GAMEPLAYABILITIES_API UGameplayCueSet : public UDataAsset
 	/** Maps GameplayCue Tag to index into above GameplayCues array. */
 	TMap<FGameplayTag, int32> GameplayCueDataMap;
 
-	static FGameplayTag	BaseGameplayCueTag();
+	static UE_API FGameplayTag	BaseGameplayCueTag();
 
 protected:
-	virtual bool HandleGameplayCueNotify_Internal(AActor* TargetActor, int32 DataIdx, EGameplayCueEvent::Type EventType, FGameplayCueParameters& Parameters);
-	virtual void BuildAccelerationMap_Internal();
+	UE_API virtual bool HandleGameplayCueNotify_Internal(AActor* TargetActor, int32 DataIdx, EGameplayCueEvent::Type EventType, FGameplayCueParameters& Parameters);
+	UE_API virtual void BuildAccelerationMap_Internal();
 };
+
+#undef UE_API
diff --git a/Engine/Plugins/Runtime/GameplayAbilities/Source/GameplayAbilities/Public/GameplayCueTranslator.h b/Engine/Plugins/Runtime/GameplayAbilities/Source/GameplayAbilities/Public/GameplayCueTranslator.h
--- a/Engine/Plugins/Runtime/GameplayAbilities/Source/GameplayAbilities/Public/GameplayCueTranslator.h
+++ b/Engine/Plugins/Runtime/GameplayAbilities/Source/GameplayAbilities/Public/GameplayCueTranslator.h
@@ -9,6 +9,8 @@
 #include "GameplayTagContainer.h"
 #include "GameplayCueTranslator.generated.h"
 
+#define UE_API GAMEPLAYABILITIES_API
+
 class AActor;
 class UGameplayCueTranslator;
 class UGameplayTagsManager;
@@ -55,7 +57,7 @@ struct FGameplayCueParameters;
 //	Editor only. Data that is only used by the GameplayCue editor tool
 // -----------------------------------------------------------------------------
 #if WITH_EDITOR
-struct GAMEPLAYABILITIES_API FGameplayCueTranslationEditorOnlyData
+struct FGameplayCueTranslationEditorOnlyData
 {
 	FGameplayCueTranslationEditorOnlyData()
 		:UniqueID(0), Enabled(true)
@@ -68,7 +70,7 @@ struct GAMEPLAYABILITIES_API FGameplayCueTranslationEditorOnlyData
 	bool	Enabled;			// Is this actually enabled, or not (and if not, we may still want to draw it in the editor but greyed out). In non editor builds, the translator rule class should "just not return the rule" if its disanled.
 };
 
-struct GAMEPLAYABILITIES_API FGameplayCueTranslationEditorInfo
+struct FGameplayCueTranslationEditorInfo
 {
 	FGameplayTag GameplayTag;	// Will only exist if there is an existing FGameplayTag
 	FName	GameplayTagName;	// Will always exist, even if tag hasn't been created
@@ -84,7 +86,7 @@ struct GAMEPLAYABILITIES_API FGameplayCueTranslationEditorInfo
 class UGameplayCueTranslator;
 
 /** Basis for name swaps. This swaps FromName to ToName */
-struct GAMEPLAYABILITIES_API FGameplayCueTranslationNameSwap
+struct FGameplayCueTranslationNameSwap
 {
 	FName	FromName;
 	TArray<FName, TInlineAllocator<4> >	ToNames;
@@ -96,7 +98,7 @@ struct GAMEPLAYABILITIES_API FGameplayCueTranslationNameSwap
 
 /** Simple index/handle for referencing items in FGameplayCueTranslationManager::TranslationLUT  */
 USTRUCT()
-struct GAMEPLAYABILITIES_API FGameplayCueTranslatorNodeIndex
+struct FGameplayCueTranslatorNodeIndex
 {
 	GENERATED_USTRUCT_BODY()
 
@@ -177,38 +179,38 @@ struct FNameSwapData
 
 /** This is the struct that does the actual translation. It lives on the GameplayCueManager and encapsulates all translation logic. */
 USTRUCT()
-struct GAMEPLAYABILITIES_API FGameplayCueTranslationManager
+struct FGameplayCueTranslationManager
 {
 	GENERATED_USTRUCT_BODY()
 
 	/** This is *the* runtime function that translates the tag (if necessary) */
-	void TranslateTag(FGameplayTag& Tag, AActor* TargetActor, const FGameplayCueParameters& Parameters);
+	UE_API void TranslateTag(FGameplayTag& Tag, AActor* TargetActor, const FGameplayCueParameters& Parameters);
 
 	/** Builds all our translation tables. This works backwards by looking at the tags in the dictionary and determining which could be derived from translation rules. This is pretty fast. */
-	void BuildTagTranslationTable();
-	bool BuildTagTranslationTable_r(const FName& TagName, const TArray<FName>& SplitNames);
+	UE_API void BuildTagTranslationTable();
+	UE_API bool BuildTagTranslationTable_r(const FName& TagName, const TArray<FName>& SplitNames);
 
 	/** 
 	 *  Builds all *possible* translation tables. This works forwards by looking at existing tags and determining what translated tags could be dervied from them. 
 	 *	This WILL populate the TranslationLUT with all possible tags, not just ones that exist in the dictionar. This is not as fast and is onyl called in the editor. 
 	 */
-	void BuildTagTranslationTable_Forward();
-	void BuildTagTranslationTable_Forward_r(const FName& TagName, const TArray<FName>& SplitNames);
+	UE_API void BuildTagTranslationTable_Forward();
+	UE_API void BuildTagTranslationTable_Forward_r(const FName& TagName, const TArray<FName>& SplitNames);
 		
 	/** refresh our name swap rules. Called internally by the manager and externally by GC tool */
-	void RefreshNameSwaps();
+	UE_API void RefreshNameSwaps();
 
-	void PrintTranslationTable();
-	void PrintTranslationTable_r(FGameplayCueTranslatorNode& Node, FString IdentStr=FString());
+	UE_API void PrintTranslationTable();
+	UE_API void PrintTranslationTable_r(FGameplayCueTranslatorNode& Node, FString IdentStr=FString());
 
 	
 #if WITH_EDITOR
 	/** Used by the GC editor to enumerate possible translation tags. Never called at runtime. */
-	bool GetTranslatedTags(const FName& ParentTag, TArray<FGameplayCueTranslationEditorInfo>& Children);
+	UE_API bool GetTranslatedTags(const FName& ParentTag, TArray<FGameplayCueTranslationEditorInfo>& Children);
 	const TArray<FNameSwapData>& GetNameSwapData() const { return AllNameSwaps; }
 
 	/** Searches, slowly, to see if the passed in tag can be translated from something else. Don't call this at runtime. */
-	FGameplayTag SearchSlowForTranslationParent(FGameplayTag Tag);
+	UE_API FGameplayTag SearchSlowForTranslationParent(FGameplayTag Tag);
 #endif
 
 private:
@@ -228,16 +230,16 @@ private:
 	/** All name swpa rules we have gathered */
 	TArray<FNameSwapData> AllNameSwaps;
 
-	bool TranslateTag_Internal(FGameplayCueTranslatorNode& Node, FGameplayTag& OutTag, const FName& TagName, AActor* TargetActor, const FGameplayCueParameters& Parameters);
+	UE_API bool TranslateTag_Internal(FGameplayCueTranslatorNode& Node, FGameplayTag& OutTag, const FName& TagName, AActor* TargetActor, const FGameplayCueParameters& Parameters);
 
-	FGameplayCueTranslatorNodeIndex GetTranslationIndexForTag(const FGameplayTag& Tag, bool CreateIfInvalid=false);
+	UE_API FGameplayCueTranslatorNodeIndex GetTranslationIndexForTag(const FGameplayTag& Tag, bool CreateIfInvalid=false);
 
-	FGameplayCueTranslatorNode* GetTranslationNodeForTag(const FGameplayTag& Tag, bool CreateIfInvalid=false);
-	FGameplayCueTranslatorNode* GetTranslationNodeForName(FName Name, bool CreateIfInvalid=false);
+	UE_API FGameplayCueTranslatorNode* GetTranslationNodeForTag(const FGameplayTag& Tag, bool CreateIfInvalid=false);
+	UE_API FGameplayCueTranslatorNode* GetTranslationNodeForName(FName Name, bool CreateIfInvalid=false);
 
-	FGameplayCueTranslatorNodeIndex GetTranslationIndexForName(FName Name, bool CreateIfInvalid=false);
+	UE_API FGameplayCueTranslatorNodeIndex GetTranslationIndexForName(FName Name, bool CreateIfInvalid=false);
 
-	void ResetTranslationLUT();
+	UE_API void ResetTranslationLUT();
 
 	// Only used in debug printing/stats
 	int32 TotalNumTranslations;
@@ -255,8 +257,8 @@ private:
  *
  */
 
-UCLASS(Abstract)
-class GAMEPLAYABILITIES_API UGameplayCueTranslator: public UObject
+UCLASS(Abstract, MinimalAPI)
+class UGameplayCueTranslator: public UObject
 {
 	GENERATED_BODY()
 
@@ -368,3 +370,5 @@ public:
 
 	*/
 };
+
+#undef UE_API
