[GameplayAbilities]

diff --git a/Engine/Plugins/Runtime/GameplayAbilities/Source/GameplayAbilities/Public/GameplayEffect.h b/Engine/Plugins/Runtime/GameplayAbilities/Source/GameplayAbilities/Public/GameplayEffect.h
--- a/Engine/Plugins/Runtime/GameplayAbilities/Source/GameplayAbilities/Public/GameplayEffect.h
+++ b/Engine/Plugins/Runtime/GameplayAbilities/Source/GameplayAbilities/Public/GameplayEffect.h
@@ -24,6 +24,8 @@
 #include "VisualLogger/VisualLoggerTypes.h"
 #include "GameplayEffect.generated.h"
 
+#define UE_API GAMEPLAYABILITIES_API
+
 /**
  * Gameplay Effects are bundles of functionality that are _applied_ to Actors.  Think of Gameplay Effects like something at _affects_ an Actor.
  * Gameplay Effects are assets, and thus immutable at runtime.  There are small exceptions to this where hacks are done to _create_ a GE at runtime (but once created and configured, the data is not modified).
@@ -99,19 +101,19 @@ enum class EGameplayEffectVersion : uint8
 	Current = AbilitiesComponent53
 };
 
-struct GAMEPLAYABILITIES_API FGameplayEffectConstants
+struct FGameplayEffectConstants
 {
 	/** Infinite duration */
-	static const float INFINITE_DURATION;
+	static UE_API const float INFINITE_DURATION;
 
 	/** No duration; Time specifying instant application of an effect */
-	static const float INSTANT_APPLICATION;
+	static UE_API const float INSTANT_APPLICATION;
 
 	/** Constant specifying that the combat effect has no period and doesn't check for over time application */
-	static const float NO_PERIOD;
+	static UE_API const float NO_PERIOD;
 
 	/** No Level/Level not set */
-	static const float INVALID_LEVEL;
+	static UE_API const float INVALID_LEVEL;
 };
 
 /** 
@@ -271,7 +273,7 @@ struct FSetByCallerFloat
 
 /** Struct representing the magnitude of a gameplay effect modifier, potentially calculated in numerous different ways */
 USTRUCT()
-struct GAMEPLAYABILITIES_API FGameplayEffectModifierMagnitude
+struct FGameplayEffectModifierMagnitude
 {
 	GENERATED_USTRUCT_BODY()
 
@@ -312,7 +314,7 @@ public:
 	 * 
 	 * @return Whether or not the magnitude can be properly calculated
 	 */
-	bool CanCalculateMagnitude(const FGameplayEffectSpec& InRelevantSpec) const;
+	UE_API bool CanCalculateMagnitude(const FGameplayEffectSpec& InRelevantSpec) const;
 
 	/**
 	 * Attempts to calculate the magnitude given the provided spec. May fail if necessary information (such as captured attributes) is missing from
@@ -323,44 +325,44 @@ public:
 	 * 
 	 * @return True if the calculation was successful, false if it was not
 	 */
-	bool AttemptCalculateMagnitude(const FGameplayEffectSpec& InRelevantSpec, OUT float& OutCalculatedMagnitude, bool WarnIfSetByCallerFail=true, float DefaultSetbyCaller=0.f) const;
+	UE_API bool AttemptCalculateMagnitude(const FGameplayEffectSpec& InRelevantSpec, OUT float& OutCalculatedMagnitude, bool WarnIfSetByCallerFail=true, float DefaultSetbyCaller=0.f) const;
 
 	/** Attempts to recalculate the magnitude given a changed aggregator. This will only recalculate if we are a modifier that is linked (non snapshot) to the given aggregator. */
-	bool AttemptRecalculateMagnitudeFromDependentAggregatorChange(const FGameplayEffectSpec& InRelevantSpec, OUT float& OutCalculatedMagnitude, const FAggregator* ChangedAggregator) const;
+	UE_API bool AttemptRecalculateMagnitudeFromDependentAggregatorChange(const FGameplayEffectSpec& InRelevantSpec, OUT float& OutCalculatedMagnitude, const FAggregator* ChangedAggregator) const;
 
 	/**
 	 * Gather all of the attribute capture definitions necessary to compute the magnitude and place them into the provided array
 	 * 
 	 * @param OutCaptureDefs	[OUT] Array populated with necessary attribute capture definitions
 	 */
-	void GetAttributeCaptureDefinitions(OUT TArray<FGameplayEffectAttributeCaptureDefinition>& OutCaptureDefs) const;
+	UE_API void GetAttributeCaptureDefinitions(OUT TArray<FGameplayEffectAttributeCaptureDefinition>& OutCaptureDefs) const;
 
 	EGameplayEffectMagnitudeCalculation GetMagnitudeCalculationType() const { return MagnitudeCalculationType; }
 
 	/** Returns the magnitude as it was entered in data. Only applies to ScalableFloat or any other type that can return data without context */
-	bool GetStaticMagnitudeIfPossible(float InLevel, float& OutMagnitude, const FString* ContextString = nullptr) const;
+	UE_API bool GetStaticMagnitudeIfPossible(float InLevel, float& OutMagnitude, const FString* ContextString = nullptr) const;
 
 	/** Returns the DataName associated with this magnitude if it is set by caller */
-	bool GetSetByCallerDataNameIfPossible(FName& OutDataName) const;
+	UE_API bool GetSetByCallerDataNameIfPossible(FName& OutDataName) const;
 
 	/** Returns SetByCaller data structure, for inspection purposes */
 	const FSetByCallerFloat& GetSetByCallerFloat() const { return SetByCallerMagnitude; }
 
 	/** Returns the custom magnitude calculation class, if any, for this magnitude. Only applies to CustomMagnitudes */
-	TSubclassOf<UGameplayModMagnitudeCalculation> GetCustomMagnitudeCalculationClass() const;
+	UE_API TSubclassOf<UGameplayModMagnitudeCalculation> GetCustomMagnitudeCalculationClass() const;
 
 	/** Implementing Serialize to clear references to assets that are not needed */
-	bool Serialize(FArchive& Ar);
+	UE_API bool Serialize(FArchive& Ar);
 
-	bool operator==(const FGameplayEffectModifierMagnitude& Other) const;
-	bool operator!=(const FGameplayEffectModifierMagnitude& Other) const;
+	UE_API bool operator==(const FGameplayEffectModifierMagnitude& Other) const;
+	UE_API bool operator!=(const FGameplayEffectModifierMagnitude& Other) const;
 
 #if WITH_EDITOR
-	FText GetValueForEditorDisplay() const;
-	EDataValidationResult IsDataValid(class FDataValidationContext& Context, const FString& PathName) const;
+	UE_API FText GetValueForEditorDisplay() const;
+	UE_API EDataValidationResult IsDataValid(class FDataValidationContext& Context, const FString& PathName) const;
 
 	UE_DEPRECATED(5.6, "Use IsDataValid")
-	void ReportErrors(const FString& PathName) const;
+	UE_API void ReportErrors(const FString& PathName) const;
 #endif
 
 protected:
@@ -476,13 +478,13 @@ struct FGameplayEffectExecutionScopedModifierInfo
  * Struct for gameplay effects that apply only if another gameplay effect (or execution) was successfully applied.
  */
 USTRUCT(BlueprintType)
-struct GAMEPLAYABILITIES_API FConditionalGameplayEffect
+struct FConditionalGameplayEffect
 {
 	GENERATED_USTRUCT_BODY()
 
-	bool CanApply(const FGameplayTagContainer& SourceTags, float SourceLevel) const;
+	UE_API bool CanApply(const FGameplayTagContainer& SourceTags, float SourceLevel) const;
 
-	FGameplayEffectSpecHandle CreateSpec(FGameplayEffectContextHandle EffectContext, float SourceLevel) const;
+	UE_API FGameplayEffectSpecHandle CreateSpec(FGameplayEffectContextHandle EffectContext, float SourceLevel) const;
 
 	/** gameplay effect that will be applied to the target */
 	UPROPERTY(EditDefaultsOnly, BlueprintReadOnly, Category = GameplayEffect)
@@ -492,8 +494,8 @@ struct GAMEPLAYABILITIES_API FConditionalGameplayEffect
 	UPROPERTY(EditDefaultsOnly, BlueprintReadOnly, Category = GameplayEffect)
 	FGameplayTagContainer RequiredSourceTags;
 
-	bool operator==(const FConditionalGameplayEffect& Other) const;
-	bool operator!=(const FConditionalGameplayEffect& Other) const;
+	UE_API bool operator==(const FConditionalGameplayEffect& Other) const;
+	UE_API bool operator!=(const FConditionalGameplayEffect& Other) const;
 };
 
 /** 
@@ -501,7 +503,7 @@ struct GAMEPLAYABILITIES_API FConditionalGameplayEffect
  * Custom executions run special logic from an outside class each time the gameplay effect executes.
  */
 USTRUCT(BlueprintType)
-struct GAMEPLAYABILITIES_API FGameplayEffectExecutionDefinition
+struct FGameplayEffectExecutionDefinition
 {
 	GENERATED_USTRUCT_BODY()
 
@@ -512,7 +514,7 @@ struct GAMEPLAYABILITIES_API FGameplayEffectExecutionDefinition
 	 * 
 	 * @param OutCaptureDefs	[OUT] Capture definitions requested by the execution
 	 */
-	void GetAttributeCaptureDefinitions(OUT TArray<FGameplayEffectAttributeCaptureDefinition>& OutCaptureDefs) const;
+	UE_API void GetAttributeCaptureDefinitions(OUT TArray<FGameplayEffectAttributeCaptureDefinition>& OutCaptureDefs) const;
 
 	/** Custom execution calculation class to run when the gameplay effect executes */
 	UPROPERTY(EditDefaultsOnly, Category=Execution)
@@ -537,7 +539,7 @@ struct GAMEPLAYABILITIES_API FGameplayEffectExecutionDefinition
  *	Does not tell us how exactly
  */
 USTRUCT(BlueprintType)
-struct GAMEPLAYABILITIES_API FGameplayModifierInfo
+struct FGameplayModifierInfo
 {
 	GENERATED_USTRUCT_BODY()
 	
@@ -568,8 +570,8 @@ struct GAMEPLAYABILITIES_API FGameplayModifierInfo
 	FGameplayTagRequirements	TargetTags;
 
 	/** Equality/Inequality operators */
-	bool operator==(const FGameplayModifierInfo& Other) const;
-	bool operator!=(const FGameplayModifierInfo& Other) const;
+	UE_API bool operator==(const FGameplayModifierInfo& Other) const;
+	UE_API bool operator!=(const FGameplayModifierInfo& Other) const;
 };
 
 /**
@@ -625,7 +627,7 @@ struct FGameplayEffectCue
 
 /** Structure that is used to combine tags from parent and child blueprints in a safe way */
 USTRUCT(BlueprintType)
-struct GAMEPLAYABILITIES_API FInheritedTagContainer
+struct FInheritedTagContainer
 {
 	GENERATED_USTRUCT_BODY()
 
@@ -641,19 +643,19 @@ struct GAMEPLAYABILITIES_API FInheritedTagContainer
 	UPROPERTY(EditDefaultsOnly, BlueprintReadOnly, Category = Application, meta=(DisplayName="Remove from Inherited"))
 	FGameplayTagContainer Removed;
 
-	void UpdateInheritedTagProperties(const FInheritedTagContainer* Parent);
+	UE_API void UpdateInheritedTagProperties(const FInheritedTagContainer* Parent);
 
 	/** Apply the Added and Removed tags to the passed-in container (does not have to be the previously configured Parent!) */
-	void ApplyTo(FGameplayTagContainer& ApplyToContainer) const;
+	UE_API void ApplyTo(FGameplayTagContainer& ApplyToContainer) const;
 
 	/** Add a tag that will appear in addition to any inherited tags */
-	void AddTag(const FGameplayTag& TagToAdd);
+	UE_API void AddTag(const FGameplayTag& TagToAdd);
 
 	/** Remove a tag that will be omitted from any inherited tags */
-	void RemoveTag(const FGameplayTag& TagToRemove);
+	UE_API void RemoveTag(const FGameplayTag& TagToRemove);
 
-	bool operator==(const FInheritedTagContainer& Other) const;
-	bool operator!=(const FInheritedTagContainer& Other) const;
+	UE_API bool operator==(const FInheritedTagContainer& Other) const;
+	UE_API bool operator!=(const FInheritedTagContainer& Other) const;
 };
 
 /** Gameplay effect duration policies */
@@ -741,7 +743,7 @@ private:
 
 /** Saves list of modified attributes, to use for gameplay cues or later processing */
 USTRUCT()
-struct GAMEPLAYABILITIES_API FGameplayEffectModifiedAttribute
+struct FGameplayEffectModifiedAttribute
 {
 	GENERATED_USTRUCT_BODY()
 
@@ -758,7 +760,7 @@ struct GAMEPLAYABILITIES_API FGameplayEffectModifiedAttribute
 
 /** Struct used to hold the result of a gameplay attribute capture; Initially seeded by definition data, but then populated by ability system component when appropriate */
 USTRUCT()
-struct GAMEPLAYABILITIES_API FGameplayEffectAttributeCaptureSpec
+struct FGameplayEffectAttributeCaptureSpec
 {
 	// Allow these as friends so they can seed the aggregator, which we don't otherwise want exposed
 	friend struct FActiveGameplayEffectsContainer;
@@ -767,15 +769,15 @@ struct GAMEPLAYABILITIES_API FGameplayEffectAttributeCaptureSpec
 	GENERATED_USTRUCT_BODY()
 
 	// Constructors
-	FGameplayEffectAttributeCaptureSpec();
-	FGameplayEffectAttributeCaptureSpec(const FGameplayEffectAttributeCaptureDefinition& InDefinition);
+	UE_API FGameplayEffectAttributeCaptureSpec();
+	UE_API FGameplayEffectAttributeCaptureSpec(const FGameplayEffectAttributeCaptureDefinition& InDefinition);
 
 	/**
 	 * Returns whether the spec actually has a valid capture yet or not
 	 * 
 	 * @return True if the spec has a valid attribute capture, false if it does not
 	 */
-	bool HasValidCapture() const;
+	UE_API bool HasValidCapture() const;
 
 	/**
 	 * Attempts to calculate the magnitude of the captured attribute given the specified parameters. Can fail if the spec doesn't have
@@ -786,7 +788,7 @@ struct GAMEPLAYABILITIES_API FGameplayEffectAttributeCaptureSpec
 	 * 
 	 * @return True if the magnitude was successfully calculated, false if it was not
 	 */
-	bool AttemptCalculateAttributeMagnitude(const FAggregatorEvaluateParameters& InEvalParams, OUT float& OutMagnitude) const;
+	UE_API bool AttemptCalculateAttributeMagnitude(const FAggregatorEvaluateParameters& InEvalParams, OUT float& OutMagnitude) const;
 
 	/**
 	 * Attempts to calculate the magnitude of the captured attribute given the specified parameters, up to the specified evaluation channel (inclusive).
@@ -798,7 +800,7 @@ struct GAMEPLAYABILITIES_API FGameplayEffectAttributeCaptureSpec
 	 * 
 	 * @return True if the magnitude was successfully calculated, false if it was not
 	 */
-	bool AttemptCalculateAttributeMagnitudeUpToChannel(const FAggregatorEvaluateParameters& InEvalParams, EGameplayModEvaluationChannel FinalChannel, OUT float& OutMagnitude) const;
+	UE_API bool AttemptCalculateAttributeMagnitudeUpToChannel(const FAggregatorEvaluateParameters& InEvalParams, EGameplayModEvaluationChannel FinalChannel, OUT float& OutMagnitude) const;
 
 	/**
 	 * Attempts to calculate the magnitude of the captured attribute given the specified parameters, including a starting base value. 
@@ -810,7 +812,7 @@ struct GAMEPLAYABILITIES_API FGameplayEffectAttributeCaptureSpec
 	 * 
 	 * @return True if the magnitude was successfully calculated, false if it was not
 	 */
-	bool AttemptCalculateAttributeMagnitudeWithBase(const FAggregatorEvaluateParameters& InEvalParams, float InBaseValue, OUT float& OutMagnitude) const;
+	UE_API bool AttemptCalculateAttributeMagnitudeWithBase(const FAggregatorEvaluateParameters& InEvalParams, float InBaseValue, OUT float& OutMagnitude) const;
 
 	/**
 	 * Attempts to calculate the base value of the captured attribute given the specified parameters. Can fail if the spec doesn't have
@@ -820,7 +822,7 @@ struct GAMEPLAYABILITIES_API FGameplayEffectAttributeCaptureSpec
 	 * 
 	 * @return True if the base value was successfully calculated, false if it was not
 	 */
-	bool AttemptCalculateAttributeBaseValue(OUT float& OutBaseValue) const;
+	UE_API bool AttemptCalculateAttributeBaseValue(OUT float& OutBaseValue) const;
 
 	/**
 	 * Attempts to calculate the "bonus" magnitude (final - base value) of the captured attribute given the specified parameters. Can fail if the spec doesn't have
@@ -831,7 +833,7 @@ struct GAMEPLAYABILITIES_API FGameplayEffectAttributeCaptureSpec
 	 * 
 	 * @return True if the bonus magnitude was successfully calculated, false if it was not
 	 */
-	bool AttemptCalculateAttributeBonusMagnitude(const FAggregatorEvaluateParameters& InEvalParams, OUT float& OutBonusMagnitude) const;
+	UE_API bool AttemptCalculateAttributeBonusMagnitude(const FAggregatorEvaluateParameters& InEvalParams, OUT float& OutBonusMagnitude) const;
 
 	/**
 	 * Attempts to calculate the contribution of the specified GE to the captured attribute given the specified parameters. Can fail if the spec doesn't have
@@ -843,7 +845,7 @@ struct GAMEPLAYABILITIES_API FGameplayEffectAttributeCaptureSpec
 	 *
 	 * @return True if the bonus magnitude was successfully calculated, false if it was not
 	 */
-	bool AttemptCalculateAttributeContributionMagnitude(const FAggregatorEvaluateParameters& InEvalParams, FActiveGameplayEffectHandle ActiveHandle, OUT float& OutBonusMagnitude) const;
+	UE_API bool AttemptCalculateAttributeContributionMagnitude(const FAggregatorEvaluateParameters& InEvalParams, FActiveGameplayEffectHandle ActiveHandle, OUT float& OutBonusMagnitude) const;
 
 	/**
 	 * Attempts to populate the specified aggregator with a snapshot of the backing captured aggregator. Can fail if the spec doesn't have
@@ -853,7 +855,7 @@ struct GAMEPLAYABILITIES_API FGameplayEffectAttributeCaptureSpec
 	 *
 	 * @return True if the aggregator was successfully snapshotted, false if it was not
 	 */
-	bool AttemptGetAttributeAggregatorSnapshot(OUT FAggregator& OutAggregatorSnapshot) const;
+	UE_API bool AttemptGetAttributeAggregatorSnapshot(OUT FAggregator& OutAggregatorSnapshot) const;
 
 	/**
 	 * Attempts to populate the specified aggregator with all of the mods of the backing captured aggregator. Can fail if the spec doesn't have
@@ -863,25 +865,25 @@ struct GAMEPLAYABILITIES_API FGameplayEffectAttributeCaptureSpec
 	 *
 	 * @return True if the aggregator had mods successfully added to it, false if it did not
 	 */
-	bool AttemptAddAggregatorModsToAggregator(OUT FAggregator& OutAggregatorToAddTo) const;
+	UE_API bool AttemptAddAggregatorModsToAggregator(OUT FAggregator& OutAggregatorToAddTo) const;
 	
 	/** Gathers made for a given capture. Note all mods are returned but only some will be qualified (use Qualified() func to determine) */
-	bool AttemptGatherAttributeMods(const FAggregatorEvaluateParameters& InEvalParams, OUT TMap<EGameplayModEvaluationChannel, const TArray<FAggregatorMod>*>& OutModMap) const;
+	UE_API bool AttemptGatherAttributeMods(const FAggregatorEvaluateParameters& InEvalParams, OUT TMap<EGameplayModEvaluationChannel, const TArray<FAggregatorMod>*>& OutModMap) const;
 	
 	/** Simple accessor to backing capture definition */
-	const FGameplayEffectAttributeCaptureDefinition& GetBackingDefinition() const;
+	UE_API const FGameplayEffectAttributeCaptureDefinition& GetBackingDefinition() const;
 
 	/** Register this handle with linked aggregators */
-	void RegisterLinkedAggregatorCallback(FActiveGameplayEffectHandle Handle) const;
+	UE_API void RegisterLinkedAggregatorCallback(FActiveGameplayEffectHandle Handle) const;
 
 	/** Unregister this handle with linked aggregators */
-	void UnregisterLinkedAggregatorCallback(FActiveGameplayEffectHandle Handle) const;
+	UE_API void UnregisterLinkedAggregatorCallback(FActiveGameplayEffectHandle Handle) const;
 	
 	/** Return true if this capture should be recalculated if the given aggregator has changed */
-	bool ShouldRefreshLinkedAggregator(const FAggregator* ChangedAggregator) const;
+	UE_API bool ShouldRefreshLinkedAggregator(const FAggregator* ChangedAggregator) const;
 
 	/** Swaps any internal references From aggregator To aggregator. Used when cloning */
-	void SwapAggregator(FAggregatorRef From, FAggregatorRef To);
+	UE_API void SwapAggregator(FAggregatorRef From, FAggregatorRef To);
 		
 private:
 
@@ -895,21 +897,21 @@ private:
 
 /** Struct used to handle a collection of captured source and target attributes */
 USTRUCT()
-struct GAMEPLAYABILITIES_API FGameplayEffectAttributeCaptureSpecContainer
+struct FGameplayEffectAttributeCaptureSpecContainer
 {
 	GENERATED_USTRUCT_BODY()
 
 public:
 
-	FGameplayEffectAttributeCaptureSpecContainer();
+	UE_API FGameplayEffectAttributeCaptureSpecContainer();
 
-	FGameplayEffectAttributeCaptureSpecContainer(FGameplayEffectAttributeCaptureSpecContainer&& Other);
+	UE_API FGameplayEffectAttributeCaptureSpecContainer(FGameplayEffectAttributeCaptureSpecContainer&& Other);
 
-	FGameplayEffectAttributeCaptureSpecContainer(const FGameplayEffectAttributeCaptureSpecContainer& Other);
+	UE_API FGameplayEffectAttributeCaptureSpecContainer(const FGameplayEffectAttributeCaptureSpecContainer& Other);
 
-	FGameplayEffectAttributeCaptureSpecContainer& operator=(FGameplayEffectAttributeCaptureSpecContainer&& Other);
+	UE_API FGameplayEffectAttributeCaptureSpecContainer& operator=(FGameplayEffectAttributeCaptureSpecContainer&& Other);
 
-	FGameplayEffectAttributeCaptureSpecContainer& operator=(const FGameplayEffectAttributeCaptureSpecContainer& Other);
+	UE_API FGameplayEffectAttributeCaptureSpecContainer& operator=(const FGameplayEffectAttributeCaptureSpecContainer& Other);
 
 	/**
 	 * Add a definition to be captured by the owner of the container. Will not add the definition if its exact
@@ -917,7 +919,7 @@ public:
 	 * 
 	 * @param InCaptureDefinition	Definition to capture with
 	 */
-	void AddCaptureDefinition(const FGameplayEffectAttributeCaptureDefinition& InCaptureDefinition);
+	UE_API void AddCaptureDefinition(const FGameplayEffectAttributeCaptureDefinition& InCaptureDefinition);
 
 	/**
 	 * Capture source or target attributes from the specified component. Should be called by the container's owner.
@@ -925,7 +927,7 @@ public:
 	 * @param InAbilitySystemComponent	Component to capture attributes from
 	 * @param InCaptureSource			Whether to capture attributes as source or target
 	 */
-	void CaptureAttributes(class UAbilitySystemComponent* InAbilitySystemComponent, EGameplayEffectAttributeCaptureSource InCaptureSource);
+	UE_API void CaptureAttributes(class UAbilitySystemComponent* InAbilitySystemComponent, EGameplayEffectAttributeCaptureSource InCaptureSource);
 
 	/**
 	 * Find a capture spec within the container matching the specified capture definition, if possible.
@@ -935,7 +937,7 @@ public:
 	 * 
 	 * @return The found attribute spec matching the specified search params, if any
 	 */
-	const FGameplayEffectAttributeCaptureSpec* FindCaptureSpecByDefinition(const FGameplayEffectAttributeCaptureDefinition& InDefinition, bool bOnlyIncludeValidCapture) const;
+	UE_API const FGameplayEffectAttributeCaptureSpec* FindCaptureSpecByDefinition(const FGameplayEffectAttributeCaptureDefinition& InDefinition, bool bOnlyIncludeValidCapture) const;
 
 	/**
 	 * Determines if the container has specs with valid captures for all of the specified definitions.
@@ -944,19 +946,19 @@ public:
 	 * 
 	 * @return True if the container has valid capture attributes for all of the specified definitions, false if it does not
 	 */
-	bool HasValidCapturedAttributes(const TArray<FGameplayEffectAttributeCaptureDefinition>& InCaptureDefsToCheck) const;
+	UE_API bool HasValidCapturedAttributes(const TArray<FGameplayEffectAttributeCaptureDefinition>& InCaptureDefsToCheck) const;
 
 	/** Returns whether the container has at least one spec w/o snapshotted attributes */
-	bool HasNonSnapshottedAttributes() const;
+	UE_API bool HasNonSnapshottedAttributes() const;
 
 	/** Registers any linked aggregators to notify this active handle if they are dirtied */
-	void RegisterLinkedAggregatorCallbacks(FActiveGameplayEffectHandle Handle) const;
+	UE_API void RegisterLinkedAggregatorCallbacks(FActiveGameplayEffectHandle Handle) const;
 
 	/** Unregisters any linked aggregators from notifying this active handle if they are dirtied */
-	void UnregisterLinkedAggregatorCallbacks(FActiveGameplayEffectHandle Handle) const;
+	UE_API void UnregisterLinkedAggregatorCallbacks(FActiveGameplayEffectHandle Handle) const;
 
 	/** Swaps any internal references From aggregator To aggregator. Used when cloning */
-	void SwapAggregator(FAggregatorRef From, FAggregatorRef To);
+	UE_API void SwapAggregator(FAggregatorRef From, FAggregatorRef To);
 
 private:
 
@@ -983,7 +985,7 @@ private:
  * is still distinct from an FActiveGameplayEffect which in an applied instance of an FGameplayEffectSpec.
  */
 USTRUCT(BlueprintType)
-struct GAMEPLAYABILITIES_API FGameplayEffectSpec
+struct FGameplayEffectSpec
 {
 	GENERATED_USTRUCT_BODY()
 
@@ -991,31 +993,31 @@ struct GAMEPLAYABILITIES_API FGameplayEffectSpec
 	//	IMPORTANT: Any state added to FGameplayEffectSpec must be handled in the move/copy constructor/operator!
 	// --------------------------------------------------------------------------------------------------------------------------
 
-	FGameplayEffectSpec();
+	UE_API FGameplayEffectSpec();
 
-	FGameplayEffectSpec(const UGameplayEffect* InDef, const FGameplayEffectContextHandle& InEffectContext, float Level = FGameplayEffectConstants::INVALID_LEVEL);
+	UE_API FGameplayEffectSpec(const UGameplayEffect* InDef, const FGameplayEffectContextHandle& InEffectContext, float Level = FGameplayEffectConstants::INVALID_LEVEL);
 
-	FGameplayEffectSpec(const FGameplayEffectSpec& Other);
+	UE_API FGameplayEffectSpec(const FGameplayEffectSpec& Other);
 
-	FGameplayEffectSpec(const FGameplayEffectSpec& Other, const FGameplayEffectContextHandle& InEffectContext);		//For cloning, copy all attributes, but set a new effectContext.
+	UE_API FGameplayEffectSpec(const FGameplayEffectSpec& Other, const FGameplayEffectContextHandle& InEffectContext);		//For cloning, copy all attributes, but set a new effectContext.
 
-	FGameplayEffectSpec(FGameplayEffectSpec&& Other);
+	UE_API FGameplayEffectSpec(FGameplayEffectSpec&& Other);
 
-	FGameplayEffectSpec& operator=(FGameplayEffectSpec&& Other);
+	UE_API FGameplayEffectSpec& operator=(FGameplayEffectSpec&& Other);
 
-	FGameplayEffectSpec& operator=(const FGameplayEffectSpec& Other);
+	UE_API FGameplayEffectSpec& operator=(const FGameplayEffectSpec& Other);
 
 	/** Can be called manually but it is preferred to use the 3 parameter constructor */
-	void Initialize(const UGameplayEffect* InDef, const FGameplayEffectContextHandle& InEffectContext, float Level = FGameplayEffectConstants::INVALID_LEVEL);
+	UE_API void Initialize(const UGameplayEffect* InDef, const FGameplayEffectContextHandle& InEffectContext, float Level = FGameplayEffectConstants::INVALID_LEVEL);
 
 	/** Initialize the spec as a linked spec. The original spec's context is preserved except for the original GE asset tags, which are stripped out */
-	void InitializeFromLinkedSpec(const UGameplayEffect* InDef, const FGameplayEffectSpec& OriginalSpec);
+	UE_API void InitializeFromLinkedSpec(const UGameplayEffect* InDef, const FGameplayEffectSpec& OriginalSpec);
 
 	/** Copies SetbyCallerMagnitudes from OriginalSpec into this */
-	void CopySetByCallerMagnitudes(const FGameplayEffectSpec& OriginalSpec);
+	UE_API void CopySetByCallerMagnitudes(const FGameplayEffectSpec& OriginalSpec);
 
 	/** Copies SetbuCallerMagnitudes, but only if magnitudes don't exist in our map (slower but preserves data) */
-	void MergeSetByCallerMagnitudes(const TMap<FGameplayTag, float>& Magnitudes);
+	UE_API void MergeSetByCallerMagnitudes(const TMap<FGameplayTag, float>& Magnitudes);
 
 	/**
 	 * Determines if the spec has capture specs with valid captures for all of the specified definitions.
@@ -1024,14 +1026,14 @@ struct GAMEPLAYABILITIES_API FGameplayEffectSpec
 	 * 
 	 * @return True if the container has valid capture attributes for all of the specified definitions, false if it does not
 	 */
-	bool HasValidCapturedAttributes(const TArray<FGameplayEffectAttributeCaptureDefinition>& InCaptureDefsToCheck) const;
+	UE_API bool HasValidCapturedAttributes(const TArray<FGameplayEffectAttributeCaptureDefinition>& InCaptureDefsToCheck) const;
 
 	/** Looks for an existing modified attribute struct, may return NULL */
-	const FGameplayEffectModifiedAttribute* GetModifiedAttribute(const FGameplayAttribute& Attribute) const;
-	FGameplayEffectModifiedAttribute* GetModifiedAttribute(const FGameplayAttribute& Attribute);
+	UE_API const FGameplayEffectModifiedAttribute* GetModifiedAttribute(const FGameplayAttribute& Attribute) const;
+	UE_API FGameplayEffectModifiedAttribute* GetModifiedAttribute(const FGameplayAttribute& Attribute);
 
 	/** Adds a new modified attribute struct, will always add so check to see if it exists first */
-	FGameplayEffectModifiedAttribute* AddModifiedAttribute(const FGameplayAttribute& Attribute);
+	UE_API FGameplayEffectModifiedAttribute* AddModifiedAttribute(const FGameplayAttribute& Attribute);
 
 	/**
 	 * Helper function to attempt to calculate the duration of the spec from its GE definition
@@ -1040,27 +1042,27 @@ struct GAMEPLAYABILITIES_API FGameplayEffectSpec
 	 * 
 	 * @return True if the calculation was successful, false if it was not
 	 */
-	bool AttemptCalculateDurationFromDef(OUT float& OutDefDuration) const;
+	UE_API bool AttemptCalculateDurationFromDef(OUT float& OutDefDuration) const;
 
 	/** Sets duration. This should only be called as the GameplayEffect is being created and applied; Ignores calls after attribute capture */
-	void SetDuration(float NewDuration, bool bLockDuration);
+	UE_API void SetDuration(float NewDuration, bool bLockDuration);
 
-	float GetDuration() const;
+	UE_API float GetDuration() const;
 
 	/** Returns the Period for the effect. If DurationPolicy is Instant this will forcibly return NO_PERIOD */
-	float GetPeriod() const;
+	UE_API float GetPeriod() const;
 	
 	UE_DEPRECATED(5.3, "This no longer applies.  Use UChanceToApplyGameplayEffectComponent instead")
 	float GetChanceToApplyToTarget() const { return 1.0f; }
 
 	/** Sets the stack count for this GE to NewStackCount if stacking is supported. */
-	void SetStackCount(int32 NewStackCount);
+	UE_API void SetStackCount(int32 NewStackCount);
 
 	/** Returns the stack count for this GE spec. */
-	int32 GetStackCount() const;
+	UE_API int32 GetStackCount() const;
 
 	/** Set the context info: who and where this spec came from. */
-	void SetContext(FGameplayEffectContextHandle NewEffectContext, bool bSkipRecaptureSourceActorTags = false);
+	UE_API void SetContext(FGameplayEffectContextHandle NewEffectContext, bool bSkipRecaptureSourceActorTags = false);
 
 	FGameplayEffectContextHandle GetContext() const
 	{
@@ -1068,33 +1070,33 @@ struct GAMEPLAYABILITIES_API FGameplayEffectSpec
 	}
 
 	/** Appends all tags granted by this gameplay effect spec */
-	void GetAllGrantedTags(OUT FGameplayTagContainer& OutContainer) const;
+	UE_API void GetAllGrantedTags(OUT FGameplayTagContainer& OutContainer) const;
 
 	/** Appends all blocked ability tags granted by this gameplay effect spec */
-	void GetAllBlockedAbilityTags(OUT FGameplayTagContainer& OutContainer) const;
+	UE_API void GetAllBlockedAbilityTags(OUT FGameplayTagContainer& OutContainer) const;
 
 	/** Appends all tags that apply to this gameplay effect spec */
-	void GetAllAssetTags(OUT FGameplayTagContainer& OutContainer) const;
+	UE_API void GetAllAssetTags(OUT FGameplayTagContainer& OutContainer) const;
 
 	/** Sets the magnitude of a SetByCaller modifier */
-	void SetSetByCallerMagnitude(FName DataName, float Magnitude);
+	UE_API void SetSetByCallerMagnitude(FName DataName, float Magnitude);
 
 	/** Sets the magnitude of a SetByCaller modifier */
-	void SetSetByCallerMagnitude(FGameplayTag DataTag, float Magnitude);
+	UE_API void SetSetByCallerMagnitude(FGameplayTag DataTag, float Magnitude);
 
 	/** Returns the magnitude of a SetByCaller modifier. Will return 0.f and Warn if the magnitude has not been set. */
-	float GetSetByCallerMagnitude(FName DataName, bool WarnIfNotFound = true, float DefaultIfNotFound = 0.f) const;
+	UE_API float GetSetByCallerMagnitude(FName DataName, bool WarnIfNotFound = true, float DefaultIfNotFound = 0.f) const;
 
 	/** Returns the magnitude of a SetByCaller modifier. Will return 0.f and Warn if the magnitude has not been set. */
-	float GetSetByCallerMagnitude(FGameplayTag DataTag, bool WarnIfNotFound = true, float DefaultIfNotFound = 0.f) const;
+	UE_API float GetSetByCallerMagnitude(FGameplayTag DataTag, bool WarnIfNotFound = true, float DefaultIfNotFound = 0.f) const;
 
-	void SetLevel(float InLevel);
+	UE_API void SetLevel(float InLevel);
 
-	float GetLevel() const;
+	UE_API float GetLevel() const;
 
-	void PrintAll() const;
+	UE_API void PrintAll() const;
 
-	FString ToSimpleString() const;
+	UE_API FString ToSimpleString() const;
 
 	const FGameplayEffectContextHandle& GetEffectContext() const
 	{
@@ -1106,7 +1108,7 @@ struct GAMEPLAYABILITIES_API FGameplayEffectSpec
 		EffectContext = EffectContext.Duplicate();
 	}
 
-	void CaptureAttributeDataFromTarget(UAbilitySystemComponent* TargetAbilitySystemComponent);
+	UE_API void CaptureAttributeDataFromTarget(UAbilitySystemComponent* TargetAbilitySystemComponent);
 
 	/**
 	 * Get the computed magnitude of the modifier on the spec with the specified index
@@ -1117,7 +1119,7 @@ struct GAMEPLAYABILITIES_API FGameplayEffectSpec
 	 * @return Computed magnitude
 	 */
 	UE_DEPRECATED(5.6, "Versions using bFactorInStackCount are deprecated, please use GetModifierMagnitude(int32 ModifierIdx) instead.")
-	float GetModifierMagnitude(int32 ModifierIdx, bool bFactorInStackCount) const;
+	UE_API float GetModifierMagnitude(int32 ModifierIdx, bool bFactorInStackCount) const;
 
 	/**
 	 * Get the computed magnitude of the modifier on the spec with the specified index
@@ -1126,39 +1128,39 @@ struct GAMEPLAYABILITIES_API FGameplayEffectSpec
 	 * 
 	 * @return Computed magnitude
 	 */
-	float GetModifierMagnitude(int32 ModifierIdx) const;
+	UE_API float GetModifierMagnitude(int32 ModifierIdx) const;
 
 	/** Fills out the modifier magnitudes inside the Modifier Specs */
-	void CalculateModifierMagnitudes();
+	UE_API void CalculateModifierMagnitudes();
 
 	/** Recapture attributes from source and target for cloning */
-	void RecaptureAttributeDataForClone(UAbilitySystemComponent* OriginalASC, UAbilitySystemComponent* NewASC);
+	UE_API void RecaptureAttributeDataForClone(UAbilitySystemComponent* OriginalASC, UAbilitySystemComponent* NewASC);
 
 	/** Recaptures source actor tags of this spec without modifying anything else */
-	void RecaptureSourceActorTags();
+	UE_API void RecaptureSourceActorTags();
 
 	/** Helper function to initialize all of the capture definitions required by the spec */
-	void SetupAttributeCaptureDefinitions();
+	UE_API void SetupAttributeCaptureDefinitions();
 
 	/** Helper function that returns the duration after applying relevant modifiers from the source and target ability system components */
-	float CalculateModifiedDuration() const;
+	UE_API float CalculateModifiedDuration() const;
 
 	/** Dynamically add an asset tag not originally from the source GE definition; Added to DynamicAssetTags as well as injected into the captured source spec tags */
-	void AddDynamicAssetTag(const FGameplayTag& TagToAdd);
+	UE_API void AddDynamicAssetTag(const FGameplayTag& TagToAdd);
 
 	/** Dynamically append asset tags not originally from the source GE definition; Added to DynamicAssetTags as well as injected into the captured source spec tags */
-	void AppendDynamicAssetTags(const FGameplayTagContainer& TagsToAppend);
+	UE_API void AppendDynamicAssetTags(const FGameplayTagContainer& TagsToAppend);
 
 	/** Simple const accessor to the dynamic asset tags */
-	const FGameplayTagContainer& GetDynamicAssetTags() const;
+	UE_API const FGameplayTagContainer& GetDynamicAssetTags() const;
 
 #if ENABLE_VISUAL_LOG
-	FVisualLogStatusCategory GrabVisLogStatus() const;
+	UE_API FVisualLogStatusCategory GrabVisLogStatus() const;
 #endif
 
 private:
 
-	void CaptureDataFromSource(bool bSkipRecaptureSourceActorTags = false);
+	UE_API void CaptureDataFromSource(bool bSkipRecaptureSourceActorTags = false);
 
 public:
 
@@ -1255,13 +1257,13 @@ private:
 
 /** This is a cut down version of the gameplay effect spec used for RPCs. */
 USTRUCT()
-struct GAMEPLAYABILITIES_API FGameplayEffectSpecForRPC
+struct FGameplayEffectSpecForRPC
 {
 	GENERATED_USTRUCT_BODY()
 
-	FGameplayEffectSpecForRPC();
+	UE_API FGameplayEffectSpecForRPC();
 
-	FGameplayEffectSpecForRPC(const FGameplayEffectSpec& InSpec);
+	UE_API FGameplayEffectSpecForRPC(const FGameplayEffectSpec& InSpec);
 
 	/** GameplayEfect definition. The static data that this spec points to. */
 	UPROPERTY()
@@ -1300,9 +1302,9 @@ struct GAMEPLAYABILITIES_API FGameplayEffectSpecForRPC
 		return AbilityLevel;
 	}
 
-	FString ToSimpleString() const;
+	UE_API FString ToSimpleString() const;
 
-	const FGameplayEffectModifiedAttribute* GetModifiedAttribute(const FGameplayAttribute& Attribute) const;
+	UE_API const FGameplayEffectModifiedAttribute* GetModifiedAttribute(const FGameplayAttribute& Attribute) const;
 };
 
 
@@ -1314,7 +1316,7 @@ struct GAMEPLAYABILITIES_API FGameplayEffectSpecForRPC
  *  -Replication callbacks
  */
 USTRUCT(BlueprintType)
-struct GAMEPLAYABILITIES_API FActiveGameplayEffect : public FFastArraySerializerItem
+struct FActiveGameplayEffect : public FFastArraySerializerItem
 {
 	GENERATED_USTRUCT_BODY()
 
@@ -1323,13 +1325,13 @@ struct GAMEPLAYABILITIES_API FActiveGameplayEffect : public FFastArraySerializer
 	// ---------------------------------------------------------------------------------------------------------------------------------
 
 	FActiveGameplayEffect() = default;
-	FActiveGameplayEffect(FActiveGameplayEffectHandle InHandle, const FGameplayEffectSpec& InSpec, float CurrentWorldTime, float InStartServerWorldTime, FPredictionKey InPredictionKey);
+	UE_API FActiveGameplayEffect(FActiveGameplayEffectHandle InHandle, const FGameplayEffectSpec& InSpec, float CurrentWorldTime, float InStartServerWorldTime, FPredictionKey InPredictionKey);
 
 	// These need to be defined because we need to omit PendingNext from move operations and the base class isn't trivially copyable
-	FActiveGameplayEffect(const FActiveGameplayEffect& Other);
-	FActiveGameplayEffect(FActiveGameplayEffect&& Other);
-	FActiveGameplayEffect& operator=(FActiveGameplayEffect&& other);
-	FActiveGameplayEffect& operator=(const FActiveGameplayEffect& other);
+	UE_API FActiveGameplayEffect(const FActiveGameplayEffect& Other);
+	UE_API FActiveGameplayEffect(FActiveGameplayEffect&& Other);
+	UE_API FActiveGameplayEffect& operator=(FActiveGameplayEffect&& other);
+	UE_API FActiveGameplayEffect& operator=(const FActiveGameplayEffect& other);
 
 	float GetTimeRemaining(float WorldTime) const
 	{
@@ -1355,26 +1357,26 @@ struct GAMEPLAYABILITIES_API FActiveGameplayEffect : public FFastArraySerializer
 
 	/** This was the core function that turns the ActiveGE 'on' or 'off.  That function can be carried out by UGameplayEffectComponents, @see UTargetTagRequirementsGameplayEffectComponent */
 	UE_DEPRECATED(5.3, "CheckOngoingTagRequirements has been deprecated in favor of UTargetTagRequirementsGameplayEffectComponent")
-	void CheckOngoingTagRequirements(const FGameplayTagContainer& OwnerTags, struct FActiveGameplayEffectsContainer& OwningContainer, bool bInvokeGameplayCueEvents = false);
+	UE_API void CheckOngoingTagRequirements(const FGameplayTagContainer& OwnerTags, struct FActiveGameplayEffectsContainer& OwningContainer, bool bInvokeGameplayCueEvents = false);
 	
 	/** Method to check if this effect should remove because the owner tags pass the RemovalTagRequirements requirement check */
 	UE_DEPRECATED(5.3, "CheckRemovalTagRequirements has been deprecated in favor of UTargetTagRequirementsGameplayEffectComponent")
-	bool CheckRemovalTagRequirements(const FGameplayTagContainer & OwnerTags, struct FActiveGameplayEffectsContainer& OwningContainer) const;
+	UE_API bool CheckRemovalTagRequirements(const FGameplayTagContainer & OwnerTags, struct FActiveGameplayEffectsContainer& OwningContainer) const;
 
-	void PrintAll() const;
+	UE_API void PrintAll() const;
 
-	void PreReplicatedRemove(const struct FActiveGameplayEffectsContainer &InArray);
-	void PostReplicatedAdd(const struct FActiveGameplayEffectsContainer &InArray);
-	void PostReplicatedChange(const struct FActiveGameplayEffectsContainer &InArray);
+	UE_API void PreReplicatedRemove(const struct FActiveGameplayEffectsContainer &InArray);
+	UE_API void PostReplicatedAdd(const struct FActiveGameplayEffectsContainer &InArray);
+	UE_API void PostReplicatedChange(const struct FActiveGameplayEffectsContainer &InArray);
 
 	/** Debug string used by Fast Array serialization */
-	FString GetDebugString();
+	UE_API FString GetDebugString();
 
 	/** Refreshes the cached StartWorldTime for this effect. To be used when the server/client world time delta changes significantly to keep the start time in sync. */
-	void RecomputeStartWorldTime(const FActiveGameplayEffectsContainer& InArray);
+	UE_API void RecomputeStartWorldTime(const FActiveGameplayEffectsContainer& InArray);
 
 	/** Refreshes the cached StartWorldTime for this effect. To be used when the server/client world time delta changes significantly to keep the start time in sync. */
-	void RecomputeStartWorldTime(const float WorldTime, const float ServerWorldTime);
+	UE_API void RecomputeStartWorldTime(const float WorldTime, const float ServerWorldTime);
 
 	bool operator==(const FActiveGameplayEffect& Other)
 	{
@@ -1435,18 +1437,18 @@ DECLARE_DYNAMIC_DELEGATE_TwoParams(FActiveGameplayEffectQueryCustomMatch_Dynamic
 
 /** Every set condition within this query must match in order for the query to match. i.e. individual query elements are ANDed together. */
 USTRUCT(BlueprintType)
-struct GAMEPLAYABILITIES_API FGameplayEffectQuery
+struct FGameplayEffectQuery
 {
 	GENERATED_USTRUCT_BODY()
 
 public:
 	// ctors and operators
-	FGameplayEffectQuery();
-	FGameplayEffectQuery(const FGameplayEffectQuery& Other);
-	FGameplayEffectQuery(FActiveGameplayEffectQueryCustomMatch InCustomMatchDelegate);
-	FGameplayEffectQuery(FGameplayEffectQuery&& Other);
-	FGameplayEffectQuery& operator=(FGameplayEffectQuery&& Other);
-	FGameplayEffectQuery& operator=(const FGameplayEffectQuery& Other);
+	UE_API FGameplayEffectQuery();
+	UE_API FGameplayEffectQuery(const FGameplayEffectQuery& Other);
+	UE_API FGameplayEffectQuery(FActiveGameplayEffectQueryCustomMatch InCustomMatchDelegate);
+	UE_API FGameplayEffectQuery(FGameplayEffectQuery&& Other);
+	UE_API FGameplayEffectQuery& operator=(FGameplayEffectQuery&& Other);
+	UE_API FGameplayEffectQuery& operator=(const FGameplayEffectQuery& Other);
 
 	/** Native delegate for providing custom matching conditions. */
 	FActiveGameplayEffectQueryCustomMatch CustomMatchDelegate;
@@ -1487,41 +1489,41 @@ public:
 	TArray<FActiveGameplayEffectHandle> IgnoreHandles;
 
 	/** Returns true if Effect matches all specified criteria of this query, including CustomMatch delegates if bound. Returns false otherwise. */
-	bool Matches(const FActiveGameplayEffect& Effect) const;
+	UE_API bool Matches(const FActiveGameplayEffect& Effect) const;
 
 	/** Returns true if Effect matches all specified criteria of this query. This does NOT check FActiveGameplayEffectQueryCustomMatch since this is performed on the spec (possibly prior to applying).
 	 *	Note: it would be reasonable to support a custom delegate that operated on the FGameplayEffectSpec itself.
 	 */
-	bool Matches(const FGameplayEffectSpec& Effect) const;
+	UE_API bool Matches(const FGameplayEffectSpec& Effect) const;
 
 	/** Returns true if the query is empty/default. E.g., it has no data set. */
-	bool IsEmpty() const;
+	UE_API bool IsEmpty() const;
 
 	// Shortcuts for easily creating common query types 
 
 	/** Creates an effect query that will match if there are any common tags between the given tags and an ActiveGameplayEffect's owning tags */
-	static FGameplayEffectQuery MakeQuery_MatchAnyOwningTags(const FGameplayTagContainer& InTags);
+	static UE_API FGameplayEffectQuery MakeQuery_MatchAnyOwningTags(const FGameplayTagContainer& InTags);
 	/** Creates an effect query that will match if all of the given tags are in the ActiveGameplayEffect's owning tags */
-	static FGameplayEffectQuery MakeQuery_MatchAllOwningTags(const FGameplayTagContainer& InTags);
+	static UE_API FGameplayEffectQuery MakeQuery_MatchAllOwningTags(const FGameplayTagContainer& InTags);
 	/** Creates an effect query that will match if there are no common tags between the given tags and an ActiveGameplayEffect's owning tags */
-	static FGameplayEffectQuery MakeQuery_MatchNoOwningTags(const FGameplayTagContainer& InTags);
+	static UE_API FGameplayEffectQuery MakeQuery_MatchNoOwningTags(const FGameplayTagContainer& InTags);
 	
 	/** Creates an effect query that will match if there are any common tags between the given tags and an ActiveGameplayEffect's tags */
-	static FGameplayEffectQuery MakeQuery_MatchAnyEffectTags(const FGameplayTagContainer& InTags);
+	static UE_API FGameplayEffectQuery MakeQuery_MatchAnyEffectTags(const FGameplayTagContainer& InTags);
 	/** Creates an effect query that will match if all of the given tags are in the ActiveGameplayEffect's tags */
-	static FGameplayEffectQuery MakeQuery_MatchAllEffectTags(const FGameplayTagContainer& InTags);
+	static UE_API FGameplayEffectQuery MakeQuery_MatchAllEffectTags(const FGameplayTagContainer& InTags);
 	/** Creates an effect query that will match if there are no common tags between the given tags and an ActiveGameplayEffect's tags */
-	static FGameplayEffectQuery MakeQuery_MatchNoEffectTags(const FGameplayTagContainer& InTags);
+	static UE_API FGameplayEffectQuery MakeQuery_MatchNoEffectTags(const FGameplayTagContainer& InTags);
 
 	/** Creates an effect query that will match if there are any common tags between the given tags and an ActiveGameplayEffect's source spec tags */
-	static FGameplayEffectQuery MakeQuery_MatchAnySourceSpecTags(const FGameplayTagContainer& InTags);
+	static UE_API FGameplayEffectQuery MakeQuery_MatchAnySourceSpecTags(const FGameplayTagContainer& InTags);
 	/** Creates an effect query that will match if all of the given tags are in the ActiveGameplayEffect's source spec tags */
-	static FGameplayEffectQuery MakeQuery_MatchAllSourceSpecTags(const FGameplayTagContainer& InTags);
+	static UE_API FGameplayEffectQuery MakeQuery_MatchAllSourceSpecTags(const FGameplayTagContainer& InTags);
 	/** Creates an effect query that will match if there are no common tags between the given tags and an ActiveGameplayEffect's source spec tags */
-	static FGameplayEffectQuery MakeQuery_MatchNoSourceSpecTags(const FGameplayTagContainer& InTags);
+	static UE_API FGameplayEffectQuery MakeQuery_MatchNoSourceSpecTags(const FGameplayTagContainer& InTags);
 
-	bool operator==(const FGameplayEffectQuery& Other) const;
-	bool operator!=(const FGameplayEffectQuery& Other) const;
+	UE_API bool operator==(const FGameplayEffectQuery& Other) const;
+	UE_API bool operator!=(const FGameplayEffectQuery& Other) const;
 };
 
 /**
@@ -1616,7 +1618,7 @@ struct FCustomModifierDependencyHandle
  *
  */
 USTRUCT()
-struct GAMEPLAYABILITIES_API FActiveGameplayEffectsContainer : public FFastArraySerializer
+struct FActiveGameplayEffectsContainer : public FFastArraySerializer
 {
 	GENERATED_USTRUCT_BODY();
 
@@ -1630,8 +1632,8 @@ struct GAMEPLAYABILITIES_API FActiveGameplayEffectsContainer : public FFastArray
 	typedef FActiveGameplayEffectIterator<const FActiveGameplayEffect, FActiveGameplayEffectsContainer> ConstIterator;
 	typedef FActiveGameplayEffectIterator<FActiveGameplayEffect, FActiveGameplayEffectsContainer> Iterator;
 
-	FActiveGameplayEffectsContainer();
-	~FActiveGameplayEffectsContainer();
+	UE_API FActiveGameplayEffectsContainer();
+	UE_API ~FActiveGameplayEffectsContainer();
 
 	UAbilitySystemComponent* Owner;
 	bool OwnerIsNetAuthority;
@@ -1653,37 +1655,37 @@ struct GAMEPLAYABILITIES_API FActiveGameplayEffectsContainer : public FFastArray
 	TArray<DebugExecutedGameplayEffectData> DebugExecutedGameplayEffects;
 
 	/** Report our current state to the VisLog */
-	void DescribeSelfToVisLog(FVisualLogEntry* Snapshot) const;
+	UE_API void DescribeSelfToVisLog(FVisualLogEntry* Snapshot) const;
 #endif // ENABLE_VISUAL_LOG
 
-	void GetActiveGameplayEffectDataByAttribute(TMultiMap<FGameplayAttribute, FActiveGameplayEffectsContainer::DebugExecutedGameplayEffectData>& EffectMap) const;
+	UE_API void GetActiveGameplayEffectDataByAttribute(TMultiMap<FGameplayAttribute, FActiveGameplayEffectsContainer::DebugExecutedGameplayEffectData>& EffectMap) const;
 
-	void RegisterWithOwner(UAbilitySystemComponent* Owner);	
+	UE_API void RegisterWithOwner(UAbilitySystemComponent* Owner);	
 	
-	FActiveGameplayEffect* ApplyGameplayEffectSpec(const FGameplayEffectSpec& Spec, FPredictionKey& InPredictionKey, bool& bFoundExistingStackableGE);
+	UE_API FActiveGameplayEffect* ApplyGameplayEffectSpec(const FGameplayEffectSpec& Spec, FPredictionKey& InPredictionKey, bool& bFoundExistingStackableGE);
 
-	FActiveGameplayEffect* GetActiveGameplayEffect(const FActiveGameplayEffectHandle Handle);
+	UE_API FActiveGameplayEffect* GetActiveGameplayEffect(const FActiveGameplayEffectHandle Handle);
 
-	const FActiveGameplayEffect* GetActiveGameplayEffect(const FActiveGameplayEffectHandle Handle) const;
+	UE_API const FActiveGameplayEffect* GetActiveGameplayEffect(const FActiveGameplayEffectHandle Handle) const;
 
 	/** Predictively execute a given effect spec. Any attribute modifications and effect execution calculations in the effect will run and then if desired predict gameplay cues
 		@note: This method will not predictively run any conditional effects that may be set up in the effect that apply post execution and will only happen if/when this spec is
 		applied on the server. 
 		
 		@note: WARNING: This will locally perform attribute changes on your client so beware. */
-	void PredictivelyExecuteEffectSpec(FGameplayEffectSpec& Spec, FPredictionKey PredictionKey, const bool bPredictGameplayCues = false);
+	UE_API void PredictivelyExecuteEffectSpec(FGameplayEffectSpec& Spec, FPredictionKey PredictionKey, const bool bPredictGameplayCues = false);
 
-	void ExecuteActiveEffectsFrom(FGameplayEffectSpec &Spec, FPredictionKey PredictionKey = FPredictionKey() );
+	UE_API void ExecuteActiveEffectsFrom(FGameplayEffectSpec &Spec, FPredictionKey PredictionKey = FPredictionKey() );
 	
-	void ExecutePeriodicGameplayEffect(FActiveGameplayEffectHandle Handle);	// This should not be outward facing to the skill system API, should only be called by the owning AbilitySystemComponent
+	UE_API void ExecutePeriodicGameplayEffect(FActiveGameplayEffectHandle Handle);	// This should not be outward facing to the skill system API, should only be called by the owning AbilitySystemComponent
 
-	bool RemoveActiveGameplayEffect(FActiveGameplayEffectHandle Handle, int32 StacksToRemove);
+	UE_API bool RemoveActiveGameplayEffect(FActiveGameplayEffectHandle Handle, int32 StacksToRemove);
 
-	void GetGameplayEffectStartTimeAndDuration(FActiveGameplayEffectHandle Handle, float& EffectStartTime, float& EffectDuration) const;
+	UE_API void GetGameplayEffectStartTimeAndDuration(FActiveGameplayEffectHandle Handle, float& EffectStartTime, float& EffectDuration) const;
 
-	float GetGameplayEffectMagnitude(FActiveGameplayEffectHandle Handle, FGameplayAttribute Attribute) const;
+	UE_API float GetGameplayEffectMagnitude(FActiveGameplayEffectHandle Handle, FGameplayAttribute Attribute) const;
 
-	void SetActiveGameplayEffectLevel(FActiveGameplayEffectHandle ActiveHandle, int32 NewLevel);
+	UE_API void SetActiveGameplayEffectLevel(FActiveGameplayEffectHandle ActiveHandle, int32 NewLevel);
 
 	/**
 	 * Update a set-by-caller magnitude for the active effect to match the new value, if possible
@@ -1692,7 +1694,7 @@ struct GAMEPLAYABILITIES_API FActiveGameplayEffectsContainer : public FFastArray
 	 * @param SetByCallerTag	Set-by-caller tag to update
 	 * @param NewValue			New value of the set-by-caller magnitude
 	 */
-	void UpdateActiveGameplayEffectSetByCallerMagnitude(FActiveGameplayEffectHandle ActiveHandle, const FGameplayTag& SetByCallerTag, float NewValue);
+	UE_API void UpdateActiveGameplayEffectSetByCallerMagnitude(FActiveGameplayEffectHandle ActiveHandle, const FGameplayTag& SetByCallerTag, float NewValue);
 
 	/**
 	 * Update set-by-caller magnitudes for the active effect to match the new values, if possible; Replaces existing values
@@ -1700,16 +1702,16 @@ struct GAMEPLAYABILITIES_API FActiveGameplayEffectsContainer : public FFastArray
 	 * @param ActiveHandle			Handle of the active effect to update
 	 * @param NewSetByCallerValues	Map of set-by-caller tag to new magnitude
 	 */
-	void UpdateActiveGameplayEffectSetByCallerMagnitudes(FActiveGameplayEffectHandle ActiveHandle, const TMap<FGameplayTag, float>& NewSetByCallerValues);
+	UE_API void UpdateActiveGameplayEffectSetByCallerMagnitudes(FActiveGameplayEffectHandle ActiveHandle, const TMap<FGameplayTag, float>& NewSetByCallerValues);
 
-	void SetAttributeBaseValue(FGameplayAttribute Attribute, float NewBaseValue);
+	UE_API void SetAttributeBaseValue(FGameplayAttribute Attribute, float NewBaseValue);
 
-	float GetAttributeBaseValue(FGameplayAttribute Attribute) const;
+	UE_API float GetAttributeBaseValue(FGameplayAttribute Attribute) const;
 
-	float GetEffectContribution(const FAggregatorEvaluateParameters& Parameters, FActiveGameplayEffectHandle ActiveHandle, FGameplayAttribute Attribute);
+	UE_API float GetEffectContribution(const FAggregatorEvaluateParameters& Parameters, FActiveGameplayEffectHandle ActiveHandle, FGameplayAttribute Attribute);
 
 	/** Actually applies given mod to the attribute */
-	void ApplyModToAttribute(const FGameplayAttribute &Attribute, TEnumAsByte<EGameplayModOp::Type> ModifierOp, float ModifierMagnitude, const FGameplayEffectModCallbackData* ModData=nullptr);
+	UE_API void ApplyModToAttribute(const FGameplayAttribute &Attribute, TEnumAsByte<EGameplayModOp::Type> ModifierOp, float ModifierMagnitude, const FGameplayEffectModCallbackData* ModData=nullptr);
 
 	/**
 	 * Get the source tags from the gameplay spec represented by the specified handle, if possible
@@ -1718,7 +1720,7 @@ struct GAMEPLAYABILITIES_API FActiveGameplayEffectsContainer : public FFastArray
 	 * 
 	 * @return Source tags from the gameplay spec represented by the handle, if possible
 	 */
-	const FGameplayTagContainer* GetGameplayEffectSourceTagsFromHandle(FActiveGameplayEffectHandle Handle) const;
+	UE_API const FGameplayTagContainer* GetGameplayEffectSourceTagsFromHandle(FActiveGameplayEffectHandle Handle) const;
 
 	/**
 	 * Get the target tags from the gameplay spec represented by the specified handle, if possible
@@ -1727,16 +1729,16 @@ struct GAMEPLAYABILITIES_API FActiveGameplayEffectsContainer : public FFastArray
 	 * 
 	 * @return Target tags from the gameplay spec represented by the handle, if possible
 	 */
-	const FGameplayTagContainer* GetGameplayEffectTargetTagsFromHandle(FActiveGameplayEffectHandle Handle) const;
+	UE_API const FGameplayTagContainer* GetGameplayEffectTargetTagsFromHandle(FActiveGameplayEffectHandle Handle) const;
 
 	/**
 	 * Populate the specified capture spec with the data necessary to capture an attribute from the container
 	 * 
 	 * @param OutCaptureSpec	[OUT] Capture spec to populate with captured data
 	 */
-	void CaptureAttributeForGameplayEffect(OUT FGameplayEffectAttributeCaptureSpec& OutCaptureSpec);
+	UE_API void CaptureAttributeForGameplayEffect(OUT FGameplayEffectAttributeCaptureSpec& OutCaptureSpec);
 
-	void PrintAllGameplayEffects() const;
+	UE_API void PrintAllGameplayEffects() const;
 
 	/**
 	 *	Returns the total number of gameplay effects.
@@ -1757,10 +1759,10 @@ struct GAMEPLAYABILITIES_API FActiveGameplayEffectsContainer : public FFastArray
 		return GameplayEffects_Internal.Num() + NumPending;
 	}
 
-	void CheckDuration(FActiveGameplayEffectHandle Handle);
+	UE_API void CheckDuration(FActiveGameplayEffectHandle Handle);
 
 	/** Returns which ELifetimeCondition can be used for this instance to replicate to relevant connections. This can be used to change the condition if the property is set to use COND_Dynamic in an object's GetLifetimeReplicatedProps implementation. */
-	ELifetimeCondition GetReplicationCondition() const;
+	UE_API ELifetimeCondition GetReplicationCondition() const;
 
 	/** Set whether the container is using COND_Dynamic and setting the proper condition at runtime. */
 	void SetIsUsingReplicationCondition(bool bInIsUsingReplicationCondition) { bIsUsingReplicationCondition = bInIsUsingReplicationCondition; }
@@ -1768,29 +1770,29 @@ struct GAMEPLAYABILITIES_API FActiveGameplayEffectsContainer : public FFastArray
 	/** Return whether the container is using COND_Dynamic and setting the proper condition at runtime. */
 	bool IsUsingReplicationCondition() const { return bIsUsingReplicationCondition; }
 
-	bool NetDeltaSerialize(FNetDeltaSerializeInfo& DeltaParms);
+	UE_API bool NetDeltaSerialize(FNetDeltaSerializeInfo& DeltaParms);
 
-	void Uninitialize();
+	UE_API void Uninitialize();
 
-	bool CanApplyAttributeModifiers(const UGameplayEffect *GameplayEffect, float Level, const FGameplayEffectContextHandle& EffectContext);
+	UE_API bool CanApplyAttributeModifiers(const UGameplayEffect *GameplayEffect, float Level, const FGameplayEffectContextHandle& EffectContext);
 	
-	TArray<float> GetActiveEffectsTimeRemaining(const FGameplayEffectQuery& Query) const;
+	UE_API TArray<float> GetActiveEffectsTimeRemaining(const FGameplayEffectQuery& Query) const;
 
-	TArray<float> GetActiveEffectsDuration(const FGameplayEffectQuery& Query) const;
+	UE_API TArray<float> GetActiveEffectsDuration(const FGameplayEffectQuery& Query) const;
 
-	TArray<TPair<float,float>> GetActiveEffectsTimeRemainingAndDuration(const FGameplayEffectQuery& Query) const;
+	UE_API TArray<TPair<float,float>> GetActiveEffectsTimeRemainingAndDuration(const FGameplayEffectQuery& Query) const;
 
-	TArray<FActiveGameplayEffectHandle> GetActiveEffects(const FGameplayEffectQuery& Query) const;
+	UE_API TArray<FActiveGameplayEffectHandle> GetActiveEffects(const FGameplayEffectQuery& Query) const;
 
-	float GetActiveEffectsEndTime(const FGameplayEffectQuery& Query, TArray<AActor*>& Instigators) const;
-	bool GetActiveEffectsEndTimeAndDuration(const FGameplayEffectQuery& Query, float& EndTime, float& Duration, TArray<AActor*>& Instigators) const;
+	UE_API float GetActiveEffectsEndTime(const FGameplayEffectQuery& Query, TArray<AActor*>& Instigators) const;
+	UE_API bool GetActiveEffectsEndTimeAndDuration(const FGameplayEffectQuery& Query, float& EndTime, float& Duration, TArray<AActor*>& Instigators) const;
 
 	/** Returns an array of all of the active gameplay effect handles */
-	TArray<FActiveGameplayEffectHandle> GetAllActiveEffectHandles() const;
+	UE_API TArray<FActiveGameplayEffectHandle> GetAllActiveEffectHandles() const;
 
-	void ModifyActiveEffectStartTime(FActiveGameplayEffectHandle Handle, float StartTimeDiff);
+	UE_API void ModifyActiveEffectStartTime(FActiveGameplayEffectHandle Handle, float StartTimeDiff);
 
-	int32 RemoveActiveEffects(const FGameplayEffectQuery& Query, int32 StacksToRemove);
+	UE_API int32 RemoveActiveEffects(const FGameplayEffectQuery& Query, int32 StacksToRemove);
 
 	/** Method called during effect application to process if any active effects should be removed from this effects application */
 	UE_DEPRECATED(5.3, "AttemptRemoveActiveEffectsOnEffectApplication has been deprecated in favor of URemoveOtherGameplayEffectComponent")
@@ -1801,52 +1803,52 @@ struct GAMEPLAYABILITIES_API FActiveGameplayEffectsContainer : public FFastArray
 	 * 
 	 * @return Count of the effects matching the specified query
 	 */
-	int32 GetActiveEffectCount(const FGameplayEffectQuery& Query, bool bEnforceOnGoingCheck = true) const;
+	UE_API int32 GetActiveEffectCount(const FGameplayEffectQuery& Query, bool bEnforceOnGoingCheck = true) const;
 
-	bool IsServerWorldTimeAvailable() const;
+	UE_API bool IsServerWorldTimeAvailable() const;
 
-	float GetServerWorldTime() const;
+	UE_API float GetServerWorldTime() const;
 
-	float GetWorldTime() const;
+	UE_API float GetWorldTime() const;
 
-	bool HasReceivedEffectWithPredictedKey(FPredictionKey PredictionKey) const;
+	UE_API bool HasReceivedEffectWithPredictedKey(FPredictionKey PredictionKey) const;
 
-	bool HasPredictedEffectWithPredictedKey(FPredictionKey PredictionKey) const;
+	UE_API bool HasPredictedEffectWithPredictedKey(FPredictionKey PredictionKey) const;
 	
 	UE_DEPRECATED(5.5, "Replaced by private SetBaseAttributeValueFromReplication that uses FGameplayAttributeData")
-	void SetBaseAttributeValueFromReplication(const FGameplayAttribute& Attribute, float NewBaseValue, float OldBaseValue);
+	UE_API void SetBaseAttributeValueFromReplication(const FGameplayAttribute& Attribute, float NewBaseValue, float OldBaseValue);
 
-	void GetAllActiveGameplayEffectSpecs(TArray<FGameplayEffectSpec>& OutSpecCopies) const;
+	UE_API void GetAllActiveGameplayEffectSpecs(TArray<FGameplayEffectSpec>& OutSpecCopies) const;
 
-	void DebugCyclicAggregatorBroadcasts(struct FAggregator* Aggregator);
+	UE_API void DebugCyclicAggregatorBroadcasts(struct FAggregator* Aggregator);
 
 	/** Performs a deep copy on the source container, duplicating all gameplay effects and reconstructing the attribute aggregator map to match the passed in source. */
-	void CloneFrom(const FActiveGameplayEffectsContainer& Source);
+	UE_API void CloneFrom(const FActiveGameplayEffectsContainer& Source);
 
 	// -------------------------------------------------------------------------------------------
 
 	UE_DEPRECATED(4.17, "Use GetGameplayAttributeValueChangeDelegate (the delegate signature has changed)")
-	FOnGameplayAttributeChange& RegisterGameplayAttributeEvent(FGameplayAttribute Attribute);
+	UE_API FOnGameplayAttributeChange& RegisterGameplayAttributeEvent(FGameplayAttribute Attribute);
 
-	FOnGameplayAttributeValueChange& GetGameplayAttributeValueChangeDelegate(FGameplayAttribute Attribute);
+	UE_API FOnGameplayAttributeValueChange& GetGameplayAttributeValueChangeDelegate(FGameplayAttribute Attribute);
 
 	UE_DEPRECATED(5.3, "Use UImmunityGameplayEffectComponent. This function will now always return false.")
 	bool HasApplicationImmunityToSpec(const FGameplayEffectSpec& SpecToApply, const FActiveGameplayEffect*& OutGEThatProvidedImmunity) const { return false; }
 
-	void IncrementLock();
-	void DecrementLock();
+	UE_API void IncrementLock();
+	UE_API void DecrementLock();
 	
 	FORCEINLINE ConstIterator CreateConstIterator() const { return ConstIterator(*this);	}
 	FORCEINLINE Iterator CreateIterator() { return Iterator(*this);	}
 
 	/** Recomputes the start time for all active abilities */
-	void RecomputeStartWorldTimes(const float WorldTime, const float ServerWorldTime);
+	UE_API void RecomputeStartWorldTimes(const float WorldTime, const float ServerWorldTime);
 
 	/** Called every time data has been modified by the FastArraySerializer */
-	void PostReplicatedReceive(const FFastArraySerializer::FPostReplicatedReceiveParameters& Parameters);
+	UE_API void PostReplicatedReceive(const FFastArraySerializer::FPostReplicatedReceiveParameters& Parameters);
 
 private:
-	void SetBaseAttributeValueFromReplication(const FGameplayAttribute& Attribute, const FGameplayAttributeData& NewValue, const FGameplayAttributeData& OldValue);
+	UE_API void SetBaseAttributeValueFromReplication(const FGameplayAttribute& Attribute, const FGameplayAttributeData& NewValue, const FGameplayAttributeData& OldValue);
 
 	/**
 	 *	Accessors for internal functions to get GameplayEffects directly by index.
@@ -1884,7 +1886,7 @@ private:
 	UPROPERTY()
 	TArray<FActiveGameplayEffect>	GameplayEffects_Internal;
 
-	void InternalUpdateNumericalAttribute(FGameplayAttribute Attribute, float NewValue, const FGameplayEffectModCallbackData* ModData, bool bFromRecursiveCall=false);
+	UE_API void InternalUpdateNumericalAttribute(FGameplayAttribute Attribute, float NewValue, const FGameplayEffectModCallbackData* ModData, bool bFromRecursiveCall=false);
 
 	/** Cached pointer to current mod data needed for callbacks. We cache it in the AGE struct to avoid passing it through all the delegate/aggregator plumbing */
 	const struct FGameplayEffectModCallbackData* CurrentModcallbackData;
@@ -1897,38 +1899,38 @@ private:
 	 * 
 	 * @return True if the mod successfully executed, false if it did not
 	 */
-	bool InternalExecuteMod(FGameplayEffectSpec& Spec, FGameplayModifierEvaluatedData& ModEvalData);
+	UE_API bool InternalExecuteMod(FGameplayEffectSpec& Spec, FGameplayModifierEvaluatedData& ModEvalData);
 public:
 	bool IsNetAuthority() const
 	{
 		return OwnerIsNetAuthority;
 	}
 private:
-	void InternalExecutePeriodicGameplayEffect(FActiveGameplayEffect& ActiveEffect);
+	UE_API void InternalExecutePeriodicGameplayEffect(FActiveGameplayEffect& ActiveEffect);
 
 	/** Called internally to actually remove a GameplayEffect or to reduce its StackCount. Returns true if we resized our internal GameplayEffect array. */
-	bool InternalRemoveActiveGameplayEffect(int32 Idx, int32 StacksToRemove, bool bPrematureRemoval);
+	UE_API bool InternalRemoveActiveGameplayEffect(int32 Idx, int32 StacksToRemove, bool bPrematureRemoval);
 	
 	/** Called both in server side creation and replication creation/deletion */
-	void InternalOnActiveGameplayEffectAdded(FActiveGameplayEffect& Effect, bool bInvokeGameplayCueEvents);
-	void InternalOnActiveGameplayEffectRemoved(FActiveGameplayEffect& Effect, bool bInvokeGameplayCueEvents, const FGameplayEffectRemovalInfo& GameplayEffectRemovalInfo);
+	UE_API void InternalOnActiveGameplayEffectAdded(FActiveGameplayEffect& Effect, bool bInvokeGameplayCueEvents);
+	UE_API void InternalOnActiveGameplayEffectRemoved(FActiveGameplayEffect& Effect, bool bInvokeGameplayCueEvents, const FGameplayEffectRemovalInfo& GameplayEffectRemovalInfo);
 
-	void RemoveActiveGameplayEffectGrantedTagsAndModifiers(const FActiveGameplayEffect& Effect, bool bInvokeGameplayCueEvents);
-	void AddActiveGameplayEffectGrantedTagsAndModifiers(FActiveGameplayEffect& Effect, bool bInvokeGameplayCueEvents);
+	UE_API void RemoveActiveGameplayEffectGrantedTagsAndModifiers(const FActiveGameplayEffect& Effect, bool bInvokeGameplayCueEvents);
+	UE_API void AddActiveGameplayEffectGrantedTagsAndModifiers(FActiveGameplayEffect& Effect, bool bInvokeGameplayCueEvents);
 
 	/** Updates tag dependency map when a GameplayEffect is removed */
-	void RemoveActiveEffectTagDependency(const FGameplayTagContainer& Tags, FActiveGameplayEffectHandle Handle);
+	UE_API void RemoveActiveEffectTagDependency(const FGameplayTagContainer& Tags, FActiveGameplayEffectHandle Handle);
 
 	/** Internal helper function to bind the active effect to all of the custom modifier magnitude external dependency delegates it contains, if any */
-	void AddCustomMagnitudeExternalDependencies(FActiveGameplayEffect& Effect);
+	UE_API void AddCustomMagnitudeExternalDependencies(FActiveGameplayEffect& Effect);
 
 	/** Internal helper function to unbind the active effect from all of the custom modifier magnitude external dependency delegates it may have bound to, if any */
-	void RemoveCustomMagnitudeExternalDependencies(FActiveGameplayEffect& Effect);
+	UE_API void RemoveCustomMagnitudeExternalDependencies(FActiveGameplayEffect& Effect);
 
 	/** Internal callback fired as a result of a custom modifier magnitude external dependency delegate firing; Updates affected active gameplay effects as necessary */
-	void OnCustomMagnitudeExternalDependencyFired(TSubclassOf<UGameplayModMagnitudeCalculation> MagnitudeCalculationClass);
+	UE_API void OnCustomMagnitudeExternalDependencyFired(TSubclassOf<UGameplayModMagnitudeCalculation> MagnitudeCalculationClass);
 
-	void RestartActiveGameplayEffectDuration(FActiveGameplayEffect& ActiveGameplayEffect);
+	UE_API void RestartActiveGameplayEffectDuration(FActiveGameplayEffect& ActiveGameplayEffect);
 
 	// -------------------------------------------------------------------------------------------
 
@@ -1945,29 +1947,29 @@ private:
 	/** A map to manage stacking while we are the source */
 	TMap<TWeakObjectPtr<UGameplayEffect>, TArray<FActiveGameplayEffectHandle> >	SourceStackingMap;
 	
-	FAggregatorRef& FindOrCreateAttributeAggregator(const FGameplayAttribute& Attribute);
+	UE_API FAggregatorRef& FindOrCreateAttributeAggregator(const FGameplayAttribute& Attribute);
 
-	void CleanupAttributeAggregator(const FGameplayAttribute& Attribute);
+	UE_API void CleanupAttributeAggregator(const FGameplayAttribute& Attribute);
 
-	void OnAttributeAggregatorDirty(FAggregator* Aggregator, FGameplayAttribute Attribute, bool FromRecursiveCall=false);
+	UE_API void OnAttributeAggregatorDirty(FAggregator* Aggregator, FGameplayAttribute Attribute, bool FromRecursiveCall=false);
 
-	void OnMagnitudeDependencyChange(FActiveGameplayEffectHandle Handle, const FAggregator* ChangedAgg);
+	UE_API void OnMagnitudeDependencyChange(FActiveGameplayEffectHandle Handle, const FAggregator* ChangedAgg);
 
-	void OnStackCountChange(FActiveGameplayEffect& ActiveEffect, int32 OldStackCount, int32 NewStackCount);
+	UE_API void OnStackCountChange(FActiveGameplayEffect& ActiveEffect, int32 OldStackCount, int32 NewStackCount);
 
-	void OnDurationChange(FActiveGameplayEffect& ActiveEffect);
+	UE_API void OnDurationChange(FActiveGameplayEffect& ActiveEffect);
 
-	void UpdateAllAggregatorModMagnitudes(FActiveGameplayEffect& ActiveEffect);
+	UE_API void UpdateAllAggregatorModMagnitudes(FActiveGameplayEffect& ActiveEffect);
 
-	void UpdateAggregatorModMagnitudes(const TSet<FGameplayAttribute>& AttributesToUpdate, FActiveGameplayEffect& ActiveEffect);
+	UE_API void UpdateAggregatorModMagnitudes(const TSet<FGameplayAttribute>& AttributesToUpdate, FActiveGameplayEffect& ActiveEffect);
 
 	/** Helper function to find the active GE that the specified spec can stack with, if any */
-	FActiveGameplayEffect* FindStackableActiveGameplayEffect(const FGameplayEffectSpec& Spec);
+	UE_API FActiveGameplayEffect* FindStackableActiveGameplayEffect(const FGameplayEffectSpec& Spec);
 	
 	/** Helper function to handle the case of same-effect stacking overflow; Returns true if the overflow application should apply, false if it should not */
-	bool HandleActiveGameplayEffectStackOverflow(const FActiveGameplayEffect& ActiveStackableGE, const FGameplayEffectSpec& OldSpec, const FGameplayEffectSpec& OverflowingSpec);
+	UE_API bool HandleActiveGameplayEffectStackOverflow(const FActiveGameplayEffect& ActiveStackableGE, const FGameplayEffectSpec& OldSpec, const FGameplayEffectSpec& OverflowingSpec);
 
-	bool ShouldUseMinimalReplication();
+	UE_API bool ShouldUseMinimalReplication();
 
 	mutable int32 ScopedLockCount;
 	int32 PendingRemoves;
@@ -2009,10 +2011,10 @@ struct TStructOpsTypeTraits< FActiveGameplayEffectsContainer > : public TStructO
  *	This scope lock will queue deletions and additions until after the scope is over. The additions and deletions will actually 
  *	go through, but we will defer the memory operations to the active gameplay effect list.
  */
-struct GAMEPLAYABILITIES_API FScopedActiveGameplayEffectLock
+struct FScopedActiveGameplayEffectLock
 {
-	FScopedActiveGameplayEffectLock(FActiveGameplayEffectsContainer& InContainer);
-	~FScopedActiveGameplayEffectLock();
+	UE_API FScopedActiveGameplayEffectLock(FActiveGameplayEffectsContainer& InContainer);
+	UE_API ~FScopedActiveGameplayEffectLock();
 
 private:
 	FActiveGameplayEffectsContainer& Container;
@@ -2065,50 +2067,50 @@ struct TStructOpsTypeTraits<FGameplayEffectVersion> : public TStructOpsTypeTrait
  *	The GameplayEffect definition. This is the data asset defined in the editor that drives everything.
  *  This is only blueprintable to allow for templating gameplay effects. Gameplay effects should NOT contain blueprint graphs.
  */
-UCLASS(Blueprintable, PrioritizeCategories="Status Duration GameplayEffect GameplayCues Stacking", meta = (ShortTooltip = "A GameplayEffect modifies attributes and tags."))
-class GAMEPLAYABILITIES_API UGameplayEffect : public UObject, public IGameplayTagAssetInterface
+UCLASS(Blueprintable, PrioritizeCategories="Status Duration GameplayEffect GameplayCues Stacking", meta = (ShortTooltip = "A GameplayEffect modifies attributes and tags."), MinimalAPI)
+class UGameplayEffect : public UObject, public IGameplayTagAssetInterface
 {
 
 public:
 	GENERATED_UCLASS_BODY()
 
 	// These are deprecated but remain for backwards compat, please use FGameplayEffectConstants:: instead.
-	static const float INFINITE_DURATION;
-	static const float INSTANT_APPLICATION;
-	static const float NO_PERIOD;	
-	static const float INVALID_LEVEL;
+	static UE_API const float INFINITE_DURATION;
+	static UE_API const float INSTANT_APPLICATION;
+	static UE_API const float NO_PERIOD;	
+	static UE_API const float INVALID_LEVEL;
 
 	UE_DEPRECATED(5.3, "The implementation and method name did not match.  Use GetGrantedTags() to get the tags Granted to the Actor this GameplayEffect is applied to.")
-	virtual void GetOwnedGameplayTags(FGameplayTagContainer& TagContainer) const override;
+	UE_API virtual void GetOwnedGameplayTags(FGameplayTagContainer& TagContainer) const override;
 
 	UE_DEPRECATED(5.3, "The implementation and method name did not match.  Use GetGrantedTags().HasTag() to check against the tags this GameplayEffect will Grant to the Actor.")
-	virtual bool HasMatchingGameplayTag(FGameplayTag TagToCheck) const override;
+	UE_API virtual bool HasMatchingGameplayTag(FGameplayTag TagToCheck) const override;
 
 	UE_DEPRECATED(5.3, "The implementation and method name did not match.  Use GetGrantedTags().HasAll() to check against the tags this GameplayEffect will Grant to the Actor.")
-	virtual bool HasAllMatchingGameplayTags(const FGameplayTagContainer& TagContainer) const override;
+	UE_API virtual bool HasAllMatchingGameplayTags(const FGameplayTagContainer& TagContainer) const override;
 
 	UE_DEPRECATED(5.3, "The implementation and method name did not match.  Use GetGrantedTags().HasAny() to check against the tags this GameplayEffect will Grant to the Actor.")
-	virtual bool HasAnyMatchingGameplayTags(const FGameplayTagContainer& TagContainer) const override;
+	UE_API virtual bool HasAnyMatchingGameplayTags(const FGameplayTagContainer& TagContainer) const override;
 
-	virtual void GetBlockedAbilityTags(FGameplayTagContainer& OutTagContainer) const;
+	UE_API virtual void GetBlockedAbilityTags(FGameplayTagContainer& OutTagContainer) const;
 
 	/** Needed to properly disable inheriting the version value from its parent. */
-	virtual void PostInitProperties() override;
+	UE_API virtual void PostInitProperties() override;
 
 	/** PostLoad gets called once after the asset has been loaded.  It does not get called again on Blueprint recompile (@see PostCDOCompiled) */
-	virtual void PostLoad() override;
+	UE_API virtual void PostLoad() override;
 
 	/** Called when the Gameplay Effect has finished loading.  This is used to catch both PostLoad (the initial load) and PostCDOCompiled (any subsequent changes) */
-	virtual void OnGameplayEffectChanged();
+	UE_API virtual void OnGameplayEffectChanged();
 
 #if WITH_EDITOR
 	/** Do our upgrades in PostCDOCompiled */
-	virtual void PostCDOCompiled(const FPostCDOCompiledContext& Context) override;
+	UE_API virtual void PostCDOCompiled(const FPostCDOCompiledContext& Context) override;
 
 	/** We need to fix-up all of the SubObjects manually since the Engine doesn't fully support this */
-	void PostCDOCompiledFixupSubobjects();
+	UE_API void PostCDOCompiledFixupSubobjects();
 
-	virtual void PostEditChangeProperty(FPropertyChangedEvent& PropertyChangedEvent) override;
+	UE_API virtual void PostEditChangeProperty(FPropertyChangedEvent& PropertyChangedEvent) override;
 #endif
 
 	/** Check for any errors */
@@ -2118,25 +2120,25 @@ public:
 	/**
 	 * Can we Apply this Gameplay Effect?  Note: Apply is the generic term for adding to the active container or executing it.
 	 */
-	bool CanApply(const FActiveGameplayEffectsContainer& ActiveGEContainer, const FGameplayEffectSpec& GESpec) const;
+	UE_API bool CanApply(const FActiveGameplayEffectsContainer& ActiveGEContainer, const FGameplayEffectSpec& GESpec) const;
 
 	/**
 	 * Receive a notify that this GameplayEffect has been added to an Active Container.  ActiveGE will be the Active version of this GameplayEffect.
 	 * Returns true if this GameplayEffect should be active (or false to inhibit).
 	 */
-	bool OnAddedToActiveContainer(FActiveGameplayEffectsContainer& ActiveGEContainer, FActiveGameplayEffect& ActiveGE) const;
+	UE_API bool OnAddedToActiveContainer(FActiveGameplayEffectsContainer& ActiveGEContainer, FActiveGameplayEffect& ActiveGE) const;
 
 	/**
      * Receive a notify that this GameplayEffect has been executed (it must be instant, as it is not added to the Container).
      * Since it is not added to the container, it will not have an associated FActiveGameplayEffect, just the FGameplayEffectSpec.
      */
-	void OnExecuted(FActiveGameplayEffectsContainer& ActiveGEContainer, FGameplayEffectSpec& GESpec, FPredictionKey& PredictionKey) const;
+	UE_API void OnExecuted(FActiveGameplayEffectsContainer& ActiveGEContainer, FGameplayEffectSpec& GESpec, FPredictionKey& PredictionKey) const;
 
 	/**
 	 * Receive a notify that this GameplayEffect has been applied (this GE is either previously added to the container or executed in such cases).
 	 * However, this also encompasses cases where a GE was added to a container previously and then applied again to 'stack'.  It does not happen for periodic executions of a duration GE.
 	 */
-	void OnApplied(FActiveGameplayEffectsContainer& ActiveGEContainer, FGameplayEffectSpec& GESpec, FPredictionKey& PredictionKey) const;
+	UE_API void OnApplied(FActiveGameplayEffectsContainer& ActiveGEContainer, FGameplayEffectSpec& GESpec, FPredictionKey& PredictionKey) const;
 
 	/** Returns all tags that this GE *has* and *does not* grant to any Actor. */
 	const FGameplayTagContainer& GetAssetTags() const { return CachedAssetTags; }
@@ -2148,10 +2150,10 @@ public:
 	const FGameplayTagContainer& GetBlockedAbilityTags() const { return CachedBlockedAbilityTags; }
 
 	/** Returns the maximum stack size for this gameplay effect */
-	int32 GetStackLimitCount() const;
+	UE_API int32 GetStackLimitCount() const;
 
 	/** Returns the stack expiration policy for this gameplay effect */
-	EGameplayEffectStackingExpirationPolicy GetStackExpirationPolicy() const;
+	UE_API EGameplayEffectStackingExpirationPolicy GetStackExpirationPolicy() const;
 
 	// ---------------------------------------------------------------------------------------------------------------------------------
 
@@ -2165,7 +2167,7 @@ public:
 	const GEComponentClass* FindComponent() const;
 
 	/** @return the first component that derives from the passed-in ClassToFind, if one exists. */
-	const UGameplayEffectComponent* FindComponent(TSubclassOf<UGameplayEffectComponent> ClassToFind) const;
+	UE_API const UGameplayEffectComponent* FindComponent(TSubclassOf<UGameplayEffectComponent> ClassToFind) const;
 
 	/** 
 	 * Add a GameplayEffectComponent to the GameplayEffect.  This does not check for duplicates and is guaranteed to return a new instance.
@@ -2185,7 +2187,7 @@ public:
 
 #if WITH_EDITOR
 	/** Allow each Gameplay Effect Component to validate its own data.  Call this version (Super::IsDataValid) _after_ your overridden version to update EditorStatusText. */
-	virtual EDataValidationResult IsDataValid(class FDataValidationContext& Context) const override;
+	UE_API virtual EDataValidationResult IsDataValid(class FDataValidationContext& Context) const override;
 
 protected:
 	// ----------------------------------------------------------------------
@@ -2193,27 +2195,27 @@ protected:
 	// ----------------------------------------------------------------------
 	
 	/** Let's keep track of the version so that we can upgrade the Components properly.  This will be set properly in PostLoad after upgrades. */
-	EGameplayEffectVersion GetVersion() const;
+	UE_API EGameplayEffectVersion GetVersion() const;
 
 	/** Sets the Version of the class to denote it's been upgraded */
-	void SetVersion(EGameplayEffectVersion Version);
+	UE_API void SetVersion(EGameplayEffectVersion Version);
 
 	/** We should intercept the Save call and revalidate all of our deprecated values to avoid hanging onto stale data */
-	virtual void PreSave(FObjectPreSaveContext SaveContext) override;
+	UE_API virtual void PreSave(FObjectPreSaveContext SaveContext) override;
 
 private:
 	// Helper functions for converting data to use components
-	void ConvertAbilitiesComponent();
-	void ConvertAdditionalEffectsComponent();
-	void ConvertAssetTagsComponent();
-	void ConvertBlockByTagsComponent();
-	void ConvertChanceToApplyComponent();
-	void ConvertCustomCanApplyComponent();
-	void ConvertImmunityComponent();
-	void ConvertRemoveOtherComponent();
-	void ConvertTagRequirementsComponent();
-	void ConvertTargetTagsComponent();
-	void ConvertUIComponent();
+	UE_API void ConvertAbilitiesComponent();
+	UE_API void ConvertAdditionalEffectsComponent();
+	UE_API void ConvertAssetTagsComponent();
+	UE_API void ConvertBlockByTagsComponent();
+	UE_API void ConvertChanceToApplyComponent();
+	UE_API void ConvertCustomCanApplyComponent();
+	UE_API void ConvertImmunityComponent();
+	UE_API void ConvertRemoveOtherComponent();
+	UE_API void ConvertTagRequirementsComponent();
+	UE_API void ConvertTargetTagsComponent();
+	UE_API void ConvertUIComponent();
 #endif
 
 	// ----------------------------------------------------------------------
@@ -2475,3 +2477,5 @@ GEComponentClass& UGameplayEffect::FindOrAddComponent()
 
 	return AddComponent<GEComponentClass>();
 }
+
+#undef UE_API
