[GameplayAbilities]

diff --git a/Engine/Plugins/Runtime/GameplayAbilities/Source/GameplayAbilities/Public/GameplayEffectAggregator.h b/Engine/Plugins/Runtime/GameplayAbilities/Source/GameplayAbilities/Public/GameplayEffectAggregator.h
--- a/Engine/Plugins/Runtime/GameplayAbilities/Source/GameplayAbilities/Public/GameplayEffectAggregator.h
+++ b/Engine/Plugins/Runtime/GameplayAbilities/Source/GameplayAbilities/Public/GameplayEffectAggregator.h
@@ -6,11 +6,13 @@
 #include "GameplayTagContainer.h"
 #include "GameplayEffectTypes.h"
 
+#define UE_API GAMEPLAYABILITIES_API
+
 struct FGameplayEffectSpec;
 struct FAggregator;
 
 /** Data that is used in aggregator evaluation that is passed from the caller/game code */
-struct GAMEPLAYABILITIES_API FAggregatorEvaluateParameters
+struct FAggregatorEvaluateParameters
 {
 	FAggregatorEvaluateParameters()
 		: SourceTags(nullptr)
@@ -36,7 +38,7 @@ struct GAMEPLAYABILITIES_API FAggregatorEvaluateParameters
 };
 
 /** Data that is used in aggregator evaluation that is intrinsic to the aggregator itself. Not passed in from the game code (though maybe originally setup once by the game code)   */
-struct GAMEPLAYABILITIES_API FAggregatorEvaluateMetaData
+struct FAggregatorEvaluateMetaData
 {
 	typedef TFunction< void(const FAggregatorEvaluateParameters&, const FAggregator*) > FCustomQualifiesFunc;
 
@@ -50,7 +52,7 @@ struct GAMEPLAYABILITIES_API FAggregatorEvaluateMetaData
 	FCustomQualifiesFunc CustomQualifiesFunc;
 };
 
-struct GAMEPLAYABILITIES_API FAggregatorMod
+struct FAggregatorMod
 {
 	const FGameplayTagRequirements*	SourceTagReqs;
 	const FGameplayTagRequirements*	TargetTagReqs;
@@ -64,7 +66,7 @@ struct GAMEPLAYABILITIES_API FAggregatorMod
 	bool Qualifies() const { return IsQualified; }
 
 	/** Called to update the Qualifies bool */
-	void UpdateQualifies(const FAggregatorEvaluateParameters& Parameters) const;
+	UE_API void UpdateQualifies(const FAggregatorEvaluateParameters& Parameters) const;
 
 	/** Intended to be used by FAggregatorEvaluateMetaData::CustomQualifiesFunc to toggle qualifications of mods */
 	void SetExplicitQualifies(bool NewQualifies) const { IsQualified = NewQualifies; }
@@ -75,7 +77,7 @@ private:
 	mutable bool IsQualified;
 };
 
-struct GAMEPLAYABILITIES_API FAggregatorModInfo
+struct FAggregatorModInfo
 {
 	EGameplayModEvaluationChannel Channel;
 	EGameplayModOp::Type Op;
@@ -83,7 +85,7 @@ struct GAMEPLAYABILITIES_API FAggregatorModInfo
 };
 
 /** Struct representing an individual aggregation channel/depth. Contains mods of all mod op types. */
-struct GAMEPLAYABILITIES_API FAggregatorModChannel
+struct FAggregatorModChannel
 {
 	/**
 	 * Evaluates the channel's mods with the specified base value and evaluation parameters
@@ -93,7 +95,7 @@ struct GAMEPLAYABILITIES_API FAggregatorModChannel
 	 * 
 	 * @return Evaluated value based upon the channel's mods
 	 */
-	float EvaluateWithBase(float InlineBaseValue, const FAggregatorEvaluateParameters& Parameters) const;
+	UE_API float EvaluateWithBase(float InlineBaseValue, const FAggregatorEvaluateParameters& Parameters) const;
 
 	/**
 	 * Evaluates a final value in reverse, attempting to determine a base value from the modifiers within the channel.
@@ -108,7 +110,7 @@ struct GAMEPLAYABILITIES_API FAggregatorModChannel
 	 * 
 	 * @return True if the reverse evaluation was successful, false if it was not
 	 */
-	bool ReverseEvaluate(float FinalValue, const FAggregatorEvaluateParameters& Parameters, OUT float& ComputedValue) const;
+	UE_API bool ReverseEvaluate(float FinalValue, const FAggregatorEvaluateParameters& Parameters, OUT float& ComputedValue) const;
 	
 	/**
 	 * Add a modifier to the channel
@@ -120,27 +122,27 @@ struct GAMEPLAYABILITIES_API FAggregatorModChannel
 	 * @param bIsPredicted			Whether the mod is predicted or not
 	 * @param ActiveHandle			Handle of the active gameplay effect that's applying the mod
 	 */
-	void AddMod(float EvaluatedMagnitude, TEnumAsByte<EGameplayModOp::Type> ModOp, const FGameplayTagRequirements* SourceTagReqs, const FGameplayTagRequirements* TargetTagReqs, bool bIsPredicted, const FActiveGameplayEffectHandle& ActiveHandle);
+	UE_API void AddMod(float EvaluatedMagnitude, TEnumAsByte<EGameplayModOp::Type> ModOp, const FGameplayTagRequirements* SourceTagReqs, const FGameplayTagRequirements* TargetTagReqs, bool bIsPredicted, const FActiveGameplayEffectHandle& ActiveHandle);
 	
 	/**
 	 * Remove all mods from the channel that match the specified gameplay effect handle
 	 * 
 	 * @param Handle	Handle to use for removal
 	 */
-	void RemoveModsWithActiveHandle(const FActiveGameplayEffectHandle& Handle);
+	UE_API void RemoveModsWithActiveHandle(const FActiveGameplayEffectHandle& Handle);
 
 	/**
 	 * Add the specified channel's mods into this channel
 	 * 
 	 * @param Other	Other channel to add mods from
 	 */
-	void AddModsFrom(const FAggregatorModChannel& Other);
+	UE_API void AddModsFrom(const FAggregatorModChannel& Other);
 
 	/** runs UpdateQualifies on all mods */
-	void UpdateQualifiesOnAllMods(const FAggregatorEvaluateParameters& Parameters) const;
+	UE_API void UpdateQualifiesOnAllMods(const FAggregatorEvaluateParameters& Parameters) const;
 
 	/** Helper function for iterating through all mods within a channel */
-	void ForEachMod(FAggregatorModInfo& Info, TFunction<void (const FAggregatorModInfo&) > Func) const;
+	UE_API void ForEachMod(FAggregatorModInfo& Info, TFunction<void (const FAggregatorModInfo&) > Func) const;
 
 	/**
 	 * Populate a mapping of channel to corresponding mods
@@ -148,7 +150,7 @@ struct GAMEPLAYABILITIES_API FAggregatorModChannel
 	 * @param Channel	Enum channel associated with this channel
 	 * @param OutMods	Mapping of channel enum to mods
 	 */
-	void GetAllAggregatorMods(EGameplayModEvaluationChannel Channel, OUT TMap<EGameplayModEvaluationChannel, const TArray<FAggregatorMod>*>& OutMods) const;
+	UE_API void GetAllAggregatorMods(EGameplayModEvaluationChannel Channel, OUT TMap<EGameplayModEvaluationChannel, const TArray<FAggregatorMod>*>& OutMods) const;
 
 	/**
 	 * Called when the mod channel's gameplay effect dependencies have potentially been swapped out for new ones, like when GE arrays are cloned.
@@ -156,7 +158,7 @@ struct GAMEPLAYABILITIES_API FAggregatorModChannel
 	 * 
 	 * @param SwappedDependencies	Mapping of old gameplay effect handles to new replacements
 	 */
-	void OnActiveEffectDependenciesSwapped(const TMap<FActiveGameplayEffectHandle, FActiveGameplayEffectHandle>& SwappedDependencies);
+	UE_API void OnActiveEffectDependenciesSwapped(const TMap<FActiveGameplayEffectHandle, FActiveGameplayEffectHandle>& SwappedDependencies);
 
 	/**
 	 * Helper function to sum all of the mods in the specified array, using the specified modifier bias and evaluation parameters
@@ -167,7 +169,7 @@ struct GAMEPLAYABILITIES_API FAggregatorModChannel
 	 * 
 	 * @return Summed value of mods
 	 */
-	static float SumMods(const TArray<FAggregatorMod>& InMods, float Bias, const FAggregatorEvaluateParameters& Parameters);
+	static UE_API float SumMods(const TArray<FAggregatorMod>& InMods, float Bias, const FAggregatorEvaluateParameters& Parameters);
 
 private:
 
@@ -176,7 +178,7 @@ private:
 };
 
 /** Struct representing a container of modifier channels */
-struct GAMEPLAYABILITIES_API FAggregatorModChannelContainer
+struct FAggregatorModChannelContainer
 {
 	/**
 	 * Find or add a modifier channel for the specified enum value
@@ -185,10 +187,10 @@ struct GAMEPLAYABILITIES_API FAggregatorModChannelContainer
 	 * 
 	 * @return Modifier channel for the specified enum value
 	 */
-	FAggregatorModChannel& FindOrAddModChannel(EGameplayModEvaluationChannel Channel);
+	UE_API FAggregatorModChannel& FindOrAddModChannel(EGameplayModEvaluationChannel Channel);
 
 	/** Simple accessor to the current number of modifier channels active */
-	int32 GetNumChannels() const;
+	UE_API int32 GetNumChannels() const;
 
 	/**
 	 * Evaluates the result of the specified base value run through each existing evaluation channel's modifiers in numeric order
@@ -202,7 +204,7 @@ struct GAMEPLAYABILITIES_API FAggregatorModChannelContainer
 	 * 
 	 * @return Result of the specified base value run through each modifier in each evaluation channel in numeric order
 	 */
-	float EvaluateWithBase(float InlineBaseValue, const FAggregatorEvaluateParameters& Parameters) const;
+	UE_API float EvaluateWithBase(float InlineBaseValue, const FAggregatorEvaluateParameters& Parameters) const;
 
 	/**
 	 * Similar to EvaluateWithBase (see comment there for details), but terminates early after evaluating the specified final channel instead of
@@ -214,7 +216,7 @@ struct GAMEPLAYABILITIES_API FAggregatorModChannelContainer
 	 * 
 	 * @return Result of the specified base value run through each modifier in each evaluation channel in numeric order
 	 */
-	float EvaluateWithBaseToChannel(float InlineBaseValue, const FAggregatorEvaluateParameters& Parameters, EGameplayModEvaluationChannel FinalChannel) const;
+	UE_API float EvaluateWithBaseToChannel(float InlineBaseValue, const FAggregatorEvaluateParameters& Parameters, EGameplayModEvaluationChannel FinalChannel) const;
 
 	/**
 	 * Evaluates a final value in reverse, attempting to determine a base value from the modifiers within all of the channels.
@@ -230,34 +232,34 @@ struct GAMEPLAYABILITIES_API FAggregatorModChannelContainer
 	 * 
 	 * @return If possible, the base value from the reverse evaluation. If not possible, the original final value is returned.
 	 */
-	float ReverseEvaluate(float FinalValue, const FAggregatorEvaluateParameters& Parameters) const;
+	UE_API float ReverseEvaluate(float FinalValue, const FAggregatorEvaluateParameters& Parameters) const;
 
 	/** Calls ::UpdateQualifies on each mod */
-	void EvaluateQualificationForAllMods(const FAggregatorEvaluateParameters& Parameters) const;
+	UE_API void EvaluateQualificationForAllMods(const FAggregatorEvaluateParameters& Parameters) const;
 
 	/**
 	 * Removes any mods from every channel matching the specified handle
 	 * 
 	 * @param ActiveHandle	Handle to use for removal
 	 */
-	void RemoveAggregatorMod(const FActiveGameplayEffectHandle& ActiveHandle);
+	UE_API void RemoveAggregatorMod(const FActiveGameplayEffectHandle& ActiveHandle);
 
 	/**
 	 * Adds the mods from specified container to this one
 	 * 
 	 * @param Other	Container to add mods from
 	 */
-	void AddModsFrom(const FAggregatorModChannelContainer& Other);
+	UE_API void AddModsFrom(const FAggregatorModChannelContainer& Other);
 
 	/** Helper function for iterating through all mods within the channel container */
-	void ForEachMod( TFunction<void (const FAggregatorModInfo&) > ) const;
+	UE_API void ForEachMod( TFunction<void (const FAggregatorModInfo&) > ) const;
 
 	/**
 	 * Populate a mapping of channel to corresponding mods for debugging purposes
 	 * 
 	 * @param OutMods	Mapping of channel enum to mods
 	 */
-	void GetAllAggregatorMods(OUT TMap<EGameplayModEvaluationChannel, const TArray<FAggregatorMod>*>& OutMods) const;
+	UE_API void GetAllAggregatorMods(OUT TMap<EGameplayModEvaluationChannel, const TArray<FAggregatorMod>*>& OutMods) const;
 
 	/**
 	 * Called when the container's gameplay effect dependencies have potentially been swapped out for new ones, like when GE arrays are cloned.
@@ -265,7 +267,7 @@ struct GAMEPLAYABILITIES_API FAggregatorModChannelContainer
 	 * 
 	 * @param SwappedDependencies	Mapping of old gameplay effect handles to new replacements
 	 */
-	void OnActiveEffectDependenciesSwapped(const TMap<FActiveGameplayEffectHandle, FActiveGameplayEffectHandle>& SwappedDependencies);
+	UE_API void OnActiveEffectDependenciesSwapped(const TMap<FActiveGameplayEffectHandle, FActiveGameplayEffectHandle>& SwappedDependencies);
 
 private:
 
@@ -273,7 +275,7 @@ private:
 	TMap<EGameplayModEvaluationChannel, FAggregatorModChannel> ModChannelsMap;
 };
 
-struct GAMEPLAYABILITIES_API FAggregator : public TSharedFromThis<FAggregator>
+struct FAggregator : public TSharedFromThis<FAggregator>
 {
 	DECLARE_MULTICAST_DELEGATE_OneParam(FOnAggregatorDirty, FAggregator*);
 
@@ -284,60 +286,60 @@ struct GAMEPLAYABILITIES_API FAggregator : public TSharedFromThis<FAggregator>
 		, BroadcastingDirtyCount(0)
 	{}
 	
-	~FAggregator();
+	UE_API ~FAggregator();
 
 	/** Simple accessor to base value */
-	float GetBaseValue() const;
-	void SetBaseValue(float NewBaseValue, bool BroadcastDirtyEvent = true);
+	UE_API float GetBaseValue() const;
+	UE_API void SetBaseValue(float NewBaseValue, bool BroadcastDirtyEvent = true);
 	
-	void ExecModOnBaseValue(TEnumAsByte<EGameplayModOp::Type> ModifierOp, float EvaluatedMagnitude);
-	static float StaticExecModOnBaseValue(float BaseValue, TEnumAsByte<EGameplayModOp::Type> ModifierOp, float EvaluatedMagnitude);
+	UE_API void ExecModOnBaseValue(TEnumAsByte<EGameplayModOp::Type> ModifierOp, float EvaluatedMagnitude);
+	static UE_API float StaticExecModOnBaseValue(float BaseValue, TEnumAsByte<EGameplayModOp::Type> ModifierOp, float EvaluatedMagnitude);
 
-	void AddAggregatorMod(float EvaluatedData, TEnumAsByte<EGameplayModOp::Type> ModifierOp, EGameplayModEvaluationChannel ModifierChannel, const FGameplayTagRequirements*	SourceTagReqs, const FGameplayTagRequirements* TargetTagReqs, bool IsPredicted, FActiveGameplayEffectHandle ActiveHandle = FActiveGameplayEffectHandle() );
+	UE_API void AddAggregatorMod(float EvaluatedData, TEnumAsByte<EGameplayModOp::Type> ModifierOp, EGameplayModEvaluationChannel ModifierChannel, const FGameplayTagRequirements*	SourceTagReqs, const FGameplayTagRequirements* TargetTagReqs, bool IsPredicted, FActiveGameplayEffectHandle ActiveHandle = FActiveGameplayEffectHandle() );
 
 	/** Removes all mods for the passed in handle and marks this as dirty to recalculate the aggregator */
-	void RemoveAggregatorMod(FActiveGameplayEffectHandle ActiveHandle);
+	UE_API void RemoveAggregatorMod(FActiveGameplayEffectHandle ActiveHandle);
 
 	/** Updates the aggregators for the past in handle, this will handle it so the UAttributeSets stats only get one update for the delta change */
-	void UpdateAggregatorMod(FActiveGameplayEffectHandle ActiveHandle, const FGameplayAttribute& Attribute, const FGameplayEffectSpec& Spec, bool bWasLocallyGenerated, FActiveGameplayEffectHandle InHandle);
+	UE_API void UpdateAggregatorMod(FActiveGameplayEffectHandle ActiveHandle, const FGameplayAttribute& Attribute, const FGameplayEffectSpec& Spec, bool bWasLocallyGenerated, FActiveGameplayEffectHandle InHandle);
 
 	/** Evaluates the Aggregator with the internal base value and given parameters */
-	float Evaluate(const FAggregatorEvaluateParameters& Parameters) const;
+	UE_API float Evaluate(const FAggregatorEvaluateParameters& Parameters) const;
 
 	/** Evaluates the aggregator with the internal base value and given parameters, up to the specified evaluation channel (inclusive) */
-	float EvaluateToChannel(const FAggregatorEvaluateParameters& Parameters, EGameplayModEvaluationChannel FinalChannel) const;
+	UE_API float EvaluateToChannel(const FAggregatorEvaluateParameters& Parameters, EGameplayModEvaluationChannel FinalChannel) const;
 
 	/** Works backwards to calculate the base value. Used on clients for doing predictive modifiers */
-	float ReverseEvaluate(float FinalValue, const FAggregatorEvaluateParameters& Parameters) const;
+	UE_API float ReverseEvaluate(float FinalValue, const FAggregatorEvaluateParameters& Parameters) const;
 
 	/** Evaluates the Aggregator with an arbitrary base value */
-	float EvaluateWithBase(float InlineBaseValue, const FAggregatorEvaluateParameters& Parameters) const;
+	UE_API float EvaluateWithBase(float InlineBaseValue, const FAggregatorEvaluateParameters& Parameters) const;
 
 	/** Evaluates the Aggregator to compute its "bonus" (final - base) value */
-	float EvaluateBonus(const FAggregatorEvaluateParameters& Parameters) const;
+	UE_API float EvaluateBonus(const FAggregatorEvaluateParameters& Parameters) const;
 
 	/** Evaluates the contribution from the GE associated with ActiveHandle */
-	float EvaluateContribution(const FAggregatorEvaluateParameters& Parameters, FActiveGameplayEffectHandle ActiveHandle) const;
+	UE_API float EvaluateContribution(const FAggregatorEvaluateParameters& Parameters, FActiveGameplayEffectHandle ActiveHandle) const;
 
 	/** Calls ::UpdateQualifies on each mod. Useful for when you need to manually inspect aggregators */
-	void EvaluateQualificationForAllMods(const FAggregatorEvaluateParameters& Parameters) const;
+	UE_API void EvaluateQualificationForAllMods(const FAggregatorEvaluateParameters& Parameters) const;
 
-	void TakeSnapshotOf(const FAggregator& AggToSnapshot);
+	UE_API void TakeSnapshotOf(const FAggregator& AggToSnapshot);
 
 	FOnAggregatorDirty OnDirty;
 	FOnAggregatorDirty OnDirtyRecursive;	// Called in case where we are in a recursive dirtying chain. This will force the backing uproperty to update but not call the game code delegates
 
-	void AddModsFrom(const FAggregator& SourceAggregator);
+	UE_API void AddModsFrom(const FAggregator& SourceAggregator);
 
-	void AddDependent(FActiveGameplayEffectHandle Handle);
-	void RemoveDependent(FActiveGameplayEffectHandle Handle);
+	UE_API void AddDependent(FActiveGameplayEffectHandle Handle);
+	UE_API void RemoveDependent(FActiveGameplayEffectHandle Handle);
 
 	/**
 	 * Populate a mapping of channel to corresponding mods
 	 * 
 	 * @param OutMods	Mapping of channel enum to mods
 	 */
-	void GetAllAggregatorMods(OUT TMap<EGameplayModEvaluationChannel, const TArray<FAggregatorMod>*>& OutMods) const;
+	UE_API void GetAllAggregatorMods(OUT TMap<EGameplayModEvaluationChannel, const TArray<FAggregatorMod>*>& OutMods) const;
 
 	/**
 	 * Called when the aggregator's gameplay effect dependencies have potentially been swapped out for new ones, like when GE arrays are cloned.
@@ -345,10 +347,10 @@ struct GAMEPLAYABILITIES_API FAggregator : public TSharedFromThis<FAggregator>
 	 * 
 	 * @param SwappedDependencies	Mapping of old gameplay effect handles to new replacements
 	 */
-	void OnActiveEffectDependenciesSwapped(const TMap<FActiveGameplayEffectHandle, FActiveGameplayEffectHandle>& SwappedDependencies);
+	UE_API void OnActiveEffectDependenciesSwapped(const TMap<FActiveGameplayEffectHandle, FActiveGameplayEffectHandle>& SwappedDependencies);
 
 	/** Helper function for iterating through all mods within the aggregator */
-	void ForEachMod( TFunction<void (const FAggregatorModInfo&) > ) const;
+	UE_API void ForEachMod( TFunction<void (const FAggregatorModInfo&) > ) const;
 
 	/** Custom meta data for the aggregator */
 	FAggregatorEvaluateMetaData* EvaluationMetaData;
@@ -358,7 +360,7 @@ struct GAMEPLAYABILITIES_API FAggregator : public TSharedFromThis<FAggregator>
 
 private:
 
-	void BroadcastOnDirty();
+	UE_API void BroadcastOnDirty();
 
 	float	BaseValue;
 	FAggregatorModChannelContainer ModChannels;
@@ -373,7 +375,7 @@ private:
 	friend class UAbilitySystemComponent;	// Only needed for DisplayDebug()
 };
 
-struct GAMEPLAYABILITIES_API FAggregatorRef
+struct FAggregatorRef
 {
 	FAggregatorRef() { }
 	FAggregatorRef(FAggregator* InData) : Data(InData) { }
@@ -382,7 +384,7 @@ struct GAMEPLAYABILITIES_API FAggregatorRef
 
 	TSharedPtr<FAggregator>	Data;
 
-	void TakeSnapshotOf(const FAggregatorRef& RefToSnapshot);
+	UE_API void TakeSnapshotOf(const FAggregatorRef& RefToSnapshot);
 };
 
 /**
@@ -392,22 +394,24 @@ struct GAMEPLAYABILITIES_API FAggregatorRef
  *	The only catch is that we store raw FAggregator*. This should only be used in scopes where aggreagtors
  *	are not deleted. There is currently no place that does. If we find to, we could add additional safety checks.
  */
-struct GAMEPLAYABILITIES_API FScopedAggregatorOnDirtyBatch
+struct FScopedAggregatorOnDirtyBatch
 {
-	FScopedAggregatorOnDirtyBatch();
-	~FScopedAggregatorOnDirtyBatch();
+	UE_API FScopedAggregatorOnDirtyBatch();
+	UE_API ~FScopedAggregatorOnDirtyBatch();
 
-	static void BeginLock();
-	static void EndLock();
+	static UE_API void BeginLock();
+	static UE_API void EndLock();
 
-	static void BeginNetReceiveLock();
-	static void EndNetReceiveLock();
+	static UE_API void BeginNetReceiveLock();
+	static UE_API void EndNetReceiveLock();
 
-	static int32	GlobalBatchCount;
-	static TSet<FAggregator*>	DirtyAggregators;
+	static UE_API int32	GlobalBatchCount;
+	static UE_API TSet<FAggregator*>	DirtyAggregators;
 
-	static bool		GlobalFromNetworkUpdate;
-	static int32	NetUpdateID;
+	static UE_API bool		GlobalFromNetworkUpdate;
+	static UE_API int32	NetUpdateID;
 };
 
 #define AGGREGATOR_BATCH_SCOPE() FScopedAggregatorOnDirtyBatch AggregatorOnDirtyBatcher;
+
+#undef UE_API
diff --git a/Engine/Plugins/Runtime/GameplayAbilities/Source/GameplayAbilities/Public/GameplayEffectAggregatorLibrary.h b/Engine/Plugins/Runtime/GameplayAbilities/Source/GameplayAbilities/Public/GameplayEffectAggregatorLibrary.h
--- a/Engine/Plugins/Runtime/GameplayAbilities/Source/GameplayAbilities/Public/GameplayEffectAggregatorLibrary.h
+++ b/Engine/Plugins/Runtime/GameplayAbilities/Source/GameplayAbilities/Public/GameplayEffectAggregatorLibrary.h
@@ -4,7 +4,11 @@
 
 #include "GameplayEffectAggregator.h"
 
-struct GAMEPLAYABILITIES_API FAggregatorEvaluateMetaDataLibrary
+#define UE_API GAMEPLAYABILITIES_API
+
+struct FAggregatorEvaluateMetaDataLibrary
 {
-	static FAggregatorEvaluateMetaData MostNegativeMod_AllPositiveMods;
-};
\ No newline at end of file
+	static UE_API FAggregatorEvaluateMetaData MostNegativeMod_AllPositiveMods;
+};
+
+#undef UE_API
diff --git a/Engine/Plugins/Runtime/GameplayAbilities/Source/GameplayAbilities/Public/GameplayEffectAttributeCaptureDefinition.h b/Engine/Plugins/Runtime/GameplayAbilities/Source/GameplayAbilities/Public/GameplayEffectAttributeCaptureDefinition.h
--- a/Engine/Plugins/Runtime/GameplayAbilities/Source/GameplayAbilities/Public/GameplayEffectAttributeCaptureDefinition.h
+++ b/Engine/Plugins/Runtime/GameplayAbilities/Source/GameplayAbilities/Public/GameplayEffectAttributeCaptureDefinition.h
@@ -5,6 +5,8 @@
 #include "AttributeSet.h"
 #include "GameplayEffectAttributeCaptureDefinition.generated.h"
 
+#define UE_API GAMEPLAYABILITIES_API
+
 /** Enumeration for options of where to capture gameplay attributes from for gameplay effects. */
 UENUM()
 enum class EGameplayEffectAttributeCaptureSource : uint8
@@ -18,7 +20,7 @@ enum class EGameplayEffectAttributeCaptureSource : uint8
 
 /** Struct defining gameplay attribute capture options for gameplay effects */
 USTRUCT(BlueprintType)
-struct GAMEPLAYABILITIES_API FGameplayEffectAttributeCaptureDefinition
+struct FGameplayEffectAttributeCaptureDefinition
 {
 	GENERATED_USTRUCT_BODY()
 
@@ -47,8 +49,8 @@ struct GAMEPLAYABILITIES_API FGameplayEffectAttributeCaptureDefinition
 	bool bSnapshot;
 
 	/** Equality/Inequality operators */
-	bool operator==(const FGameplayEffectAttributeCaptureDefinition& Other) const;
-	bool operator!=(const FGameplayEffectAttributeCaptureDefinition& Other) const;
+	UE_API bool operator==(const FGameplayEffectAttributeCaptureDefinition& Other) const;
+	UE_API bool operator!=(const FGameplayEffectAttributeCaptureDefinition& Other) const;
 
 	/**
 	 * Get type hash for the capture definition; Implemented to allow usage in TMap
@@ -64,6 +66,7 @@ struct GAMEPLAYABILITIES_API FGameplayEffectAttributeCaptureDefinition
 		return Hash;
 	}
 
-	FString ToSimpleString() const;
+	UE_API FString ToSimpleString() const;
 };
 
+#undef UE_API
diff --git a/Engine/Plugins/Runtime/GameplayAbilities/Source/GameplayAbilities/Public/GameplayEffectCalculation.h b/Engine/Plugins/Runtime/GameplayAbilities/Source/GameplayAbilities/Public/GameplayEffectCalculation.h
--- a/Engine/Plugins/Runtime/GameplayAbilities/Source/GameplayAbilities/Public/GameplayEffectCalculation.h
+++ b/Engine/Plugins/Runtime/GameplayAbilities/Source/GameplayAbilities/Public/GameplayEffectCalculation.h
@@ -8,16 +8,18 @@
 #include "GameplayEffectTypes.h"
 #include "GameplayEffectCalculation.generated.h"
 
+#define UE_API GAMEPLAYABILITIES_API
+
 /** Abstract base for specialized gameplay effect calculations; Capable of specifing attribute captures */
-UCLASS(BlueprintType, Blueprintable, Abstract)
-class GAMEPLAYABILITIES_API UGameplayEffectCalculation : public UObject
+UCLASS(BlueprintType, Blueprintable, Abstract, MinimalAPI)
+class UGameplayEffectCalculation : public UObject
 {
 	GENERATED_UCLASS_BODY()
 
 public:
 
 	/** Simple accessor to capture definitions for attributes */
-	virtual const TArray<FGameplayEffectAttributeCaptureDefinition>& GetAttributeCaptureDefinitions() const;
+	UE_API virtual const TArray<FGameplayEffectAttributeCaptureDefinition>& GetAttributeCaptureDefinitions() const;
 
 protected:
 
@@ -25,3 +27,5 @@ protected:
 	UPROPERTY(EditDefaultsOnly, BlueprintReadOnly, Category=Attributes)
 	TArray<FGameplayEffectAttributeCaptureDefinition> RelevantAttributesToCapture;
 };
+
+#undef UE_API
diff --git a/Engine/Plugins/Runtime/GameplayAbilities/Source/GameplayAbilities/Public/GameplayEffectComponent.h b/Engine/Plugins/Runtime/GameplayAbilities/Source/GameplayAbilities/Public/GameplayEffectComponent.h
--- a/Engine/Plugins/Runtime/GameplayAbilities/Source/GameplayAbilities/Public/GameplayEffectComponent.h
+++ b/Engine/Plugins/Runtime/GameplayAbilities/Source/GameplayAbilities/Public/GameplayEffectComponent.h
@@ -5,6 +5,8 @@
 #include "UObject/Object.h"
 #include "GameplayEffectComponent.generated.h"
 
+#define UE_API GAMEPLAYABILITIES_API
+
 struct FActiveGameplayEffect;
 struct FActiveGameplayEffectsContainer;
 struct FGameplayEffectSpec;
@@ -26,17 +28,17 @@ class UGameplayEffect;
  * 
  * @see GameplayEffect.h for further notes, especially on the terminology used (Added vs. Executed vs. Apply).
  */
-UCLASS(Abstract, Const, DefaultToInstanced, EditInlineNew, CollapseCategories, Within=GameplayEffect)
-class GAMEPLAYABILITIES_API UGameplayEffectComponent : public UObject
+UCLASS(Abstract, Const, DefaultToInstanced, EditInlineNew, CollapseCategories, Within=GameplayEffect, MinimalAPI)
+class UGameplayEffectComponent : public UObject
 {
 	GENERATED_BODY()
 
 public:
 	/** Constructor */
-	UGameplayEffectComponent();
+	UE_API UGameplayEffectComponent();
 
 	/** Returns the GameplayEffect that owns this Component (the Outer) */
-	UGameplayEffect* GetOwner() const;
+	UE_API UGameplayEffect* GetOwner() const;
 
 	/** 
 	 * Can the GameplayEffectSpec apply to the passed-in ASC?  All Components of the GE must return true, or a single one can return false to prohibit the application.
@@ -76,7 +78,7 @@ public:
 	 * Allow each Gameplay Effect Component to validate its own data.  Any warnings/errors will immediately show up in the Gameplay Effect when in Editor.
 	 * The default implementation ensures we only have a single type of any given class.  Override this function to change that functionality and use Super::Super::IsDataValid if needed.
 	 */
-	virtual EDataValidationResult IsDataValid(class FDataValidationContext& Context) const override;
+	UE_API virtual EDataValidationResult IsDataValid(class FDataValidationContext& Context) const override;
 #endif
 
 protected:
@@ -96,3 +98,4 @@ const GEComponentClass* FindParentComponent(const GEComponentClass& ChildCompone
 	return ParentGE ? ParentGE->template FindComponent<GEComponentClass>() : nullptr;
 }
 
+#undef UE_API
diff --git a/Engine/Plugins/Runtime/GameplayAbilities/Source/GameplayAbilities/Public/GameplayEffectComponents/AbilitiesGameplayEffectComponent.h b/Engine/Plugins/Runtime/GameplayAbilities/Source/GameplayAbilities/Public/GameplayEffectComponents/AbilitiesGameplayEffectComponent.h
--- a/Engine/Plugins/Runtime/GameplayAbilities/Source/GameplayAbilities/Public/GameplayEffectComponents/AbilitiesGameplayEffectComponent.h
+++ b/Engine/Plugins/Runtime/GameplayAbilities/Source/GameplayAbilities/Public/GameplayEffectComponents/AbilitiesGameplayEffectComponent.h
@@ -6,6 +6,8 @@
 #include "GameplayAbilitySpec.h"
 #include "AbilitiesGameplayEffectComponent.generated.h"
 
+#define UE_API GAMEPLAYABILITIES_API
+
 /** Options on how to configure a GameplayAbilitySpec to grant an Actor */
 USTRUCT()
 struct FGameplayAbilitySpecConfig
@@ -32,8 +34,8 @@ struct FGameplayAbilitySpecConfig
 /**
  * Grants additional Gameplay Abilities to the Target of a Gameplay Effect while active
  */
-UCLASS(DisplayName="Grant Gameplay Abilities")
-class GAMEPLAYABILITIES_API UAbilitiesGameplayEffectComponent : public UGameplayEffectComponent
+UCLASS(DisplayName="Grant Gameplay Abilities", MinimalAPI)
+class UAbilitiesGameplayEffectComponent : public UGameplayEffectComponent
 {
 	friend class UGameplayEffect; // for upgrade path
 
@@ -41,35 +43,37 @@ class GAMEPLAYABILITIES_API UAbilitiesGameplayEffectComponent : public UGameplay
 
 public:
 	/** Constructor */
-	UAbilitiesGameplayEffectComponent();
+	UE_API UAbilitiesGameplayEffectComponent();
 
 	/** Register for the appropriate events when we're applied */
-	virtual bool OnActiveGameplayEffectAdded(FActiveGameplayEffectsContainer& ActiveGEContainer, FActiveGameplayEffect& ActiveGE) const override;
+	UE_API virtual bool OnActiveGameplayEffectAdded(FActiveGameplayEffectsContainer& ActiveGEContainer, FActiveGameplayEffect& ActiveGE) const override;
 
 	/** Adds an entry for granting Gameplay Abilities */
-	void AddGrantedAbilityConfig(const FGameplayAbilitySpecConfig& Config);
+	UE_API void AddGrantedAbilityConfig(const FGameplayAbilitySpecConfig& Config);
 
 #if WITH_EDITOR
 	/** Warn on misconfigured Gameplay Effect */
-	virtual EDataValidationResult IsDataValid(class FDataValidationContext& Context) const override;
+	UE_API virtual EDataValidationResult IsDataValid(class FDataValidationContext& Context) const override;
 #endif
 
 protected:
 	/** This allows us to be notified when the owning GameplayEffect has had its inhibition changed (which can happen on the initial application). */
-	virtual void OnInhibitionChanged(FActiveGameplayEffectHandle ActiveGEHandle, bool bIsInhibited) const;
+	UE_API virtual void OnInhibitionChanged(FActiveGameplayEffectHandle ActiveGEHandle, bool bIsInhibited) const;
 
 	/** We should grant the configured Gameplay Abilities */
-	virtual void GrantAbilities(FActiveGameplayEffectHandle ActiveGEHandle) const;
+	UE_API virtual void GrantAbilities(FActiveGameplayEffectHandle ActiveGEHandle) const;
 
 	/** We should remove the configured Gameplay Abilities */
-	virtual void RemoveAbilities(FActiveGameplayEffectHandle ActiveGEHandle) const;
+	UE_API virtual void RemoveAbilities(FActiveGameplayEffectHandle ActiveGEHandle) const;
 
 private:
 	/** We must undo all effects when removed */
-	void OnActiveGameplayEffectRemoved(const FGameplayEffectRemovalInfo& RemovalInfo) const;
+	UE_API void OnActiveGameplayEffectRemoved(const FGameplayEffectRemovalInfo& RemovalInfo) const;
 
 protected:
 	/** Abilities to Grant to the Target while this Gameplay Effect is active */
 	UPROPERTY(EditDefaultsOnly, Category = GrantAbilities)
 	TArray<FGameplayAbilitySpecConfig>	GrantAbilityConfigs;
 };
+
+#undef UE_API
diff --git a/Engine/Plugins/Runtime/GameplayAbilities/Source/GameplayAbilities/Public/GameplayEffectComponents/AdditionalEffectsGameplayEffectComponent.h b/Engine/Plugins/Runtime/GameplayAbilities/Source/GameplayAbilities/Public/GameplayEffectComponents/AdditionalEffectsGameplayEffectComponent.h
--- a/Engine/Plugins/Runtime/GameplayAbilities/Source/GameplayAbilities/Public/GameplayEffectComponents/AdditionalEffectsGameplayEffectComponent.h
+++ b/Engine/Plugins/Runtime/GameplayAbilities/Source/GameplayAbilities/Public/GameplayEffectComponents/AdditionalEffectsGameplayEffectComponent.h
@@ -6,9 +6,11 @@
 #include "GameplayEffectComponent.h"
 #include "AdditionalEffectsGameplayEffectComponent.generated.h"
 
+#define UE_API GAMEPLAYABILITIES_API
+
 /** Add additional Gameplay Effects that attempt to activate under certain conditions (or no conditions) */
-UCLASS(CollapseCategories, DisplayName="Apply Additional Effects")
-class GAMEPLAYABILITIES_API UAdditionalEffectsGameplayEffectComponent : public UGameplayEffectComponent
+UCLASS(CollapseCategories, DisplayName="Apply Additional Effects", MinimalAPI)
+class UAdditionalEffectsGameplayEffectComponent : public UGameplayEffectComponent
 {
 	GENERATED_BODY()
 
@@ -16,18 +18,18 @@ public:
 	/**
      * Called when a Gameplay Effect is Added to the ActiveGameplayEffectsContainer.  We register a callback to execute the OnComplete Gameplay Effects.
      */
-	virtual bool OnActiveGameplayEffectAdded(FActiveGameplayEffectsContainer& GEContainer, FActiveGameplayEffect& ActiveGE) const override;
+	UE_API virtual bool OnActiveGameplayEffectAdded(FActiveGameplayEffectsContainer& GEContainer, FActiveGameplayEffect& ActiveGE) const override;
 
 	/**
 	 * Called when a Gameplay Effect is applied.  This executes the OnApplication Gameplay Effects.
 	 */
-	virtual void OnGameplayEffectApplied(FActiveGameplayEffectsContainer& ActiveGEContainer, FGameplayEffectSpec& GESpec, FPredictionKey& PredictionKey) const override;
+	UE_API virtual void OnGameplayEffectApplied(FActiveGameplayEffectsContainer& ActiveGEContainer, FGameplayEffectSpec& GESpec, FPredictionKey& PredictionKey) const override;
 
 #if WITH_EDITOR
 	/**
 	 * There are some fields that are incompatible with one another, so let's catch them during configuration phase. 
 	 */
-	virtual EDataValidationResult IsDataValid(class FDataValidationContext& Context) const override;
+	UE_API virtual EDataValidationResult IsDataValid(class FDataValidationContext& Context) const override;
 #endif // WITH_EDITOR
 
 
@@ -35,7 +37,7 @@ protected:
 	/**
 	 * Whenever the ActiveGE gets removed, we want to apply the configured OnComplete GameplayEffects
 	 */
-	void OnActiveGameplayEffectRemoved(const FGameplayEffectRemovalInfo& RemovalInfo, FActiveGameplayEffectsContainer* ActiveGEContainer) const;
+	UE_API void OnActiveGameplayEffectRemoved(const FGameplayEffectRemovalInfo& RemovalInfo, FActiveGameplayEffectsContainer* ActiveGEContainer) const;
 
 public:
 	/**
@@ -61,3 +63,5 @@ public:
 	UPROPERTY(EditDefaultsOnly, Category = OnComplete)
 	TArray<TSubclassOf<UGameplayEffect>> OnCompletePrematurely;
 };
+
+#undef UE_API
diff --git a/Engine/Plugins/Runtime/GameplayAbilities/Source/GameplayAbilities/Public/GameplayEffectComponents/AssetTagsGameplayEffectComponent.h b/Engine/Plugins/Runtime/GameplayAbilities/Source/GameplayAbilities/Public/GameplayEffectComponents/AssetTagsGameplayEffectComponent.h
--- a/Engine/Plugins/Runtime/GameplayAbilities/Source/GameplayAbilities/Public/GameplayEffectComponents/AssetTagsGameplayEffectComponent.h
+++ b/Engine/Plugins/Runtime/GameplayAbilities/Source/GameplayAbilities/Public/GameplayEffectComponents/AssetTagsGameplayEffectComponent.h
@@ -6,28 +6,30 @@
 #include "GameplayEffectComponent.h"
 #include "AssetTagsGameplayEffectComponent.generated.h"
 
+#define UE_API GAMEPLAYABILITIES_API
+
 /** These are tags that the Gameplay Effect Asset itself 'has' (owns). These do _not_ transfer to any Actors */
-UCLASS(DisplayName="Tags This Effect Has (Asset Tags)")
-class GAMEPLAYABILITIES_API UAssetTagsGameplayEffectComponent : public UGameplayEffectComponent
+UCLASS(DisplayName="Tags This Effect Has (Asset Tags)", MinimalAPI)
+class UAssetTagsGameplayEffectComponent : public UGameplayEffectComponent
 {
 	GENERATED_BODY()
 
 public:
 	/** Setup a better EditorFriendlyName and do some initialization */
-	virtual void PostInitProperties() override;
+	UE_API virtual void PostInitProperties() override;
 
 	/** Needed to properly apply FInheritedTagContainer properties */
-	virtual void OnGameplayEffectChanged() override;
+	UE_API virtual void OnGameplayEffectChanged() override;
 
 	/** Gets the Asset Tags inherited tag structure (as configured) */
 	const FInheritedTagContainer& GetConfiguredAssetTagChanges() const { return InheritableAssetTags; }
 
 	/** Applies the Asset Tags to the GameplayEffect (and saves those changes) */
-	void SetAndApplyAssetTagChanges(const FInheritedTagContainer& TagContainerMods);
+	UE_API void SetAndApplyAssetTagChanges(const FInheritedTagContainer& TagContainerMods);
 
 #if WITH_EDITOR
 	/** Needed to properly update FInheritedTagContainer properties */
-	virtual void PostEditChangeProperty(FPropertyChangedEvent& PropertyChangedEvent) override;
+	UE_API virtual void PostEditChangeProperty(FPropertyChangedEvent& PropertyChangedEvent) override;
 
 private:
 	/** Get a cached version of the FProperty Name for PostEditChangeProperty */
@@ -40,10 +42,12 @@ private:
 
 private:
 	/** Applies the Asset Tags to the GameplayEffect */
-	void ApplyAssetTagChanges() const;
+	UE_API void ApplyAssetTagChanges() const;
 
 private:
 	/** The GameplayEffect's Tags: tags the the GE *has* and DOES NOT give to the actor. */
 	UPROPERTY(EditDefaultsOnly, Category = Tags, meta = (DisplayName = "AssetTags", Categories = "OwnedTagsCategory"))
 	FInheritedTagContainer InheritableAssetTags;
 };
+
+#undef UE_API
diff --git a/Engine/Plugins/Runtime/GameplayAbilities/Source/GameplayAbilities/Public/GameplayEffectComponents/BlockAbilityTagsGameplayEffectComponent.h b/Engine/Plugins/Runtime/GameplayAbilities/Source/GameplayAbilities/Public/GameplayEffectComponents/BlockAbilityTagsGameplayEffectComponent.h
--- a/Engine/Plugins/Runtime/GameplayAbilities/Source/GameplayAbilities/Public/GameplayEffectComponents/BlockAbilityTagsGameplayEffectComponent.h
+++ b/Engine/Plugins/Runtime/GameplayAbilities/Source/GameplayAbilities/Public/GameplayEffectComponents/BlockAbilityTagsGameplayEffectComponent.h
@@ -6,31 +6,33 @@
 #include "GameplayEffectComponent.h"
 #include "BlockAbilityTagsGameplayEffectComponent.generated.h"
 
+#define UE_API GAMEPLAYABILITIES_API
+
 /** Handles blocking the activation of Gameplay Abilities based on Gameplay Tags for the Target actor of the owner Gameplay Effect */
-UCLASS(DisplayName="Block Abilities with Tags")
-class GAMEPLAYABILITIES_API UBlockAbilityTagsGameplayEffectComponent : public UGameplayEffectComponent
+UCLASS(DisplayName="Block Abilities with Tags", MinimalAPI)
+class UBlockAbilityTagsGameplayEffectComponent : public UGameplayEffectComponent
 {
 	GENERATED_BODY()
 	
 public:
 	/** Setup an EditorFriendlyName and do some initialization */
-	virtual void PostInitProperties() override;
+	UE_API virtual void PostInitProperties() override;
 
 	/** Needed to properly apply FInheritedTagContainer properties */
-	virtual void OnGameplayEffectChanged() override;
+	UE_API virtual void OnGameplayEffectChanged() override;
 
 	/** Gets the Blocked Ability Tags inherited tag structure (as configured) */
 	const FInheritedTagContainer& GetConfiguredBlockedAbilityTagChanges() const { return InheritableBlockedAbilityTagsContainer; }
 
 	/** Applies the Blocked Ability Tags to the GameplayEffect (and saves those changes) so that when the Gameplay Effect is applied, these Blocking Tags are then applied to the Target Actor */
-	void SetAndApplyBlockedAbilityTagChanges(const FInheritedTagContainer& TagContainerMods);
+	UE_API void SetAndApplyBlockedAbilityTagChanges(const FInheritedTagContainer& TagContainerMods);
 
 #if WITH_EDITOR
 	/** Needed to properly update FInheritedTagContainer properties */
-	virtual void PostEditChangeProperty(FPropertyChangedEvent& PropertyChangedEvent) override;
+	UE_API virtual void PostEditChangeProperty(FPropertyChangedEvent& PropertyChangedEvent) override;
 
 	/** Validate incompatible configurations */
-	virtual EDataValidationResult IsDataValid(class FDataValidationContext& Context) const override;
+	UE_API virtual EDataValidationResult IsDataValid(class FDataValidationContext& Context) const override;
 
 private:
 	/** Get a cached version of the FProperty Name for PostEditChangeProperty */
@@ -43,10 +45,12 @@ private:
 
 private:
 	/** Applies the Blocked Ability Tags to the owning GameplayEffect */
-	void ApplyBlockedAbilityTagChanges() const;
+	UE_API void ApplyBlockedAbilityTagChanges() const;
 
 private:
 	/** These tags are applied to the target actor of the Gameplay Effect.  Blocked Ability Tags prevent Gameplay Abilities with these tags from executing. */
 	UPROPERTY(EditDefaultsOnly, Category = None, meta = (DisplayName = "Block Abilities w/ Tags"))
 	FInheritedTagContainer InheritableBlockedAbilityTagsContainer;
 };
+
+#undef UE_API
diff --git a/Engine/Plugins/Runtime/GameplayAbilities/Source/GameplayAbilities/Public/GameplayEffectComponents/ChanceToApplyGameplayEffectComponent.h b/Engine/Plugins/Runtime/GameplayAbilities/Source/GameplayAbilities/Public/GameplayEffectComponents/ChanceToApplyGameplayEffectComponent.h
--- a/Engine/Plugins/Runtime/GameplayAbilities/Source/GameplayAbilities/Public/GameplayEffectComponents/ChanceToApplyGameplayEffectComponent.h
+++ b/Engine/Plugins/Runtime/GameplayAbilities/Source/GameplayAbilities/Public/GameplayEffectComponents/ChanceToApplyGameplayEffectComponent.h
@@ -7,27 +7,31 @@
 
 #include "ChanceToApplyGameplayEffectComponent.generated.h"
 
+#define UE_API GAMEPLAYABILITIES_API
+
 /** Applies a probablity to the application conditions of the Gameplay Effect. */
-UCLASS(DisplayName="Chance To Apply This Effect")
-class GAMEPLAYABILITIES_API UChanceToApplyGameplayEffectComponent : public UGameplayEffectComponent
+UCLASS(DisplayName="Chance To Apply This Effect", MinimalAPI)
+class UChanceToApplyGameplayEffectComponent : public UGameplayEffectComponent
 {
 	GENERATED_BODY()
 	
 public:
 	/** Warn of errors in the Editor */
-	virtual void PostLoad() override;
+	UE_API virtual void PostLoad() override;
 
 	/** Determine if we can apply this GameplayEffect or not */
-	virtual bool CanGameplayEffectApply(const FActiveGameplayEffectsContainer& ActiveGEContainer, const FGameplayEffectSpec& GESpec) const override;
+	UE_API virtual bool CanGameplayEffectApply(const FActiveGameplayEffectsContainer& ActiveGEContainer, const FGameplayEffectSpec& GESpec) const override;
 
 	/** Get the configured chance to apply to target */
 	const FScalableFloat& GetChanceToApplyToTarget() const { return ChanceToApplyToTarget; }
 
 	/** Set the chance to apply to target */
-	void SetChanceToApplyToTarget(const FScalableFloat& ChanceToApplyToTarget);
+	UE_API void SetChanceToApplyToTarget(const FScalableFloat& ChanceToApplyToTarget);
 
 private:
 	/** Probability that this gameplay effect will be applied to the target actor (0.0 for never, 1.0 for always) */
 	UPROPERTY(EditDefaultsOnly, Category = Application, meta = (GameplayAttribute = "True"))
 	FScalableFloat ChanceToApplyToTarget;
 };
+
+#undef UE_API
diff --git a/Engine/Plugins/Runtime/GameplayAbilities/Source/GameplayAbilities/Public/GameplayEffectComponents/CustomCanApplyGameplayEffectComponent.h b/Engine/Plugins/Runtime/GameplayAbilities/Source/GameplayAbilities/Public/GameplayEffectComponents/CustomCanApplyGameplayEffectComponent.h
--- a/Engine/Plugins/Runtime/GameplayAbilities/Source/GameplayAbilities/Public/GameplayEffectComponents/CustomCanApplyGameplayEffectComponent.h
+++ b/Engine/Plugins/Runtime/GameplayAbilities/Source/GameplayAbilities/Public/GameplayEffectComponents/CustomCanApplyGameplayEffectComponent.h
@@ -7,20 +7,24 @@
 
 #include "CustomCanApplyGameplayEffectComponent.generated.h"
 
+#define UE_API GAMEPLAYABILITIES_API
+
 /** Handles configuration of a CustomApplicationRequirement function to see if this GameplayEffect should apply */
-UCLASS(DisplayName="Custom Can Apply This Effect")
-class GAMEPLAYABILITIES_API UCustomCanApplyGameplayEffectComponent : public UGameplayEffectComponent
+UCLASS(DisplayName="Custom Can Apply This Effect", MinimalAPI)
+class UCustomCanApplyGameplayEffectComponent : public UGameplayEffectComponent
 {
 	GENERATED_BODY()
 	
 public:
-	UCustomCanApplyGameplayEffectComponent();
+	UE_API UCustomCanApplyGameplayEffectComponent();
 
 	/** Determine if we can apply this GameplayEffect or not */
-	virtual bool CanGameplayEffectApply(const FActiveGameplayEffectsContainer& ActiveGEContainer, const FGameplayEffectSpec& GESpec) const override;
+	UE_API virtual bool CanGameplayEffectApply(const FActiveGameplayEffectsContainer& ActiveGEContainer, const FGameplayEffectSpec& GESpec) const override;
 
 public:
 	/** Custom application requirements */
 	UPROPERTY(EditDefaultsOnly, Category = Application, meta = (DisplayName = "Custom Application Requirement"))
 	TArray<TSubclassOf<UGameplayEffectCustomApplicationRequirement>> ApplicationRequirements;
 };
+
+#undef UE_API
diff --git a/Engine/Plugins/Runtime/GameplayAbilities/Source/GameplayAbilities/Public/GameplayEffectComponents/ImmunityGameplayEffectComponent.h b/Engine/Plugins/Runtime/GameplayAbilities/Source/GameplayAbilities/Public/GameplayEffectComponents/ImmunityGameplayEffectComponent.h
--- a/Engine/Plugins/Runtime/GameplayAbilities/Source/GameplayAbilities/Public/GameplayEffectComponents/ImmunityGameplayEffectComponent.h
+++ b/Engine/Plugins/Runtime/GameplayAbilities/Source/GameplayAbilities/Public/GameplayEffectComponents/ImmunityGameplayEffectComponent.h
@@ -6,36 +6,40 @@
 #include "GameplayEffectComponent.h"
 #include "ImmunityGameplayEffectComponent.generated.h"
 
+#define UE_API GAMEPLAYABILITIES_API
+
 /** 
  * Immunity is blocking the application of _other_ GameplayEffectSpecs.
  * This registers a global handler on the ASC to block the application of other GameplayEffectSpecs.
  */
-UCLASS(DisplayName="Immunity to Other Effects")
-class GAMEPLAYABILITIES_API UImmunityGameplayEffectComponent : public UGameplayEffectComponent
+UCLASS(DisplayName="Immunity to Other Effects", MinimalAPI)
+class UImmunityGameplayEffectComponent : public UGameplayEffectComponent
 {
 	friend class UGameplayEffect; // Needed for upgrade path
 	GENERATED_BODY()
 
 public:
 	/** Constructor to set EditorFriendlyName */
-	UImmunityGameplayEffectComponent();
+	UE_API UImmunityGameplayEffectComponent();
 
 	/** We need to register our callback to check for Immunity */
-	virtual bool OnActiveGameplayEffectAdded(FActiveGameplayEffectsContainer& ActiveGEContainer, FActiveGameplayEffect& ActiveGE) const override;
+	UE_API virtual bool OnActiveGameplayEffectAdded(FActiveGameplayEffectsContainer& ActiveGEContainer, FActiveGameplayEffect& ActiveGE) const override;
 
 #if WITH_EDITOR
 	/**
 	 * Warn about instant Gameplay Effects as they do not persist over time (and therefore cannot grant immunity)
 	 */
-	virtual EDataValidationResult IsDataValid(class FDataValidationContext& Context) const override;
+	UE_API virtual EDataValidationResult IsDataValid(class FDataValidationContext& Context) const override;
 #endif // WITH_EDITOR
 
 protected:
 	/** We register with the AbilitySystemComponent to try and block any GESpecs we think we should be immune to */
-	bool AllowGameplayEffectApplication(const FActiveGameplayEffectsContainer& ActiveGEContainer, const FGameplayEffectSpec& GESpecToConsider, FActiveGameplayEffectHandle ImmunityActiveGE) const;
+	UE_API bool AllowGameplayEffectApplication(const FActiveGameplayEffectsContainer& ActiveGEContainer, const FGameplayEffectSpec& GESpecToConsider, FActiveGameplayEffectHandle ImmunityActiveGE) const;
 
 public:
 	/** Grants immunity to GameplayEffects that match this query. */
 	UPROPERTY(EditDefaultsOnly, BlueprintReadOnly, Category = None)
 	TArray<FGameplayEffectQuery> ImmunityQueries;
 };
+
+#undef UE_API
diff --git a/Engine/Plugins/Runtime/GameplayAbilities/Source/GameplayAbilities/Public/GameplayEffectComponents/RemoveOtherGameplayEffectComponent.h b/Engine/Plugins/Runtime/GameplayAbilities/Source/GameplayAbilities/Public/GameplayEffectComponents/RemoveOtherGameplayEffectComponent.h
--- a/Engine/Plugins/Runtime/GameplayAbilities/Source/GameplayAbilities/Public/GameplayEffectComponents/RemoveOtherGameplayEffectComponent.h
+++ b/Engine/Plugins/Runtime/GameplayAbilities/Source/GameplayAbilities/Public/GameplayEffectComponents/RemoveOtherGameplayEffectComponent.h
@@ -6,28 +6,30 @@
 #include "GameplayEffectComponent.h"
 #include "RemoveOtherGameplayEffectComponent.generated.h"
 
+#define UE_API GAMEPLAYABILITIES_API
+
 struct FGameplayEffectRemovalInfo;
 
 /** Remove other Gameplay Effects based on certain conditions */
-UCLASS(DisplayName="Remove Other Effects")
-class GAMEPLAYABILITIES_API URemoveOtherGameplayEffectComponent : public UGameplayEffectComponent
+UCLASS(DisplayName="Remove Other Effects", MinimalAPI)
+class URemoveOtherGameplayEffectComponent : public UGameplayEffectComponent
 {
 	GENERATED_BODY()
 
 public:
 	/** Constructor to set EditorFriendlyName */
-	URemoveOtherGameplayEffectComponent();
+	UE_API URemoveOtherGameplayEffectComponent();
 
 	/**
 	 * We will re-run RemoveGameplayEffectQueries every time the owning Gameplay Effect is applied.
 	 */
-	virtual void OnGameplayEffectApplied(FActiveGameplayEffectsContainer& ActiveGEContainer, FGameplayEffectSpec& GESpec, FPredictionKey& PredictionKey) const override;
+	UE_API virtual void OnGameplayEffectApplied(FActiveGameplayEffectsContainer& ActiveGEContainer, FGameplayEffectSpec& GESpec, FPredictionKey& PredictionKey) const override;
 
 #if WITH_EDITOR
 	/**
 	 * Warn about periodic Gameplay Effects, that you probably instead want the OngoingTagRequirements in TagRequirementsGameplayEffectComponent
 	 */
-	virtual EDataValidationResult IsDataValid(class FDataValidationContext& Context) const override;
+	UE_API virtual EDataValidationResult IsDataValid(class FDataValidationContext& Context) const override;
 #endif // WITH_EDITOR
 
 public:
@@ -35,3 +37,5 @@ public:
 	UPROPERTY(EditDefaultsOnly, Category = None)
 	TArray<FGameplayEffectQuery> RemoveGameplayEffectQueries;
 };
+
+#undef UE_API
diff --git a/Engine/Plugins/Runtime/GameplayAbilities/Source/GameplayAbilities/Public/GameplayEffectComponents/TargetTagRequirementsGameplayEffectComponent.h b/Engine/Plugins/Runtime/GameplayAbilities/Source/GameplayAbilities/Public/GameplayEffectComponents/TargetTagRequirementsGameplayEffectComponent.h
--- a/Engine/Plugins/Runtime/GameplayAbilities/Source/GameplayAbilities/Public/GameplayEffectComponents/TargetTagRequirementsGameplayEffectComponent.h
+++ b/Engine/Plugins/Runtime/GameplayAbilities/Source/GameplayAbilities/Public/GameplayEffectComponents/TargetTagRequirementsGameplayEffectComponent.h
@@ -6,29 +6,31 @@
 #include "GameplayEffectComponent.h"
 #include "TargetTagRequirementsGameplayEffectComponent.generated.h"
 
+#define UE_API GAMEPLAYABILITIES_API
+
 struct FGameplayEffectRemovalInfo;
 
 /** Specifies tag requirements that the Target (owner of the Gameplay Effect) must have if this GE should apply or continue to execute */
-UCLASS(DisplayName="Require Tags to Apply/Continue This Effect")
-class GAMEPLAYABILITIES_API UTargetTagRequirementsGameplayEffectComponent : public UGameplayEffectComponent
+UCLASS(DisplayName="Require Tags to Apply/Continue This Effect", MinimalAPI)
+class UTargetTagRequirementsGameplayEffectComponent : public UGameplayEffectComponent
 {
 	GENERATED_BODY()
 
 public:
 	/** Constructor to set EditorFriendlyName */
-	UTargetTagRequirementsGameplayEffectComponent();
+	UE_API UTargetTagRequirementsGameplayEffectComponent();
 
 	/** Can we apply to the ActiveGEContainer? */
-	virtual bool CanGameplayEffectApply(const FActiveGameplayEffectsContainer& ActiveGEContainer, const FGameplayEffectSpec& GESpec) const override;
+	UE_API virtual bool CanGameplayEffectApply(const FActiveGameplayEffectsContainer& ActiveGEContainer, const FGameplayEffectSpec& GESpec) const override;
 
 	/** Once we've applied, we need to register for ongoing requirements */
-	virtual bool OnActiveGameplayEffectAdded(FActiveGameplayEffectsContainer& GEContainer, FActiveGameplayEffect& ActiveGE) const override;
+	UE_API virtual bool OnActiveGameplayEffectAdded(FActiveGameplayEffectsContainer& GEContainer, FActiveGameplayEffect& ActiveGE) const override;
 
 #if WITH_EDITOR
 	/**
 	 * Validate incompatible configurations
 	 */
-	virtual EDataValidationResult IsDataValid(class FDataValidationContext& Context) const override;
+	UE_API virtual EDataValidationResult IsDataValid(class FDataValidationContext& Context) const override;
 #endif // WITH_EDITOR
 
 private:
@@ -54,3 +56,5 @@ public:
 	UPROPERTY(EditDefaultsOnly, Category = Tags)
 	FGameplayTagRequirements RemovalTagRequirements;
 };
+
+#undef UE_API
diff --git a/Engine/Plugins/Runtime/GameplayAbilities/Source/GameplayAbilities/Public/GameplayEffectComponents/TargetTagsGameplayEffectComponent.h b/Engine/Plugins/Runtime/GameplayAbilities/Source/GameplayAbilities/Public/GameplayEffectComponents/TargetTagsGameplayEffectComponent.h
--- a/Engine/Plugins/Runtime/GameplayAbilities/Source/GameplayAbilities/Public/GameplayEffectComponents/TargetTagsGameplayEffectComponent.h
+++ b/Engine/Plugins/Runtime/GameplayAbilities/Source/GameplayAbilities/Public/GameplayEffectComponents/TargetTagsGameplayEffectComponent.h
@@ -6,31 +6,33 @@
 #include "GameplayEffectComponent.h"
 #include "TargetTagsGameplayEffectComponent.generated.h"
 
+#define UE_API GAMEPLAYABILITIES_API
+
 /** Handles granting Tags to the Target (sometimes referred to as the Owner) of the Gameplay Effect */
-UCLASS(DisplayName="Grant Tags to Target Actor")
-class GAMEPLAYABILITIES_API UTargetTagsGameplayEffectComponent : public UGameplayEffectComponent
+UCLASS(DisplayName="Grant Tags to Target Actor", MinimalAPI)
+class UTargetTagsGameplayEffectComponent : public UGameplayEffectComponent
 {
 	GENERATED_BODY()
 
 public:
 	/** Setup an EditorFriendlyName and do some initialization */
-	virtual void PostInitProperties() override;
+	UE_API virtual void PostInitProperties() override;
 
 	/** Needed to properly load FInheritedTagContainer properties */
-	virtual void OnGameplayEffectChanged() override;
+	UE_API virtual void OnGameplayEffectChanged() override;
 
 	/** Gets the (Granted) Target Tags inherited tag structure (as configured) */
 	const FInheritedTagContainer& GetConfiguredTargetTagChanges() const { return InheritableGrantedTagsContainer; }
 
 	/** Applies the (Granted) Target Tags to the GameplayEffect (and saves those changes) so that it may grant these Tags the Target Actor when Applied */
-	void SetAndApplyTargetTagChanges(const FInheritedTagContainer& TagContainerMods);
+	UE_API void SetAndApplyTargetTagChanges(const FInheritedTagContainer& TagContainerMods);
 
 #if WITH_EDITOR
 	/** Needed to properly update FInheritedTagContainer properties */
-	virtual void PostEditChangeProperty(FPropertyChangedEvent& PropertyChangedEvent) override;
+	UE_API virtual void PostEditChangeProperty(FPropertyChangedEvent& PropertyChangedEvent) override;
 
 	/** Validate incompatible configurations */
-	virtual EDataValidationResult IsDataValid(class FDataValidationContext& Context) const override;
+	UE_API virtual EDataValidationResult IsDataValid(class FDataValidationContext& Context) const override;
 
 private:
 	/** Get a cached version of the FProperty Name for PostEditChangeProperty */
@@ -43,10 +45,12 @@ private:
 
 private:
 	/** Applies the configured (Granted) Target Tags to the owning GameplayEffect. */
-	void ApplyTargetTagChanges() const;
+	UE_API void ApplyTargetTagChanges() const;
 
 private:
 	/** These tags are applied (granted) to the target actor of the Gameplay Effect.  The Target would then "own" these Tags. */
 	UPROPERTY(EditDefaultsOnly, Category = None, meta = (DisplayName = "Add Tags", Categories = "OwnedTagsCategory"))
 	FInheritedTagContainer InheritableGrantedTagsContainer;
 };
+
+#undef UE_API
diff --git a/Engine/Plugins/Runtime/GameplayAbilities/Source/GameplayAbilities/Public/GameplayEffectCustomApplicationRequirement.h b/Engine/Plugins/Runtime/GameplayAbilities/Source/GameplayAbilities/Public/GameplayEffectCustomApplicationRequirement.h
--- a/Engine/Plugins/Runtime/GameplayAbilities/Source/GameplayAbilities/Public/GameplayEffectCustomApplicationRequirement.h
+++ b/Engine/Plugins/Runtime/GameplayAbilities/Source/GameplayAbilities/Public/GameplayEffectCustomApplicationRequirement.h
@@ -5,18 +5,22 @@
 #include "UObject/Object.h"
 #include "GameplayEffectCustomApplicationRequirement.generated.h"
 
+#define UE_API GAMEPLAYABILITIES_API
+
 struct FGameplayEffectSpec;
 class UGameplayEffect;
 class UAbilitySystemComponent;
 
 /** Class used to perform custom gameplay effect modifier calculations, either via blueprint or native code */ 
-UCLASS(BlueprintType, Blueprintable, Abstract)
-class GAMEPLAYABILITIES_API UGameplayEffectCustomApplicationRequirement : public UObject
+UCLASS(BlueprintType, Blueprintable, Abstract, MinimalAPI)
+class UGameplayEffectCustomApplicationRequirement : public UObject
 {
 	GENERATED_BODY()
 
 public:
 	/** Return whether the gameplay effect should be applied or not */
 	UFUNCTION(BlueprintNativeEvent, Category="Calculation")
-	bool CanApplyGameplayEffect(const UGameplayEffect* GameplayEffect, const FGameplayEffectSpec& Spec, UAbilitySystemComponent* ASC) const;
+	UE_API bool CanApplyGameplayEffect(const UGameplayEffect* GameplayEffect, const FGameplayEffectSpec& Spec, UAbilitySystemComponent* ASC) const;
 };
+
+#undef UE_API
diff --git a/Engine/Plugins/Runtime/GameplayAbilities/Source/GameplayAbilities/Public/GameplayEffectExecutionCalculation.h b/Engine/Plugins/Runtime/GameplayAbilities/Source/GameplayAbilities/Public/GameplayEffectExecutionCalculation.h
--- a/Engine/Plugins/Runtime/GameplayAbilities/Source/GameplayAbilities/Public/GameplayEffectExecutionCalculation.h
+++ b/Engine/Plugins/Runtime/GameplayAbilities/Source/GameplayAbilities/Public/GameplayEffectExecutionCalculation.h
@@ -12,39 +12,41 @@
 #include "GameplayEffectCalculation.h"
 #include "GameplayEffectExecutionCalculation.generated.h"
 
+#define UE_API GAMEPLAYABILITIES_API
+
 class UAbilitySystemComponent;
 
 /** Struct representing parameters for a custom gameplay effect execution. Should not be held onto via reference, used just for the scope of the execution */
 USTRUCT(BlueprintType)
-struct GAMEPLAYABILITIES_API FGameplayEffectCustomExecutionParameters
+struct FGameplayEffectCustomExecutionParameters
 {
 	GENERATED_USTRUCT_BODY()
 
 public:
 
 	// Constructors
-	FGameplayEffectCustomExecutionParameters();
-	FGameplayEffectCustomExecutionParameters(FGameplayEffectSpec& InOwningSpec, const TArray<FGameplayEffectExecutionScopedModifierInfo>& InScopedMods, UAbilitySystemComponent* InTargetAbilityComponent, const FGameplayTagContainer& InPassedIntags, const FPredictionKey& InPredictionKey);
-	FGameplayEffectCustomExecutionParameters(FGameplayEffectSpec& InOwningSpec, const TArray<FGameplayEffectExecutionScopedModifierInfo>& InScopedMods, UAbilitySystemComponent* InTargetAbilityComponent, const FGameplayTagContainer& InPassedIntags, const FPredictionKey& InPredictionKey, const TArray<FActiveGameplayEffectHandle>& InIgnoreHandles);
+	UE_API FGameplayEffectCustomExecutionParameters();
+	UE_API FGameplayEffectCustomExecutionParameters(FGameplayEffectSpec& InOwningSpec, const TArray<FGameplayEffectExecutionScopedModifierInfo>& InScopedMods, UAbilitySystemComponent* InTargetAbilityComponent, const FGameplayTagContainer& InPassedIntags, const FPredictionKey& InPredictionKey);
+	UE_API FGameplayEffectCustomExecutionParameters(FGameplayEffectSpec& InOwningSpec, const TArray<FGameplayEffectExecutionScopedModifierInfo>& InScopedMods, UAbilitySystemComponent* InTargetAbilityComponent, const FGameplayTagContainer& InPassedIntags, const FPredictionKey& InPredictionKey, const TArray<FActiveGameplayEffectHandle>& InIgnoreHandles);
 
 	/** Simple accessor to owning gameplay spec */
-	const FGameplayEffectSpec& GetOwningSpec() const;
+	UE_API const FGameplayEffectSpec& GetOwningSpec() const;
 
 	/** Non const access. Be careful with this, especially when modifying a spec after attribute capture. */
-	FGameplayEffectSpec* GetOwningSpecForPreExecuteMod() const;
+	UE_API FGameplayEffectSpec* GetOwningSpecForPreExecuteMod() const;
 
 	/** Simple accessor to target ability system component */
-	UAbilitySystemComponent* GetTargetAbilitySystemComponent() const;
+	UE_API UAbilitySystemComponent* GetTargetAbilitySystemComponent() const;
 
 	/** Simple accessor to source ability system component (could be null!) */
-	UAbilitySystemComponent* GetSourceAbilitySystemComponent() const;
+	UE_API UAbilitySystemComponent* GetSourceAbilitySystemComponent() const;
 	
 	/** Simple accessor to the Passed In Tags to this execution */
-	const FGameplayTagContainer& GetPassedInTags() const;
+	UE_API const FGameplayTagContainer& GetPassedInTags() const;
 
-	TArray<FActiveGameplayEffectHandle> GetIgnoreHandles() const;
+	UE_API TArray<FActiveGameplayEffectHandle> GetIgnoreHandles() const;
 	
-	FPredictionKey GetPredictionKey() const;
+	UE_API FPredictionKey GetPredictionKey() const;
 
 	/**
 	 * Attempts to calculate the magnitude of a captured attribute given the specified parameters. Can fail if the gameplay spec doesn't have
@@ -56,7 +58,7 @@ public:
 	 * 
 	 * @return True if the magnitude was successfully calculated, false if it was not
 	 */
-	bool AttemptCalculateCapturedAttributeMagnitude(const FGameplayEffectAttributeCaptureDefinition& InCaptureDef, const FAggregatorEvaluateParameters& InEvalParams, OUT float& OutMagnitude) const;
+	UE_API bool AttemptCalculateCapturedAttributeMagnitude(const FGameplayEffectAttributeCaptureDefinition& InCaptureDef, const FAggregatorEvaluateParameters& InEvalParams, OUT float& OutMagnitude) const;
 	
 	/**
 	 * Attempts to calculate the magnitude of a captured attribute given the specified parameters, including a starting base value. 
@@ -69,7 +71,7 @@ public:
 	 * 
 	 * @return True if the magnitude was successfully calculated, false if it was not
 	 */
-	bool AttemptCalculateCapturedAttributeMagnitudeWithBase(const FGameplayEffectAttributeCaptureDefinition& InCaptureDef, const FAggregatorEvaluateParameters& InEvalParams, float InBaseValue, OUT float& OutMagnitude) const;
+	UE_API bool AttemptCalculateCapturedAttributeMagnitudeWithBase(const FGameplayEffectAttributeCaptureDefinition& InCaptureDef, const FAggregatorEvaluateParameters& InEvalParams, float InBaseValue, OUT float& OutMagnitude) const;
 
 	/**
 	 * Attempts to calculate the base value of a captured attribute given the specified parameters. Can fail if the gameplay spec doesn't have
@@ -80,7 +82,7 @@ public:
 	 * 
 	 * @return True if the base value was successfully calculated, false if it was not
 	 */
-	bool AttemptCalculateCapturedAttributeBaseValue(const FGameplayEffectAttributeCaptureDefinition& InCaptureDef, OUT float& OutBaseValue) const;
+	UE_API bool AttemptCalculateCapturedAttributeBaseValue(const FGameplayEffectAttributeCaptureDefinition& InCaptureDef, OUT float& OutBaseValue) const;
 
 	/**
 	 * Attempts to calculate the bonus magnitude of a captured attribute given the specified parameters. Can fail if the gameplay spec doesn't have
@@ -92,7 +94,7 @@ public:
 	 * 
 	 * @return True if the bonus magnitude was successfully calculated, false if it was not
 	 */
-	bool AttemptCalculateCapturedAttributeBonusMagnitude(const FGameplayEffectAttributeCaptureDefinition& InCaptureDef, const FAggregatorEvaluateParameters& InEvalParams, OUT float& OutBonusMagnitude) const;
+	UE_API bool AttemptCalculateCapturedAttributeBonusMagnitude(const FGameplayEffectAttributeCaptureDefinition& InCaptureDef, const FAggregatorEvaluateParameters& InEvalParams, OUT float& OutBonusMagnitude) const;
 	
 	/**
 	 * Attempts to populate the specified aggregator with a snapshot of a backing captured aggregator. Can fail if the gameplay spec doesn't have
@@ -103,7 +105,7 @@ public:
 	 * 
 	 * @return True if the aggregator was successfully snapshotted, false if it was not
 	 */
-	bool AttemptGetCapturedAttributeAggregatorSnapshot(const FGameplayEffectAttributeCaptureDefinition& InCaptureDef, OUT FAggregator& OutSnapshottedAggregator) const;
+	UE_API bool AttemptGetCapturedAttributeAggregatorSnapshot(const FGameplayEffectAttributeCaptureDefinition& InCaptureDef, OUT FAggregator& OutSnapshottedAggregator) const;
 
 	/** 
 	 *	Returns all modifiers for a given captured def. Note the returned list is a direct reference to the internal list stored in the attribute aggregators (possibly a snapshot copy or possibly the "real" ones).
@@ -112,10 +114,10 @@ public:
 	 *	
 	 *	Consider using ForEachQualifiedAttributeMod when you want to "do something for every qualifier mod".
 	 */
-	bool AttemptGatherAttributeMods(const FGameplayEffectAttributeCaptureDefinition& InCaptureDef, const FAggregatorEvaluateParameters& InEvalParams, OUT TMap<EGameplayModEvaluationChannel, const TArray<FAggregatorMod>*>& OutModMap) const;
+	UE_API bool AttemptGatherAttributeMods(const FGameplayEffectAttributeCaptureDefinition& InCaptureDef, const FAggregatorEvaluateParameters& InEvalParams, OUT TMap<EGameplayModEvaluationChannel, const TArray<FAggregatorMod>*>& OutModMap) const;
 
 	/** Runs given TFunction on every qualifier mod for a given AttributeCaptureDefinition */
-	bool ForEachQualifiedAttributeMod(const FGameplayEffectAttributeCaptureDefinition& InCaptureDef, const FAggregatorEvaluateParameters& InEvalParams, TFunction< void(EGameplayModEvaluationChannel, EGameplayModOp::Type, const FAggregatorMod&) >) const;
+	UE_API bool ForEachQualifiedAttributeMod(const FGameplayEffectAttributeCaptureDefinition& InCaptureDef, const FAggregatorEvaluateParameters& InEvalParams, TFunction< void(EGameplayModEvaluationChannel, EGameplayModOp::Type, const FAggregatorMod&) >) const;
 
 	/**
 	 * Attempts to calculate the magnitude of a transient aggregator given the specified parameters.
@@ -126,7 +128,7 @@ public:
 	 * 
 	 * @return True if the magnitude was successfully calculated, false if it was not
 	 */
-	bool AttemptCalculateTransientAggregatorMagnitude(const FGameplayTag& InAggregatorIdentifier, const FAggregatorEvaluateParameters& InEvalParams, OUT float& OutMagnitude) const;
+	UE_API bool AttemptCalculateTransientAggregatorMagnitude(const FGameplayTag& InAggregatorIdentifier, const FAggregatorEvaluateParameters& InEvalParams, OUT float& OutMagnitude) const;
 	
 	/**
 	 * Attempts to calculate the magnitude of a transient aggregator given the specified parameters, including a starting base value.
@@ -138,7 +140,7 @@ public:
 	 * 
 	 * @return True if the magnitude was successfully calculated, false if it was not
 	 */
-	bool AttemptCalculateTransientAggregatorMagnitudeWithBase(const FGameplayTag& InAggregatorIdentifier, const FAggregatorEvaluateParameters& InEvalParams, float InBaseValue, OUT float& OutMagnitude) const;
+	UE_API bool AttemptCalculateTransientAggregatorMagnitudeWithBase(const FGameplayTag& InAggregatorIdentifier, const FAggregatorEvaluateParameters& InEvalParams, float InBaseValue, OUT float& OutMagnitude) const;
 
 	/**
 	 * Attempts to calculate the base value of a transient aggregator given the specified parameters.
@@ -148,7 +150,7 @@ public:
 	 * 
 	 * @return True if the base value was successfully calculated, false if it was not
 	 */
-	bool AttemptCalculateTransientAggregatorBaseValue(const FGameplayTag& InAggregatorIdentifier, OUT float& OutBaseValue) const;
+	UE_API bool AttemptCalculateTransientAggregatorBaseValue(const FGameplayTag& InAggregatorIdentifier, OUT float& OutBaseValue) const;
 
 	/**
 	 * Attempts to calculate the bonus magnitude of a transient aggregator given the specified parameters.
@@ -159,7 +161,7 @@ public:
 	 * 
 	 * @return True if the bonus magnitude was successfully calculated, false if it was not
 	 */
-	bool AttemptCalculateTransientAggregatorBonusMagnitude(const FGameplayTag& InAggregatorIdentifier, const FAggregatorEvaluateParameters& InEvalParams, OUT float& OutBonusMagnitude) const;
+	UE_API bool AttemptCalculateTransientAggregatorBonusMagnitude(const FGameplayTag& InAggregatorIdentifier, const FAggregatorEvaluateParameters& InEvalParams, OUT float& OutBonusMagnitude) const;
 	
 	/**
 	 * Attempts to populate the specified aggregator with a snapshot of a backing transient aggregator. Can fail if the transient aggregator doesn't exist as a result
@@ -170,7 +172,7 @@ public:
 	 * 
 	 * @return True if the aggregator was successfully snapshotted, false if it was not
 	 */
-	bool AttemptGetCapturedAttributeAggregatorSnapshot(const FGameplayTag& InAggregatorIdentifier, OUT FAggregator& OutSnapshottedAggregator) const;
+	UE_API bool AttemptGetCapturedAttributeAggregatorSnapshot(const FGameplayTag& InAggregatorIdentifier, OUT FAggregator& OutSnapshottedAggregator) const;
 
 private:
 
@@ -196,41 +198,41 @@ private:
 
 /** Struct representing the output of a custom gameplay effect execution. */
 USTRUCT(BlueprintType)
-struct GAMEPLAYABILITIES_API FGameplayEffectCustomExecutionOutput
+struct FGameplayEffectCustomExecutionOutput
 {
 	GENERATED_USTRUCT_BODY()
 
 public:
 
 	/** Constructor */
-	FGameplayEffectCustomExecutionOutput();
+	UE_API FGameplayEffectCustomExecutionOutput();
 
 	/** Mark that the execution has manually handled the stack count and the GE system should not attempt to automatically act upon it for emitted modifiers */
-	void MarkStackCountHandledManually();
+	UE_API void MarkStackCountHandledManually();
 
 	/** Simple accessor for determining whether the execution has manually handled the stack count or not */
-	bool IsStackCountHandledManually() const;
+	UE_API bool IsStackCountHandledManually() const;
 
 	/** Accessor for determining if GameplayCue events have already been handled */
-	bool AreGameplayCuesHandledManually() const;
+	UE_API bool AreGameplayCuesHandledManually() const;
 
 	/** Mark that the execution wants conditional gameplay effects to trigger */
-	void MarkConditionalGameplayEffectsToTrigger();
+	UE_API void MarkConditionalGameplayEffectsToTrigger();
 
 	/** Mark that the execution wants conditional gameplay effects to trigger */
-	void MarkGameplayCuesHandledManually();
+	UE_API void MarkGameplayCuesHandledManually();
 
 	/** Simple accessor for determining whether the execution wants conditional gameplay effects to trigger or not */
-	bool ShouldTriggerConditionalGameplayEffects() const;
+	UE_API bool ShouldTriggerConditionalGameplayEffects() const;
 
 	/** Add the specified evaluated data to the execution's output modifiers */
-	void AddOutputModifier(const FGameplayModifierEvaluatedData& InOutputMod);
+	UE_API void AddOutputModifier(const FGameplayModifierEvaluatedData& InOutputMod);
 
 	/** Simple accessor to output modifiers of the execution */
-	const TArray<FGameplayModifierEvaluatedData>& GetOutputModifiers() const;
+	UE_API const TArray<FGameplayModifierEvaluatedData>& GetOutputModifiers() const;
 
 	/** Simple accessor to output modifiers of the execution */
-	void GetOutputModifiers(OUT TArray<FGameplayModifierEvaluatedData>& OutOutputModifiers) const;
+	UE_API void GetOutputModifiers(OUT TArray<FGameplayModifierEvaluatedData>& OutOutputModifiers) const;
 
 	/** Returns direct access to output modifiers of the execution (avoid copy) */
 	TArray<FGameplayModifierEvaluatedData>& GetOutputModifiersRef() { return OutputModifiers; }
@@ -254,8 +256,8 @@ private:
 	uint32 bHandledGameplayCuesManually : 1;
 };
 
-UCLASS(BlueprintType, Blueprintable, Abstract)
-class GAMEPLAYABILITIES_API UGameplayEffectExecutionCalculation : public UGameplayEffectCalculation
+UCLASS(BlueprintType, Blueprintable, Abstract, MinimalAPI)
+class UGameplayEffectExecutionCalculation : public UGameplayEffectCalculation
 {
 	GENERATED_UCLASS_BODY()
 
@@ -283,17 +285,17 @@ public:
 	 * 
 	 * @param OutScopableModifiers	[OUT] Array to populate with definitions valid as scoped modifiers
 	 */
-	virtual void GetValidScopedModifierAttributeCaptureDefinitions(OUT TArray<FGameplayEffectAttributeCaptureDefinition>& OutScopableModifiers) const;
+	UE_API virtual void GetValidScopedModifierAttributeCaptureDefinitions(OUT TArray<FGameplayEffectAttributeCaptureDefinition>& OutScopableModifiers) const;
 
 	/**
 	 * Gets the collection of identifiers of valid transient aggregators ("temporary variable aggregators")
 	 * 
 	 * @return Array of valid transient aggregator identifiers
 	 */
-	virtual const FGameplayTagContainer& GetValidTransientAggregatorIdentifiers() const;
+	UE_API virtual const FGameplayTagContainer& GetValidTransientAggregatorIdentifiers() const;
 
 	/** Returns if this execution requires passed in tags */
-	virtual bool DoesRequirePassedInTags() const;
+	UE_API virtual bool DoesRequirePassedInTags() const;
 
 #endif // #if WITH_EDITORONLY_DATA
 
@@ -309,7 +311,7 @@ public:
 	 * @param OutExecutionOutput	[OUT] Output data populated by the execution detailing further behavior or results of the execution
 	 */
 	UFUNCTION(BlueprintNativeEvent, Category="Calculation")
-	void Execute(const FGameplayEffectCustomExecutionParameters& ExecutionParams, FGameplayEffectCustomExecutionOutput& OutExecutionOutput) const;
+	UE_API void Execute(const FGameplayEffectCustomExecutionParameters& ExecutionParams, FGameplayEffectCustomExecutionOutput& OutExecutionOutput) const;
 };
 
 
@@ -327,3 +329,5 @@ public:
 	P##Property = FindFieldChecked<FProperty>(S::StaticClass(), GET_MEMBER_NAME_CHECKED(S, P)); \
 	P##Def = FGameplayEffectAttributeCaptureDefinition(P##Property, EGameplayEffectAttributeCaptureSource::T, B); \
 }
+
+#undef UE_API
diff --git a/Engine/Plugins/Runtime/GameplayAbilities/Source/GameplayAbilities/Public/GameplayEffectTypes.h b/Engine/Plugins/Runtime/GameplayAbilities/Source/GameplayAbilities/Public/GameplayEffectTypes.h
--- a/Engine/Plugins/Runtime/GameplayAbilities/Source/GameplayAbilities/Public/GameplayEffectTypes.h
+++ b/Engine/Plugins/Runtime/GameplayAbilities/Source/GameplayAbilities/Public/GameplayEffectTypes.h
@@ -16,6 +16,8 @@
 #include "PhysicalMaterials/PhysicalMaterial.h"
 #include "GameplayEffectTypes.generated.h"
 
+#define UE_API GAMEPLAYABILITIES_API
+
 #define SKILL_SYSTEM_AGGREGATOR_DEBUG 1
 
 #if SKILL_SYSTEM_AGGREGATOR_DEBUG
@@ -68,12 +70,12 @@ enum class EGameplayModEvaluationChannel : uint8
 
 /** Struct representing evaluation channel settings for a gameplay modifier */
 USTRUCT()
-struct GAMEPLAYABILITIES_API FGameplayModEvaluationChannelSettings
+struct FGameplayModEvaluationChannelSettings
 {
 	GENERATED_USTRUCT_BODY()
 	
 	/** Constructor */
-	FGameplayModEvaluationChannelSettings();
+	UE_API FGameplayModEvaluationChannelSettings();
 
 	/**
 	 * Get the modifier evaluation channel to use
@@ -81,18 +83,18 @@ struct GAMEPLAYABILITIES_API FGameplayModEvaluationChannelSettings
 	 * @return	Either the channel directly specified within the settings, if valid, or Channel0 in the event of a game not using modifier
 	 *			channels or in the case of an invalid channel being specified within the settings
 	 */
-	EGameplayModEvaluationChannel GetEvaluationChannel() const;
+	UE_API EGameplayModEvaluationChannel GetEvaluationChannel() const;
 
 	/** Editor-only constants to aid in hiding evaluation channel settings when appropriate */
 #if WITH_EDITORONLY_DATA
-	static const FName ForceHideMetadataKey;
-	static const FString ForceHideMetadataEnabledValue;
+	static UE_API const FName ForceHideMetadataKey;
+	static UE_API const FString ForceHideMetadataEnabledValue;
 #endif // #if WITH_EDITORONLY_DATA
 
-	void SetEvaluationChannel(EGameplayModEvaluationChannel NewChannel);
+	UE_API void SetEvaluationChannel(EGameplayModEvaluationChannel NewChannel);
 
-	bool operator==(const FGameplayModEvaluationChannelSettings& Other) const;
-	bool operator!=(const FGameplayModEvaluationChannelSettings& Other) const;
+	UE_API bool operator==(const FGameplayModEvaluationChannelSettings& Other) const;
+	UE_API bool operator!=(const FGameplayModEvaluationChannelSettings& Other) const;
 
 protected:
 
@@ -240,7 +242,7 @@ struct FGameplayModifierEvaluatedData
  * It is passed throughout effect execution so it is a great place to track transient information about an execution
  */
 USTRUCT()
-struct GAMEPLAYABILITIES_API FGameplayEffectContext
+struct FGameplayEffectContext
 {
 	GENERATED_USTRUCT_BODY()
 
@@ -270,13 +272,13 @@ struct GAMEPLAYABILITIES_API FGameplayEffectContext
 	}
 
 	/** Returns the list of gameplay tags applicable to this effect, defaults to the owner's tags. SpecTagContainer remains untouched by default. */
-	virtual void GetOwnedGameplayTags(OUT FGameplayTagContainer& ActorTagContainer, OUT FGameplayTagContainer& SpecTagContainer) const;
+	UE_API virtual void GetOwnedGameplayTags(OUT FGameplayTagContainer& ActorTagContainer, OUT FGameplayTagContainer& SpecTagContainer) const;
 
 	/** Sets the instigator and effect causer. Instigator is who owns the ability that spawned this, EffectCauser is the actor that is the physical source of the effect, such as a weapon. They can be the same. */
-	virtual void AddInstigator(class AActor *InInstigator, class AActor *InEffectCauser);
+	UE_API virtual void AddInstigator(class AActor *InInstigator, class AActor *InEffectCauser);
 
 	/** Sets the ability that was used to spawn this */
-	virtual void SetAbility(const UGameplayAbility* InGameplayAbility);
+	UE_API virtual void SetAbility(const UGameplayAbility* InGameplayAbility);
 
 	/** Returns the immediate instigator that applied this effect */
 	virtual AActor* GetInstigator() const
@@ -285,10 +287,10 @@ struct GAMEPLAYABILITIES_API FGameplayEffectContext
 	}
 
 	/** Returns the CDO of the ability used to instigate this context */
-	const UGameplayAbility* GetAbility() const;
+	UE_API const UGameplayAbility* GetAbility() const;
 
 	/** Returns the specific instance that instigated this, may not always be set */
-	const UGameplayAbility* GetAbilityInstance_NotReplicated() const;
+	UE_API const UGameplayAbility* GetAbilityInstance_NotReplicated() const;
 
 	/** Gets the ability level this was evaluated at */
 	int32 GetAbilityLevel() const
@@ -341,10 +343,10 @@ struct GAMEPLAYABILITIES_API FGameplayEffectContext
 	}
 
 	/** Add actors to the stored actor list */
-	virtual void AddActors(const TArray<TWeakObjectPtr<AActor>>& IActor, bool bReset = false);
+	UE_API virtual void AddActors(const TArray<TWeakObjectPtr<AActor>>& IActor, bool bReset = false);
 
 	/** Add a hit result for targeting */
-	virtual void AddHitResult(const FHitResult& InHitResult, bool bReset = false);
+	UE_API virtual void AddHitResult(const FHitResult& InHitResult, bool bReset = false);
 
 	/** Returns actor list, may be empty */
 	virtual const TArray<TWeakObjectPtr<AActor>>& GetActors() const
@@ -365,7 +367,7 @@ struct GAMEPLAYABILITIES_API FGameplayEffectContext
 	}
 
 	/** Adds an origin point */
-	virtual void AddOrigin(FVector InOrigin);
+	UE_API virtual void AddOrigin(FVector InOrigin);
 
 	/** Returns origin point, may be invalid if HasOrigin is false */
 	virtual const FVector& GetOrigin() const
@@ -380,7 +382,7 @@ struct GAMEPLAYABILITIES_API FGameplayEffectContext
 	}
 
 	/** Returns debug string */
-	virtual FString ToString() const;
+	UE_API virtual FString ToString() const;
 
 	/** Returns the actual struct used for serialization, subclasses must override this! */
 	virtual UScriptStruct* GetScriptStruct() const
@@ -402,16 +404,16 @@ struct GAMEPLAYABILITIES_API FGameplayEffectContext
 	}
 
 	/** True if this was instigated by a locally controlled actor */
-	virtual bool IsLocallyControlled() const;
+	UE_API virtual bool IsLocallyControlled() const;
 
 	/** True if this was instigated by a locally controlled player */
-	virtual bool IsLocallyControlledPlayer() const;
+	UE_API virtual bool IsLocallyControlledPlayer() const;
 
 	/** Custom serialization, subclasses must override this */
-	virtual bool NetSerialize(FArchive& Ar, class UPackageMap* Map, bool& bOutSuccess);
+	UE_API virtual bool NetSerialize(FArchive& Ar, class UPackageMap* Map, bool& bOutSuccess);
 
 protected:
-	static bool CanActorReferenceBeReplicated(const AActor* Actor);
+	static UE_API bool CanActorReferenceBeReplicated(const AActor* Actor);
 
 	// The object pointers here have to be weak because contexts aren't necessarily tracked by GC in all cases
 
@@ -485,7 +487,7 @@ struct TStructOpsTypeTraits< FGameplayEffectContext > : public TStructOpsTypeTra
  * Handle that wraps a FGameplayEffectContext or subclass, to allow it to be polymorphic and replicate properly
  */
 USTRUCT(BlueprintType)
-struct GAMEPLAYABILITIES_API FGameplayEffectContextHandle
+struct FGameplayEffectContextHandle
 {
 	GENERATED_USTRUCT_BODY()
 
@@ -784,7 +786,7 @@ struct GAMEPLAYABILITIES_API FGameplayEffectContextHandle
 	}
 
 	/** Custom serializer, handles polymorphism of context */
-	bool NetSerialize(FArchive& Ar, class UPackageMap* Map, bool& bOutSuccess);
+	UE_API bool NetSerialize(FArchive& Ar, class UPackageMap* Map, bool& bOutSuccess);
 
 private:
 	friend UE::Net::FGameplayEffectContextHandleAccessorForNetSerializer;
@@ -829,7 +831,7 @@ struct FGameplayEffectRemovalInfo
 
 /** Metadata about a gameplay cue execution */
 USTRUCT(BlueprintType, meta = (HasNativeBreak = "/Script/GameplayAbilities.AbilitySystemBlueprintLibrary.BreakGameplayCueParameters", HasNativeMake = "/Script/GameplayAbilities.AbilitySystemBlueprintLibrary.MakeGameplayCueParameters"))
-struct GAMEPLAYABILITIES_API FGameplayCueParameters
+struct FGameplayCueParameters
 {
 	GENERATED_USTRUCT_BODY()
 
@@ -843,10 +845,10 @@ struct GAMEPLAYABILITIES_API FGameplayCueParameters
 	{}
 
 	/** Projects can override this via UAbilitySystemGlobals */
-	FGameplayCueParameters(const struct FGameplayEffectSpecForRPC &Spec);
-	FGameplayCueParameters(const struct FGameplayEffectContextHandle& EffectContext);
+	UE_API FGameplayCueParameters(const struct FGameplayEffectSpecForRPC &Spec);
+	UE_API FGameplayCueParameters(const struct FGameplayEffectContextHandle& EffectContext);
 
-	bool operator==(const FGameplayCueParameters& Other) const;
+	UE_API bool operator==(const FGameplayCueParameters& Other) const;
 	bool operator!=(const FGameplayCueParameters& Other) const
 	{
 		return !(*this == Other);
@@ -924,22 +926,22 @@ struct GAMEPLAYABILITIES_API FGameplayCueParameters
 	bool bGameplayEffectActive = true;
 
 	/** Optimized serializer */
-	bool NetSerialize(FArchive& Ar, class UPackageMap* Map, bool& bOutSuccess);
+	UE_API bool NetSerialize(FArchive& Ar, class UPackageMap* Map, bool& bOutSuccess);
 
 	/** Returns true if this is locally controlled, using fallback actor if nothing else available */
-	bool IsInstigatorLocallyControlled(AActor* FallbackActor = nullptr) const;
+	UE_API bool IsInstigatorLocallyControlled(AActor* FallbackActor = nullptr) const;
 
 	/** Fallback actor is used if the parameters have nullptr for instigator and effect causer */
-	bool IsInstigatorLocallyControlledPlayer(AActor* FallbackActor=nullptr) const;
+	UE_API bool IsInstigatorLocallyControlledPlayer(AActor* FallbackActor=nullptr) const;
 
 	/** Returns the actor that instigated this originally, generally attached to an ability system component */
-	AActor* GetInstigator() const;
+	UE_API AActor* GetInstigator() const;
 
 	/** Returns the actor that physically caused the damage, could be a projectile or weapon */
-	AActor* GetEffectCauser() const;
+	UE_API AActor* GetEffectCauser() const;
 
 	/** Returns the object that originally caused this, game-specific but usually not an actor */
-	const UObject* GetSourceObject() const;
+	UE_API const UObject* GetSourceObject() const;
 };
 
 template<>
@@ -1038,7 +1040,7 @@ namespace EGameplayTagEventType
  * while simultaneously tracking the count of parent tags as well. Events/delegates are fired whenever the tag counts
  * of any tag (explicit or parent) are modified.
  */
-struct GAMEPLAYABILITIES_API FGameplayTagCountContainer
+struct FGameplayTagCountContainer
 {	
 	FGameplayTagCountContainer()
 	{}
@@ -1245,7 +1247,7 @@ struct GAMEPLAYABILITIES_API FGameplayTagCountContainer
 	 *	Broadcasts the AnyChange event for this tag. This is called when the stack count of the backing gameplay effect change.
 	 *	It is up to the receiver of the broadcasted delegate to decide what to do with this.
 	 */
-	void Notify_StackCountChange(const FGameplayTag& Tag);
+	UE_API void Notify_StackCountChange(const FGameplayTag& Tag);
 
 	/**
 	 * Return delegate that can be bound to for when the specific tag's count changes to or off of zero
@@ -1254,7 +1256,7 @@ struct GAMEPLAYABILITIES_API FGameplayTagCountContainer
 	 * 
 	 * @return Delegate for when the specified tag's count changes to or off of zero
 	 */
-	FOnGameplayEffectTagCountChanged& RegisterGameplayTagEvent(const FGameplayTag& Tag, EGameplayTagEventType::Type EventType=EGameplayTagEventType::NewOrRemoved);
+	UE_API FOnGameplayEffectTagCountChanged& RegisterGameplayTagEvent(const FGameplayTag& Tag, EGameplayTagEventType::Type EventType=EGameplayTagEventType::NewOrRemoved);
 	
 	/**
 	 * Return delegate that can be bound to for when the any tag's count changes to or off of zero
@@ -1277,7 +1279,7 @@ struct GAMEPLAYABILITIES_API FGameplayTagCountContainer
 	 * 
 	 * @param bResetCallbacks	If true, also remove all of the registered tag count change delegates
 	 */
-	void Reset(bool bResetCallbacks = true);
+	UE_API void Reset(bool bResetCallbacks = true);
 
 	/** Fills in ParentTags from GameplayTags */
 	void FillParentTags()
@@ -1309,16 +1311,16 @@ private:
 	FGameplayTagContainer ExplicitTags;
 
 	/** Internal helper function to adjust the explicit tag list & corresponding maps/delegates/etc. as necessary */
-	bool UpdateTagMap_Internal(const FGameplayTag& Tag, int32 CountDelta);
+	UE_API bool UpdateTagMap_Internal(const FGameplayTag& Tag, int32 CountDelta);
 
 	/** Internal helper function to adjust the explicit tag list & corresponding maps/delegates/etc. as necessary. This does not call FillParentTags or any of the tag change delegates. These delegates are returned and must be executed by the caller. */
-	bool UpdateTagMapDeferredParentRemoval_Internal(const FGameplayTag& Tag, int32 CountDelta, TArray<FDeferredTagChangeDelegate>& DeferredTagChangeDelegates);
+	UE_API bool UpdateTagMapDeferredParentRemoval_Internal(const FGameplayTag& Tag, int32 CountDelta, TArray<FDeferredTagChangeDelegate>& DeferredTagChangeDelegates);
 
 	/** Internal helper function to adjust the explicit tag list & corresponding map. */
-	bool UpdateExplicitTags(const FGameplayTag& Tag, int32 CountDelta, bool bDeferParentTagsOnRemove);
+	UE_API bool UpdateExplicitTags(const FGameplayTag& Tag, int32 CountDelta, bool bDeferParentTagsOnRemove);
 
 	/** Internal helper function to collect the delegates that need to be called when Tag has its count changed by CountDelta. */
-	bool GatherTagChangeDelegates(const FGameplayTag& Tag, int32 CountDelta, TArray<FDeferredTagChangeDelegate>& TagChangeDelegates);
+	UE_API bool GatherTagChangeDelegates(const FGameplayTag& Tag, int32 CountDelta, TArray<FDeferredTagChangeDelegate>& TagChangeDelegates);
 };
 
 
@@ -1328,7 +1330,7 @@ private:
  * It only supports boolean, integer, and float properties.
  */
 USTRUCT()
-struct GAMEPLAYABILITIES_API FGameplayTagBlueprintPropertyMapping
+struct FGameplayTagBlueprintPropertyMapping
 {
 	GENERATED_BODY()
 
@@ -1375,39 +1377,39 @@ public:
  * to bind the delegate and it's address could change causing an invalid binding.
  */
 USTRUCT()
-struct GAMEPLAYABILITIES_API FGameplayTagBlueprintPropertyMap
+struct FGameplayTagBlueprintPropertyMap
 {
 	GENERATED_BODY()
 
 public:
 
-	FGameplayTagBlueprintPropertyMap();
-	FGameplayTagBlueprintPropertyMap(const FGameplayTagBlueprintPropertyMap& Other);
-	~FGameplayTagBlueprintPropertyMap();
+	UE_API FGameplayTagBlueprintPropertyMap();
+	UE_API FGameplayTagBlueprintPropertyMap(const FGameplayTagBlueprintPropertyMap& Other);
+	UE_API ~FGameplayTagBlueprintPropertyMap();
 
 	/** Call this to initialize and bind the properties with the ability system component. */
-	void Initialize(UObject* Owner, class UAbilitySystemComponent* ASC);
+	UE_API void Initialize(UObject* Owner, class UAbilitySystemComponent* ASC);
 
 	/** Call to manually apply the current tag state, can handle cases where callbacks were skipped */
-	void ApplyCurrentTags();
+	UE_API void ApplyCurrentTags();
 
 #if WITH_EDITOR
 	UE_DEPRECATED(5.3, "The signature for IsDataValid has changed.  Please use the one that takes a FDataValidationContext")
 	EDataValidationResult IsDataValid(const UObject* ContainingAsset, TArray<FText>& ValidationErrors) { return EDataValidationResult::NotValidated; }
 
 	/** This can optionally be called in the owner's IsDataValid() for data validation. */
-	EDataValidationResult IsDataValid(const UObject* ContainingAsset, class FDataValidationContext& Context) const;
+	UE_API EDataValidationResult IsDataValid(const UObject* ContainingAsset, class FDataValidationContext& Context) const;
 #endif // #if WITH_EDITOR
 
 protected:
 
-	void Unregister();
+	UE_API void Unregister();
 
-	void GameplayTagEventCallback(const FGameplayTag Tag, int32 NewCount, TWeakObjectPtr<UObject> RegisteredOwner);
+	UE_API void GameplayTagEventCallback(const FGameplayTag Tag, int32 NewCount, TWeakObjectPtr<UObject> RegisteredOwner);
 
-	bool IsPropertyTypeValid(const FProperty* Property) const;
+	UE_API bool IsPropertyTypeValid(const FProperty* Property) const;
 
-	EGameplayTagEventType::Type GetGameplayTagEventType(const FProperty* Property) const;
+	UE_API EGameplayTagEventType::Type GetGameplayTagEventType(const FProperty* Property) const;
 
 protected:
 
@@ -1421,7 +1423,7 @@ protected:
 
 /** Encapsulate require and ignore tags */
 USTRUCT(BlueprintType)
-struct GAMEPLAYABILITIES_API FGameplayTagRequirements
+struct FGameplayTagRequirements
 {
 	GENERATED_USTRUCT_BODY()
 
@@ -1438,25 +1440,25 @@ struct GAMEPLAYABILITIES_API FGameplayTagRequirements
 	FGameplayTagQuery TagQuery;
 
 	/** True if all required tags and no ignore tags found */
-	bool	RequirementsMet(const FGameplayTagContainer& Container) const;
+	UE_API bool	RequirementsMet(const FGameplayTagContainer& Container) const;
 
 	/** True if neither RequireTags or IgnoreTags has any tags */
-	bool	IsEmpty() const;
+	UE_API bool	IsEmpty() const;
 
 	/** Return debug string */
-	FString ToString() const;
+	UE_API FString ToString() const;
 
-	bool operator==(const FGameplayTagRequirements& Other) const;
-	bool operator!=(const FGameplayTagRequirements& Other) const;
+	UE_API bool operator==(const FGameplayTagRequirements& Other) const;
+	UE_API bool operator!=(const FGameplayTagRequirements& Other) const;
 
 	/** Converts the RequireTags and IgnoreTags fields into an equivalent FGameplayTagQuery */
-	[[nodiscard]] FGameplayTagQuery ConvertTagFieldsToTagQuery() const;
+	[[nodiscard]] UE_API FGameplayTagQuery ConvertTagFieldsToTagQuery() const;
 };
 
 
 /** Structure used to combine tags from different sources during effect execution */
 USTRUCT()
-struct GAMEPLAYABILITIES_API FTagContainerAggregator
+struct FTagContainerAggregator
 {
 	GENERATED_USTRUCT_BODY()
 
@@ -1497,15 +1499,15 @@ struct GAMEPLAYABILITIES_API FTagContainerAggregator
 	}
 
 	/** Returns tags from the source or target actor */
-	FGameplayTagContainer& GetActorTags();
-	const FGameplayTagContainer& GetActorTags() const;
+	UE_API FGameplayTagContainer& GetActorTags();
+	UE_API const FGameplayTagContainer& GetActorTags() const;
 
 	/** Get tags that came from the effect spec */
-	FGameplayTagContainer& GetSpecTags();
-	const FGameplayTagContainer& GetSpecTags() const;
+	UE_API FGameplayTagContainer& GetSpecTags();
+	UE_API const FGameplayTagContainer& GetSpecTags() const;
 
 	/** Returns combination of spec and actor tags */
-	const FGameplayTagContainer* GetAggregatedTags() const;
+	UE_API const FGameplayTagContainer* GetAggregatedTags() const;
 
 private:
 
@@ -1526,12 +1528,12 @@ private:
 
 /** Allows blueprints to generate a GameplayEffectSpec once and then reference it by handle, to apply it multiple times/multiple targets. */
 USTRUCT(BlueprintType)
-struct GAMEPLAYABILITIES_API FGameplayEffectSpecHandle
+struct FGameplayEffectSpecHandle
 {
 	GENERATED_USTRUCT_BODY()
 
-	FGameplayEffectSpecHandle();
-	FGameplayEffectSpecHandle(FGameplayEffectSpec* DataPtr);
+	UE_API FGameplayEffectSpecHandle();
+	UE_API FGameplayEffectSpecHandle(FGameplayEffectSpec* DataPtr);
 
 	/** Internal pointer to effect spec */
 	TSharedPtr<FGameplayEffectSpec>	Data;
@@ -1546,7 +1548,7 @@ struct GAMEPLAYABILITIES_API FGameplayEffectSpecHandle
 		return Data.IsValid();
 	}
 
-	bool NetSerialize(FArchive& Ar, class UPackageMap* Map, bool& bOutSuccess);
+	UE_API bool NetSerialize(FArchive& Ar, class UPackageMap* Map, bool& bOutSuccess);
 
 	/** Comparison operator */
 	bool operator==(FGameplayEffectSpecHandle const& Other) const
@@ -1578,7 +1580,7 @@ struct TStructOpsTypeTraits<FGameplayEffectSpecHandle> : public TStructOpsTypeTr
 
 /** Map that stores count of tags, in a form that is optimized for replication */
 USTRUCT()
-struct GAMEPLAYABILITIES_API FMinimalReplicationTagCountMap
+struct FMinimalReplicationTagCountMap
 {
 	GENERATED_USTRUCT_BODY()
 
@@ -1599,7 +1601,7 @@ struct GAMEPLAYABILITIES_API FMinimalReplicationTagCountMap
 		TagMap.FindOrAdd(Tag)++;
 	}
 
-	void RemoveTag(const FGameplayTag& Tag);
+	UE_API void RemoveTag(const FGameplayTag& Tag);
 
 	void AddTags(const FGameplayTagContainer& Container)
 	{
@@ -1635,7 +1637,7 @@ struct GAMEPLAYABILITIES_API FMinimalReplicationTagCountMap
 		}
 	}
 
-	bool NetSerialize(FArchive& Ar, class UPackageMap* Map, bool& bOutSuccess);
+	UE_API bool NetSerialize(FArchive& Ar, class UPackageMap* Map, bool& bOutSuccess);
 
 	TMap<FGameplayTag, int32>	TagMap;
 
@@ -1659,10 +1661,10 @@ struct GAMEPLAYABILITIES_API FMinimalReplicationTagCountMap
 	/** If true, we will skip updating the Owner ASC if we replicate on a connection owned by the ASC */
 	void SetRequireNonOwningNetConnection(bool b) { bRequireNonOwningNetConnection = b; }
 
-	void SetOwner(UAbilitySystemComponent* ASC);
+	UE_API void SetOwner(UAbilitySystemComponent* ASC);
 
 	// Removes all tags that this container is granting
-	void RemoveAllTags();
+	UE_API void RemoveAllTags();
 
 private:
 	friend FMinimalReplicationTagCountMapForNetSerializer;
@@ -1671,7 +1673,7 @@ private:
 	bool bRequireNonOwningNetConnection = false;
 	TWeakObjectPtr<UNetConnection> LastConnection;
 
-	void UpdateOwnerTagMap();
+	UE_API void UpdateOwnerTagMap();
 };
 
 template<>
@@ -1686,3 +1688,5 @@ struct TStructOpsTypeTraits<FMinimalReplicationTagCountMap> : public TStructOpsT
 };
 
 DECLARE_MULTICAST_DELEGATE(FOnExternalGameplayModifierDependencyChange);
+
+#undef UE_API
diff --git a/Engine/Plugins/Runtime/GameplayAbilities/Source/GameplayAbilities/Public/GameplayEffectUIData.h b/Engine/Plugins/Runtime/GameplayAbilities/Source/GameplayAbilities/Public/GameplayEffectUIData.h
--- a/Engine/Plugins/Runtime/GameplayAbilities/Source/GameplayAbilities/Public/GameplayEffectUIData.h
+++ b/Engine/Plugins/Runtime/GameplayAbilities/Source/GameplayAbilities/Public/GameplayEffectUIData.h
@@ -5,13 +5,17 @@
 #include "GameplayEffectComponent.h"
 #include "GameplayEffectUIData.generated.h"
 
+#define UE_API GAMEPLAYABILITIES_API
+
 /**
  * UGameplayEffectUIData
  * Base class to provide game-specific data about how to describe a Gameplay Effect in the UI. Subclass with data to use in your game.
  * In Unreal Engine 5.3, this now derives from UGameplayEffectComponent so you can use it directly as a GameplayEffectComponent.
  */
-UCLASS(Blueprintable, Abstract, EditInlineNew, CollapseCategories)
-class GAMEPLAYABILITIES_API UGameplayEffectUIData : public UGameplayEffectComponent
+UCLASS(Blueprintable, Abstract, EditInlineNew, CollapseCategories, MinimalAPI)
+class UGameplayEffectUIData : public UGameplayEffectComponent
 {
 	GENERATED_BODY()
 };
+
+#undef UE_API
diff --git a/Engine/Plugins/Runtime/GameplayAbilities/Source/GameplayAbilities/Public/GameplayEffectUIData_TextOnly.h b/Engine/Plugins/Runtime/GameplayAbilities/Source/GameplayAbilities/Public/GameplayEffectUIData_TextOnly.h
--- a/Engine/Plugins/Runtime/GameplayAbilities/Source/GameplayAbilities/Public/GameplayEffectUIData_TextOnly.h
+++ b/Engine/Plugins/Runtime/GameplayAbilities/Source/GameplayAbilities/Public/GameplayEffectUIData_TextOnly.h
@@ -5,12 +5,14 @@
 #include "GameplayEffectUIData.h"
 #include "GameplayEffectUIData_TextOnly.generated.h"
 
+#define UE_API GAMEPLAYABILITIES_API
+
 /**
  * UI data that contains only text. This is mostly used as an example of a subclass of UGameplayEffectUIData.
  * If your game needs only text, this is a reasonable class to use. To include more data, make a custom subclass of UGameplayEffectUIData.
  */
-UCLASS(DisplayName="UI Data (Text Only)")
-class GAMEPLAYABILITIES_API UGameplayEffectUIData_TextOnly : public UGameplayEffectUIData
+UCLASS(DisplayName="UI Data (Text Only)", MinimalAPI)
+class UGameplayEffectUIData_TextOnly : public UGameplayEffectUIData
 {
 	GENERATED_BODY()
 	
@@ -18,3 +20,5 @@ public:
 	UPROPERTY(EditDefaultsOnly, BlueprintReadOnly, Category = Data, meta = (MultiLine = "true"))
 	FText Description;
 };
+
+#undef UE_API
