[GameplayAbilities]

diff --git a/Engine/Plugins/Runtime/GameplayAbilities/Source/GameplayAbilities/Public/GameplayModMagnitudeCalculation.h b/Engine/Plugins/Runtime/GameplayAbilities/Source/GameplayAbilities/Public/GameplayModMagnitudeCalculation.h
--- a/Engine/Plugins/Runtime/GameplayAbilities/Source/GameplayAbilities/Public/GameplayModMagnitudeCalculation.h
+++ b/Engine/Plugins/Runtime/GameplayAbilities/Source/GameplayAbilities/Public/GameplayModMagnitudeCalculation.h
@@ -8,9 +8,11 @@
 #include "GameplayEffect.h"
 #include "GameplayModMagnitudeCalculation.generated.h"
 
+#define UE_API GAMEPLAYABILITIES_API
+
 /** Class used to perform custom gameplay effect modifier calculations, either via blueprint or native code */ 
-UCLASS(BlueprintType, Blueprintable, Abstract)
-class GAMEPLAYABILITIES_API UGameplayModMagnitudeCalculation : public UGameplayEffectCalculation
+UCLASS(BlueprintType, Blueprintable, Abstract, MinimalAPI)
+class UGameplayModMagnitudeCalculation : public UGameplayEffectCalculation
 {
 
 public:
@@ -25,7 +27,7 @@ public:
 	 * @return Computed magnitude of the modifier
 	 */
 	UFUNCTION(BlueprintNativeEvent, Category="Calculation")
-	float CalculateBaseMagnitude(const FGameplayEffectSpec& Spec) const;
+	UE_API float CalculateBaseMagnitude(const FGameplayEffectSpec& Spec) const;
 
 	/**
 	 * If the magnitude resultant from the custom calculation depends on game code-specific conditions that are not under the purview of the ability system,
@@ -37,15 +39,15 @@ public:
 	 * 
 	 * @return Multicast delegate that will fire when this calculation's external dependencies change, if any
 	 */
-	virtual FOnExternalGameplayModifierDependencyChange* GetExternalModifierDependencyMulticast(const FGameplayEffectSpec& Spec, UWorld* World) const;
+	UE_API virtual FOnExternalGameplayModifierDependencyChange* GetExternalModifierDependencyMulticast(const FGameplayEffectSpec& Spec, UWorld* World) const;
 
 	/** Simple accessor to bAllowNonNetAuthorityDependencyRegistration with some validation: Read the comment on that variable for usage!!! */
-	bool ShouldAllowNonNetAuthorityDependencyRegistration() const;
+	UE_API bool ShouldAllowNonNetAuthorityDependencyRegistration() const;
 
 protected:
 	
 	/** Convenience method to get attribute magnitude during a CalculateMagnitude call */
-	bool GetCapturedAttributeMagnitude(const FGameplayEffectAttributeCaptureDefinition& Def, const FGameplayEffectSpec& Spec, const FAggregatorEvaluateParameters& EvaluationParameters, OUT float& Magnitude) const;
+	UE_API bool GetCapturedAttributeMagnitude(const FGameplayEffectAttributeCaptureDefinition& Def, const FGameplayEffectSpec& Spec, const FAggregatorEvaluateParameters& EvaluationParameters, OUT float& Magnitude) const;
 
 	/** 
 	 * Whether the calculation allows non-net authorities to register the external dependency multi-cast delegate or not; Effectively
@@ -74,7 +76,7 @@ protected:
 	 * @return The magnitude value if found, zero otherwise
 	 */
 	UFUNCTION(BlueprintCallable, BlueprintPure = false, Category = "Ability|GameplayEffect", meta=(DisplayName="Get Captured Attribute Magnitude", ScriptName="GetCapturedAttributeMagnitude"))
-	float K2_GetCapturedAttributeMagnitude(const FGameplayEffectSpec& EffectSpec, FGameplayAttribute Attribute, const FGameplayTagContainer& SourceTags, const FGameplayTagContainer& TargetTags) const;
+	UE_API float K2_GetCapturedAttributeMagnitude(const FGameplayEffectSpec& EffectSpec, FGameplayAttribute Attribute, const FGameplayTagContainer& SourceTags, const FGameplayTagContainer& TargetTags) const;
 
 	/**
 	 * Extracts the Set by Caller Magnitude from a Gameplay Effect Spec
@@ -85,7 +87,7 @@ protected:
 	 * @return The magnitude value if found, zero otherwise
 	 */
 	UFUNCTION(BlueprintCallable, BlueprintPure = false, Category = "Ability|GameplayEffect")
-	float GetSetByCallerMagnitudeByTag(const FGameplayEffectSpec& EffectSpec, const FGameplayTag& Tag) const;
+	UE_API float GetSetByCallerMagnitudeByTag(const FGameplayEffectSpec& EffectSpec, const FGameplayTag& Tag) const;
 
 	/**
 	 * Extracts the Set by Caller Magnitude from a Gameplay Effect Spec
@@ -96,7 +98,7 @@ protected:
 	 * @return The magnitude value if found, zero otherwise
 	 */
 	UFUNCTION(BlueprintCallable, BlueprintPure = false, Category = "Ability|GameplayEffect")
-	float GetSetByCallerMagnitudeByName(const FGameplayEffectSpec& EffectSpec, const FName& MagnitudeName) const;
+	UE_API float GetSetByCallerMagnitudeByName(const FGameplayEffectSpec& EffectSpec, const FName& MagnitudeName) const;
 
 	/**
 	 * Copies and returns the source aggregated tags from a Gameplay Effect Spec
@@ -106,7 +108,7 @@ protected:
 	 * @return Gameplay Tag Container with the copied tags. The container will be empty if no captured source tags exist.
 	 */
 	UFUNCTION(BlueprintCallable, BlueprintPure = false, Category = "Ability|GameplayEffect")
-	FGameplayTagContainer GetSourceAggregatedTags(const FGameplayEffectSpec& EffectSpec) const;
+	UE_API FGameplayTagContainer GetSourceAggregatedTags(const FGameplayEffectSpec& EffectSpec) const;
 
 	/**
 	 * Returns the source actor tags from a Gameplay Effect Spec
@@ -116,7 +118,7 @@ protected:
 	 * @return Gameplay Tag Container with the copied tags. The container will be empty if no captured source tags exist.
 	 */
 	UFUNCTION(BlueprintCallable, Category = "Ability|GameplayEffect")
-	const FGameplayTagContainer& GetSourceActorTags(const FGameplayEffectSpec& EffectSpec) const;
+	UE_API const FGameplayTagContainer& GetSourceActorTags(const FGameplayEffectSpec& EffectSpec) const;
 
 	/**
 	 * Returns the source spec tags from a Gameplay Effect Spec
@@ -126,7 +128,7 @@ protected:
 	 * @return Gameplay Tag Container with the copied tags. The container will be empty if no captured source tags exist.
 	 */
 	UFUNCTION(BlueprintCallable, Category = "Ability|GameplayEffect")
-	const FGameplayTagContainer& GetSourceSpecTags(const FGameplayEffectSpec& EffectSpec) const;
+	UE_API const FGameplayTagContainer& GetSourceSpecTags(const FGameplayEffectSpec& EffectSpec) const;
 
 	/**
 	 * Copies and returns the target aggregated tags from a Gameplay Effect Spec
@@ -136,7 +138,7 @@ protected:
 	 * @return Gameplay Tag Container with the copied tags. The container will be empty if no captured source tags exist.
 	 */
 	UFUNCTION(BlueprintCallable, BlueprintPure = false, Category = "Ability|GameplayEffect")
-	FGameplayTagContainer GetTargetAggregatedTags(const FGameplayEffectSpec& EffectSpec) const;
+	UE_API FGameplayTagContainer GetTargetAggregatedTags(const FGameplayEffectSpec& EffectSpec) const;
 
 	/**
 	 * Returns the target actor tags from a Gameplay Effect Spec
@@ -147,7 +149,7 @@ protected:
 	 * @return Gameplay Tag Container with the copied tags. The container will be empty if no captured source tags exist.
 	 */
 	UFUNCTION(BlueprintCallable, Category = "Ability|GameplayEffect")
-	const FGameplayTagContainer& GetTargetActorTags(const FGameplayEffectSpec& EffectSpec) const;
+	UE_API const FGameplayTagContainer& GetTargetActorTags(const FGameplayEffectSpec& EffectSpec) const;
 
 	/**
 	 * Returns the target spec tags from a Gameplay Effect Spec
@@ -158,5 +160,7 @@ protected:
 	 * @return Gameplay Tag Container with the copied tags. The container will be empty if no captured source tags exist.
 	 */
 	UFUNCTION(BlueprintCallable, Category = "Ability|GameplayEffect")
-	const FGameplayTagContainer& GetTargetSpecTags(const FGameplayEffectSpec& EffectSpec) const;
+	UE_API const FGameplayTagContainer& GetTargetSpecTags(const FGameplayEffectSpec& EffectSpec) const;
 };
+
+#undef UE_API
diff --git a/Engine/Plugins/Runtime/GameplayAbilities/Source/GameplayAbilities/Public/GameplayPrediction.h b/Engine/Plugins/Runtime/GameplayAbilities/Source/GameplayAbilities/Public/GameplayPrediction.h
--- a/Engine/Plugins/Runtime/GameplayAbilities/Source/GameplayAbilities/Public/GameplayPrediction.h
+++ b/Engine/Plugins/Runtime/GameplayAbilities/Source/GameplayAbilities/Public/GameplayPrediction.h
@@ -9,6 +9,8 @@
 #include "Templates/TypeCompatibleBytes.h"
 #include "GameplayPrediction.generated.h"
 
+#define UE_API GAMEPLAYABILITIES_API
+
 class UAbilitySystemComponent;
 namespace UE::Net
 {
@@ -291,7 +293,7 @@ DECLARE_DELEGATE(FPredictionKeyEvent);
 
 PRAGMA_DISABLE_DEPRECATION_WARNINGS // PredictiveConnection
 USTRUCT()
-struct GAMEPLAYABILITIES_API FPredictionKey
+struct FPredictionKey
 {
 	GENERATED_USTRUCT_BODY()
 
@@ -311,24 +313,24 @@ struct GAMEPLAYABILITIES_API FPredictionKey
 	bool bIsServerInitiated = false;
 
 	/** Construct a new prediction key with no dependencies */
-	static FPredictionKey CreateNewPredictionKey(const UAbilitySystemComponent*);
+	static UE_API FPredictionKey CreateNewPredictionKey(const UAbilitySystemComponent*);
 
 	/** Construct a new server initiation key, for abilities activated on the server */
-	static FPredictionKey CreateNewServerInitiatedKey(const UAbilitySystemComponent*);
+	static UE_API FPredictionKey CreateNewServerInitiatedKey(const UAbilitySystemComponent*);
 
 	/** Create a new dependent prediction key: keep our existing base or use the current key as the base. */
-	void GenerateDependentPredictionKey();
+	UE_API void GenerateDependentPredictionKey();
 
 	/** Creates new delegate called only when this key is rejected. */
-	FPredictionKeyEvent& NewRejectedDelegate();
+	UE_API FPredictionKeyEvent& NewRejectedDelegate();
 
 	/** Creates new delegate called only when replicated state catches up to this key. */
-	FPredictionKeyEvent& NewCaughtUpDelegate();
+	UE_API FPredictionKeyEvent& NewCaughtUpDelegate();
 	
 	/** Add a new delegate that is called if the key is rejected or caught up to. */
-	void NewRejectOrCaughtUpDelegate(FPredictionKeyEvent Event);
+	UE_API void NewRejectOrCaughtUpDelegate(FPredictionKeyEvent Event);
 	
-	bool NetSerialize(FArchive& Ar, class UPackageMap* Map, bool& bOutSuccess);
+	UE_API bool NetSerialize(FArchive& Ar, class UPackageMap* Map, bool& bOutSuccess);
 
 	/** A key is valid if it's non-zero, prediction keys for other clients will serialize down as 0 and be invalid */
 	bool IsValidKey() const
@@ -399,7 +401,7 @@ struct GAMEPLAYABILITIES_API FPredictionKey
 private:
 	friend UE::Net::FPredictionKeyNetSerializer;
 
-	void GenerateNewPredictionKey();
+	UE_API void GenerateNewPredictionKey();
 
 	explicit FPredictionKey(int32 Key)
 		: Current(static_cast<KeyType>(Key))
@@ -474,18 +476,18 @@ class UGameplayAbility;
  *	A structure for allowing scoped prediction windows.
  */
 
-struct GAMEPLAYABILITIES_API FScopedPredictionWindow
+struct FScopedPredictionWindow
 {
 	/** To be called on server when a new prediction key is received from the client (In an RPC). 
 	 *	InSetReplicatedPredictionKey should be set to false in cases where we want a scoped prediction key but have already repped the prediction key.
 	 *	(For example, cached target data will restore the prediction key that the TD was sent with, but this key was already repped down as confirmed when received)
 	 **/
-	FScopedPredictionWindow(UAbilitySystemComponent* AbilitySystemComponent, FPredictionKey InPredictionKey, bool InSetReplicatedPredictionKey = true);
+	UE_API FScopedPredictionWindow(UAbilitySystemComponent* AbilitySystemComponent, FPredictionKey InPredictionKey, bool InSetReplicatedPredictionKey = true);
 
 	/** To be called in the callsite where the predictive code will take place. This generates a new PredictionKey and acts as a synchonization point between client and server for that key.  */
-	FScopedPredictionWindow(UAbilitySystemComponent* AbilitySystemComponent, bool CanGenerateNewKey=true);
+	UE_API FScopedPredictionWindow(UAbilitySystemComponent* AbilitySystemComponent, bool CanGenerateNewKey=true);
 
-	~FScopedPredictionWindow();
+	UE_API ~FScopedPredictionWindow();
 
 private:
 
@@ -522,13 +524,13 @@ enum class EGasPredictionKeyResult : uint8
  * For example, we use this in a case where we want to play a Montage locally, and we notify the server unreliably (i.e. it is an optional event).
  * Note: Nothing stops something within this scope from creating a valid key, then sending that valid key to the server.  This is meant for generated keys that aren't already going to be sent.
  */
-struct GAMEPLAYABILITIES_API FScopedDiscardPredictions
+struct FScopedDiscardPredictions
 {
 	// Pass in the AbilitySystemComponent to discard predictions on.  Optionally, if we want to do something other than just Drop the prediction events you can override that.
-	explicit FScopedDiscardPredictions(UAbilitySystemComponent* AbilitySystemComponent, EGasPredictionKeyResult HowToHandlePredictions = EGasPredictionKeyResult::SilentlyDrop);
+	UE_API explicit FScopedDiscardPredictions(UAbilitySystemComponent* AbilitySystemComponent, EGasPredictionKeyResult HowToHandlePredictions = EGasPredictionKeyResult::SilentlyDrop);
 
 	// When the scope ends, we will perform the final drop/ack/nack of the prediction chain
-	~FScopedDiscardPredictions();
+	UE_API ~FScopedDiscardPredictions();
 
 private:
 	// Weak pointer to the 'Owning' ASC.  If this ends up being invalid/null we don't do anything.
@@ -617,4 +619,6 @@ struct TStructOpsTypeTraits< FReplicatedPredictionKeyMap > : public TStructOpsTy
 	{
 		WithNetDeltaSerializer = true,
 	};
-};
\ No newline at end of file
+};
+
+#undef UE_API
diff --git a/Engine/Plugins/Runtime/GameplayAbilities/Source/GameplayAbilities/Public/GameplayTagResponseTable.h b/Engine/Plugins/Runtime/GameplayAbilities/Source/GameplayAbilities/Public/GameplayTagResponseTable.h
--- a/Engine/Plugins/Runtime/GameplayAbilities/Source/GameplayAbilities/Public/GameplayTagResponseTable.h
+++ b/Engine/Plugins/Runtime/GameplayAbilities/Source/GameplayAbilities/Public/GameplayTagResponseTable.h
@@ -11,6 +11,8 @@
 #include "GameplayEffect.h"
 #include "GameplayTagResponseTable.generated.h"
 
+#define UE_API GAMEPLAYABILITIES_API
+
 class UAbilitySystemComponent;
 
 USTRUCT()
@@ -56,8 +58,8 @@ struct FGameplayTagResponseTableEntry
  *	For example, "for every count of "Status.Haste" I get 1 level of GE_Response_Haste. This class facilitates
  *	building this table and automatically registering and responding to tag events on the ability system component.
  */
-UCLASS()
-class GAMEPLAYABILITIES_API UGameplayTagReponseTable : public UDataAsset
+UCLASS(MinimalAPI)
+class UGameplayTagReponseTable : public UDataAsset
 {
 	GENERATED_UCLASS_BODY()
 
@@ -67,14 +69,14 @@ class GAMEPLAYABILITIES_API UGameplayTagReponseTable : public UDataAsset
 	/** Registers for tag events for the given ability system component. Note this will happen to every spawned ASC, we may want to allow games
 	 *	to limit what classe this is called on, or potentially build into the table class restrictions for each response entry.
 	 */
-	void RegisterResponseForEvents(UAbilitySystemComponent* ASC);
+	UE_API void RegisterResponseForEvents(UAbilitySystemComponent* ASC);
 
-	virtual void PostLoad() override;
+	UE_API virtual void PostLoad() override;
 
 protected:
 
 	UFUNCTION()
-	void TagResponseEvent(const FGameplayTag Tag, int32 NewCount, UAbilitySystemComponent* ASC, int32 idx);
+	UE_API void TagResponseEvent(const FGameplayTag Tag, int32 NewCount, UAbilitySystemComponent* ASC, int32 idx);
 	
 	/** Temporary structs to avoid extra heap allocations every time we recalculate tag count */
 	mutable FGameplayEffectQuery Query;
@@ -97,9 +99,11 @@ protected:
 
 	float LastASCPurgeTime;
 
-	void Remove(UAbilitySystemComponent* ASC, TArray<FActiveGameplayEffectHandle>& Handles);
+	UE_API void Remove(UAbilitySystemComponent* ASC, TArray<FActiveGameplayEffectHandle>& Handles);
 
-	void AddOrUpdate(UAbilitySystemComponent* ASC, const TArray<TSubclassOf<UGameplayEffect> >& ResponseGameplayEffects, int32 TotalCount, TArray<FActiveGameplayEffectHandle>& Handles);
+	UE_API void AddOrUpdate(UAbilitySystemComponent* ASC, const TArray<TSubclassOf<UGameplayEffect> >& ResponseGameplayEffects, int32 TotalCount, TArray<FActiveGameplayEffectHandle>& Handles);
 
-	int32 GetCount(const FGameplayTagReponsePair& Pair, UAbilitySystemComponent* ASC) const;
+	UE_API int32 GetCount(const FGameplayTagReponsePair& Pair, UAbilitySystemComponent* ASC) const;
 };
+
+#undef UE_API
diff --git a/Engine/Plugins/Runtime/GameplayAbilities/Source/GameplayAbilities/Public/ScalableFloat.h b/Engine/Plugins/Runtime/GameplayAbilities/Source/GameplayAbilities/Public/ScalableFloat.h
--- a/Engine/Plugins/Runtime/GameplayAbilities/Source/GameplayAbilities/Public/ScalableFloat.h
+++ b/Engine/Plugins/Runtime/GameplayAbilities/Source/GameplayAbilities/Public/ScalableFloat.h
@@ -7,9 +7,11 @@
 #include "DataRegistryId.h"
 #include "ScalableFloat.generated.h"
 
+#define UE_API GAMEPLAYABILITIES_API
+
 /** Generic numerical value in the form Value * Curve[Level] */
 USTRUCT(BlueprintType)
-struct GAMEPLAYABILITIES_API FScalableFloat
+struct FScalableFloat
 {
 	GENERATED_USTRUCT_BODY()
 
@@ -50,19 +52,19 @@ public:
 	FDataRegistryType		RegistryType;
 
 	/** Returns the scaled value at a given level */
-	float GetValueAtLevel(float Level, const FString* ContextString = nullptr) const;
+	UE_API float GetValueAtLevel(float Level, const FString* ContextString = nullptr) const;
 
 	/** Returns the scaled value at level 0 */
-	float GetValue(const FString* ContextString = nullptr) const;
+	UE_API float GetValue(const FString* ContextString = nullptr) const;
 
 	/** Used when using a scalable float as a bool */
-	bool AsBool(float Level = 0, const FString* ContextString = nullptr) const;
+	UE_API bool AsBool(float Level = 0, const FString* ContextString = nullptr) const;
 
 	/** Returns the value as an int32 */
-	int32 AsInteger(float Level = 0, const FString* ContextString = nullptr) const;
+	UE_API int32 AsInteger(float Level = 0, const FString* ContextString = nullptr) const;
 
 	/** Gets the value and possible curve at a given level, returning false if it failed to find a good value */
-	bool EvaluateCurveAtLevel(float& OutValue, const FRealCurve*& OutCurve, float Level, const FString& ContextString, bool bWarnIfInvalid = true) const;
+	UE_API bool EvaluateCurveAtLevel(float& OutValue, const FRealCurve*& OutCurve, float Level, const FString& ContextString, bool bWarnIfInvalid = true) const;
 
 	/** True if there is no curve lookup */
 	bool IsStatic() const
@@ -71,10 +73,10 @@ public:
 	}
 
 	/** Sets raw value */
-	void SetValue(float NewValue);
+	UE_API void SetValue(float NewValue);
 
 	/** Overrides raw value and curve reference */
-	void SetScalingValue(float InCoeffecient, FName InRowName, UCurveTable * InTable);
+	UE_API void SetScalingValue(float InCoeffecient, FName InRowName, UCurveTable * InTable);
 
 	/** Returns static value, only safe if this has no curve reference */
 	float GetValueChecked() const
@@ -84,29 +86,29 @@ public:
 	}
 
 	/** Outputs human readable string */
-	FString ToSimpleString() const;
+	UE_API FString ToSimpleString() const;
 
 	/** Error checking: Returns false if this has an invalid curve reference but will not print warnings */
-	bool IsValid() const;
+	UE_API bool IsValid() const;
 
 	/** Error checking: Checks if float is valid and prints detailed warnings if not valid */
 	UE_DEPRECATED(5.6, "Use IsDataValid and manually output the errors")
-	bool IsValidWithWarnings(const FString& ContextString) const;
+	UE_API bool IsValidWithWarnings(const FString& ContextString) const;
 
 #if WITH_EDITOR
 	/** Error checking must be performed by the owning class, then call this to get specific errors */
-	EDataValidationResult IsDataValid(class FDataValidationContext& Context, const FString& PathName) const;
+	UE_API EDataValidationResult IsDataValid(class FDataValidationContext& Context, const FString& PathName) const;
 #endif
 
 	/** Equality/Inequality operators */
-	bool operator==(const FScalableFloat& Other) const;
-	bool operator!=(const FScalableFloat& Other) const;
+	UE_API bool operator==(const FScalableFloat& Other) const;
+	UE_API bool operator!=(const FScalableFloat& Other) const;
 
 	/** copy operator to prevent duplicate handles */
-	FScalableFloat& operator=(const FScalableFloat& Src);
+	UE_API FScalableFloat& operator=(const FScalableFloat& Src);
 
 	/** Used to upgrade a float or int8/int16/int32 property into an FScalableFloat */
-	bool SerializeFromMismatchedTag(const FPropertyTag& Tag, FStructuredArchive::FSlot Slot);
+	UE_API bool SerializeFromMismatchedTag(const FPropertyTag& Tag, FStructuredArchive::FSlot Slot);
 
 private:
 
@@ -131,3 +133,5 @@ struct TStructOpsTypeTraits<FScalableFloat>
 /** Macro to call IsValidWithWarnings with a correct error info */
 #define SCALABLEFLOAT_REPORTERROR_WITHPATHNAME(Scalable, PathNameString)\
 	Scalable.IsValidWithWarnings(FString::Printf(TEXT("%s.%s"), *PathNameString, TEXT(#Scalable)))
+
+#undef UE_API
diff --git a/Engine/Plugins/Runtime/GameplayAbilities/Source/GameplayAbilitiesEditor/Public/AttributeDetails.h b/Engine/Plugins/Runtime/GameplayAbilities/Source/GameplayAbilitiesEditor/Public/AttributeDetails.h
--- a/Engine/Plugins/Runtime/GameplayAbilities/Source/GameplayAbilitiesEditor/Public/AttributeDetails.h
+++ b/Engine/Plugins/Runtime/GameplayAbilities/Source/GameplayAbilitiesEditor/Public/AttributeDetails.h
@@ -12,6 +12,8 @@
 #include "IPropertyTypeCustomization.h"
 #include "PropertyEditorModule.h"
 
+#define UE_API GAMEPLAYABILITIESEDITOR_API
+
 class IDetailLayoutBuilder;
 class IPropertyHandle;
 class IPropertyUtilities;
@@ -66,10 +68,10 @@ private:
 	void OnAttributeChanged(FProperty* SelectedAttribute);
 };
 
-class GAMEPLAYABILITIESEDITOR_API FScalableFloatDetails : public IPropertyTypeCustomization
+class FScalableFloatDetails : public IPropertyTypeCustomization
 {
 public:
-	static TSharedRef<IPropertyTypeCustomization> MakeInstance();
+	static UE_API TSharedRef<IPropertyTypeCustomization> MakeInstance();
 	static constexpr float DefaultMinPreviewLevel = 0.f;
 	static constexpr float DefaultMaxPreviewLevel = 30.f;
 
@@ -82,57 +84,57 @@ public:
 
 protected:
 
-	virtual void CustomizeHeader( TSharedRef<class IPropertyHandle> StructPropertyHandle, class FDetailWidgetRow& HeaderRow, IPropertyTypeCustomizationUtils& StructCustomizationUtils ) override;
-	virtual void CustomizeChildren( TSharedRef<class IPropertyHandle> StructPropertyHandle, class IDetailChildrenBuilder& StructBuilder, IPropertyTypeCustomizationUtils& StructCustomizationUtils ) override;              
+	UE_API virtual void CustomizeHeader( TSharedRef<class IPropertyHandle> StructPropertyHandle, class FDetailWidgetRow& HeaderRow, IPropertyTypeCustomizationUtils& StructCustomizationUtils ) override;
+	UE_API virtual void CustomizeChildren( TSharedRef<class IPropertyHandle> StructPropertyHandle, class IDetailChildrenBuilder& StructBuilder, IPropertyTypeCustomizationUtils& StructCustomizationUtils ) override;              
 
-	bool IsEditable() const;
-	void UpdatePreviewLevels();
+	UE_API bool IsEditable() const;
+	UE_API void UpdatePreviewLevels();
 
 	// Curve Table selector
-	TSharedRef<SWidget> CreateCurveTableWidget();
-	TSharedRef<SWidget> GetCurveTablePicker();
-	void OnSelectCurveTable(const FAssetData& AssetData);
-	void OnCloseMenu();
-	FText GetCurveTableText() const;
-	FText GetCurveTableTooltip() const;
-	EVisibility GetCurveTableVisiblity() const;
-	EVisibility GetAssetButtonVisiblity() const;
-	void OnBrowseTo();
-	void OnClear();
-	void OnUseSelected();
+	UE_API TSharedRef<SWidget> CreateCurveTableWidget();
+	UE_API TSharedRef<SWidget> GetCurveTablePicker();
+	UE_API void OnSelectCurveTable(const FAssetData& AssetData);
+	UE_API void OnCloseMenu();
+	UE_API FText GetCurveTableText() const;
+	UE_API FText GetCurveTableTooltip() const;
+	UE_API EVisibility GetCurveTableVisiblity() const;
+	UE_API EVisibility GetAssetButtonVisiblity() const;
+	UE_API void OnBrowseTo();
+	UE_API void OnClear();
+	UE_API void OnUseSelected();
 	
 	// Registry Type selector
-	TSharedRef<SWidget> CreateRegistryTypeWidget();
-	FString GetRegistryTypeValueString() const;
-	FText GetRegistryTypeTooltip() const;
-	EVisibility GetRegistryTypeVisiblity() const;
+	UE_API TSharedRef<SWidget> CreateRegistryTypeWidget();
+	UE_API FString GetRegistryTypeValueString() const;
+	UE_API FText GetRegistryTypeTooltip() const;
+	UE_API EVisibility GetRegistryTypeVisiblity() const;
 
 	// Curve source accessors
-	void OnCurveSourceChanged();
-	void RefreshSourceData();
-	class UCurveTable* GetCurveTable(FPropertyAccess::Result* OutResult = nullptr) const;
-	FDataRegistryType GetRegistryType(FPropertyAccess::Result* OutResult = nullptr) const;
+	UE_API void OnCurveSourceChanged();
+	UE_API void RefreshSourceData();
+	UE_API class UCurveTable* GetCurveTable(FPropertyAccess::Result* OutResult = nullptr) const;
+	UE_API FDataRegistryType GetRegistryType(FPropertyAccess::Result* OutResult = nullptr) const;
 
 	// Row/item name widget
-	TSharedRef<SWidget> CreateRowNameWidget();
-	EVisibility GetRowNameVisibility() const;
-	FText GetRowNameComboBoxContentText() const;
-	FText GetRowNameComboBoxContentTooltip() const;
-	void OnRowNameChanged();
+	UE_API TSharedRef<SWidget> CreateRowNameWidget();
+	UE_API EVisibility GetRowNameVisibility() const;
+	UE_API FText GetRowNameComboBoxContentText() const;
+	UE_API FText GetRowNameComboBoxContentTooltip() const;
+	UE_API void OnRowNameChanged();
 
 	// Preview widgets
-	EVisibility GetPreviewVisibility() const;
-	float GetPreviewLevel() const;
-	void SetPreviewLevel(float NewLevel);
-	FText GetRowValuePreviewLabel() const;
-	FText GetRowValuePreviewText() const;
+	UE_API EVisibility GetPreviewVisibility() const;
+	UE_API float GetPreviewLevel() const;
+	UE_API void SetPreviewLevel(float NewLevel);
+	UE_API FText GetRowValuePreviewLabel() const;
+	UE_API FText GetRowValuePreviewText() const;
 
 	// Row accessors and callbacks
-	FName GetRowName(FPropertyAccess::Result* OutResult = nullptr) const;
-	const FRealCurve* GetRealCurve(FPropertyAccess::Result* OutResult = nullptr) const;
-	FDataRegistryId GetRegistryId(FPropertyAccess::Result* OutResult = nullptr) const;
-	void SetRegistryId(FDataRegistryId NewId);
-	void GetCustomRowNames(TArray<FName>& OutRows) const;
+	UE_API FName GetRowName(FPropertyAccess::Result* OutResult = nullptr) const;
+	UE_API const FRealCurve* GetRealCurve(FPropertyAccess::Result* OutResult = nullptr) const;
+	UE_API FDataRegistryId GetRegistryId(FPropertyAccess::Result* OutResult = nullptr) const;
+	UE_API void SetRegistryId(FDataRegistryId NewId);
+	UE_API void GetCustomRowNames(TArray<FName>& OutRows) const;
 
 	TSharedPtr<IPropertyHandle> ValueProperty;
 	TSharedPtr<IPropertyHandle> CurveTableHandleProperty;
@@ -147,3 +149,5 @@ protected:
 	float MaxPreviewLevel;
 	bool bSourceRefreshQueued;
 };
+
+#undef UE_API
diff --git a/Engine/Plugins/Runtime/GameplayAbilities/Source/GameplayAbilitiesEditor/Public/GameplayAbilitiesEditorModule.h b/Engine/Plugins/Runtime/GameplayAbilities/Source/GameplayAbilitiesEditor/Public/GameplayAbilitiesEditorModule.h
--- a/Engine/Plugins/Runtime/GameplayAbilities/Source/GameplayAbilitiesEditor/Public/GameplayAbilitiesEditorModule.h
+++ b/Engine/Plugins/Runtime/GameplayAbilities/Source/GameplayAbilitiesEditor/Public/GameplayAbilitiesEditorModule.h
@@ -6,7 +6,9 @@
 #include "Modules/ModuleInterface.h"
 #include "Modules/ModuleManager.h"
 
-struct GAMEPLAYABILITIESEDITOR_API FGameplayCueEditorStrings
+#define UE_API GAMEPLAYABILITIESEDITOR_API
+
+struct FGameplayCueEditorStrings
 {
 	FString GameplayCueNotifyDescription1;
 	FString GameplayCueNotifyDescription2;
@@ -79,3 +81,5 @@ public:
 	virtual FGetGameplayCueEditorStrings& GetGameplayCueEditorStringsDelegate() = 0;
 	
 };
+
+#undef UE_API
diff --git a/Engine/Plugins/Runtime/GameplayAbilities/Source/GameplayAbilitiesEditor/Public/GameplayAbilityAudit.h b/Engine/Plugins/Runtime/GameplayAbilities/Source/GameplayAbilitiesEditor/Public/GameplayAbilityAudit.h
--- a/Engine/Plugins/Runtime/GameplayAbilities/Source/GameplayAbilitiesEditor/Public/GameplayAbilityAudit.h
+++ b/Engine/Plugins/Runtime/GameplayAbilities/Source/GameplayAbilitiesEditor/Public/GameplayAbilityAudit.h
@@ -6,6 +6,8 @@
 #include "Abilities/GameplayAbility.h"
 #include "GameplayAbilityAudit.generated.h"
 
+#define UE_API GAMEPLAYABILITIESEDITOR_API
+
 class UBlueprint;
 class UEdGraph;
 
@@ -41,15 +43,15 @@ ENUM_CLASS_FLAGS(EGameplayAbilityEndInBlueprints)
 
 /** This contains all of the data we gather about a Blueprint during the Audit process */
 USTRUCT(BlueprintInternalUseOnlyHierarchical)
-struct GAMEPLAYABILITIESEDITOR_API FGameplayAbilityAuditRow : public FTableRowBase
+struct FGameplayAbilityAuditRow : public FTableRowBase
 {
 	GENERATED_BODY()
 	
 	/** Fill this structure with data from a UBlueprint (GameplayAbility asset).  Return true if successful. */
-	virtual void FillDataFromGameplayAbilityBlueprint(const UBlueprint& GameplayAbilityBlueprint);
+	UE_API virtual void FillDataFromGameplayAbilityBlueprint(const UBlueprint& GameplayAbilityBlueprint);
 
 	/** Fill this structure with data from a UGameplayAbility (likely compiled from a UBlueprint).  Return true if successful. */
-	virtual void FillDataFromGameplayAbility(const UGameplayAbility& GameplayAbility);
+	UE_API virtual void FillDataFromGameplayAbility(const UGameplayAbility& GameplayAbility);
 
 protected:
 
@@ -126,4 +128,6 @@ protected:
 namespace GameplayAbilityAudit
 {
 	GAMEPLAYABILITIESEDITOR_API TArray<UEdGraph*> GatherAllGraphsIncludingMacros(const UBlueprint& LoadedInstance);
-}
\ No newline at end of file
+}
+
+#undef UE_API
diff --git a/Engine/Plugins/Runtime/GameplayAbilities/Source/GameplayAbilitiesEditor/Public/GameplayEffectCreationMenu.h b/Engine/Plugins/Runtime/GameplayAbilities/Source/GameplayAbilitiesEditor/Public/GameplayEffectCreationMenu.h
--- a/Engine/Plugins/Runtime/GameplayAbilities/Source/GameplayAbilitiesEditor/Public/GameplayEffectCreationMenu.h
+++ b/Engine/Plugins/Runtime/GameplayAbilities/Source/GameplayAbilitiesEditor/Public/GameplayEffectCreationMenu.h
@@ -5,6 +5,8 @@
 #include "GameplayEffect.h"
 #include "GameplayEffectCreationMenu.generated.h"
 
+#define UE_API GAMEPLAYABILITIESEDITOR_API
+
 class UGameplayEffect;
 
 USTRUCT()
@@ -25,23 +27,24 @@ struct FGameplayEffectCreationData
 };
 
 /** Container to hold EventKeywords for PIE testing */
-UCLASS(config=Game,defaultconfig)
-class GAMEPLAYABILITIESEDITOR_API UGameplayEffectCreationMenu : public UObject
+UCLASS(config=Game,defaultconfig, MinimalAPI)
+class UGameplayEffectCreationMenu : public UObject
 {
 	GENERATED_BODY()
 
 public:
 
-	UGameplayEffectCreationMenu();
+	UE_API UGameplayEffectCreationMenu();
 	
 	// Set this in your project to programatically define default GE names
-	static TFunction< FString(FString BaseName, FString Path) > GetDefaultAssetNameFunc;
+	static UE_API TFunction< FString(FString BaseName, FString Path) > GetDefaultAssetNameFunc;
 
 
 
-	void AddMenuExtensions() const;
+	UE_API void AddMenuExtensions() const;
 
 	UPROPERTY(config, EditAnywhere, Category="Gameplay Effect")
 	TArray<FGameplayEffectCreationData> Definitions;
 };
 
+#undef UE_API
