Added helper fns to Ability system library: HasAnyAbilitiesWithCondition and HasAnyAbilitiesWithAssetTag

#tests called the new helper fns in another WIP dependant CL and ensured they worked as expected.

diff --git a/Engine/Plugins/Runtime/GameplayAbilities/Source/GameplayAbilities/Private/AbilitySystemBlueprintLibrary.cpp b/Engine/Plugins/Runtime/GameplayAbilities/Source/GameplayAbilities/Private/AbilitySystemBlueprintLibrary.cpp
--- a/Engine/Plugins/Runtime/GameplayAbilities/Source/GameplayAbilities/Private/AbilitySystemBlueprintLibrary.cpp
+++ b/Engine/Plugins/Runtime/GameplayAbilities/Source/GameplayAbilities/Private/AbilitySystemBlueprintLibrary.cpp
@@ -1459,6 +1459,68 @@ bool UAbilitySystemBlueprintLibrary::IsGameplayAbilityActive(const UGameplayAbil
 	return GameplayAbility->IsActive();
 }
 
+bool UAbilitySystemBlueprintLibrary::HasAnyAbilitiesByPredicate(
+	const UAbilitySystemComponent* AbilitySystemComponent, 
+	const TFunctionRef<bool(const UGameplayAbility& AbilityRef)> Predicate, 
+	bool bOnlyRunPredicateOnAbilityCDOs /*= false*/)
+{
+	if (AbilitySystemComponent == nullptr)
+	{
+		return false;
+	}
+
+	// Copy off the ActivatableAbilities for safety to avoid situations where the Predicate is affecting the list while iterating.
+	TArray<FGameplayAbilitySpec> ActivatableAbilitiesCopy = AbilitySystemComponent->GetActivatableAbilities();
+	for (const FGameplayAbilitySpec& Spec : ActivatableAbilitiesCopy)
+	{
+		if (Spec.Ability == nullptr)
+		{
+			continue;
+		}
+
+		TArray<UGameplayAbility*> AbilitiesToCheck;
+
+		if (Spec.Ability->IsInstantiated() && !bOnlyRunPredicateOnAbilityCDOs)
+		{
+			AbilitiesToCheck.Append(Spec.GetAbilityInstances());
+		}
+		else
+		{
+			AbilitiesToCheck.Add(Spec.Ability);
+		}
+
+		for (UGameplayAbility* AbilityToCheck : AbilitiesToCheck)
+		{
+			if (AbilityToCheck == nullptr)
+			{
+				continue;
+			}
+
+			const bool bPassesCondition = Predicate(*AbilityToCheck);
+			if (bPassesCondition)
+			{
+				return true;
+			}
+		}
+	}
+
+	return false;
+}
+
+bool UAbilitySystemBlueprintLibrary::HasAnyAbilitiesWithAssetTag(const UAbilitySystemComponent* AbilitySystemComponent, FGameplayTag AssetTag)
+{
+	if (AbilitySystemComponent == nullptr)
+	{
+		return false;
+	}
+
+	constexpr bool bOnlyCheckAbilityCDO = true;
+	return HasAnyAbilitiesByPredicate(AbilitySystemComponent, [AssetTag](const UGameplayAbility& AbilityRef) -> bool
+	{
+		return AbilityRef.GetAssetTags().HasTag(AssetTag);
+	}, bOnlyCheckAbilityCDO);
+}
+
 bool UAbilitySystemBlueprintLibrary::EqualEqual_GameplayAbilitySpecHandle(const FGameplayAbilitySpecHandle& A, const FGameplayAbilitySpecHandle& B)
 {
 	return A == B;
diff --git a/Engine/Plugins/Runtime/GameplayAbilities/Source/GameplayAbilities/Public/AbilitySystemBlueprintLibrary.h b/Engine/Plugins/Runtime/GameplayAbilities/Source/GameplayAbilities/Public/AbilitySystemBlueprintLibrary.h
--- a/Engine/Plugins/Runtime/GameplayAbilities/Source/GameplayAbilities/Public/AbilitySystemBlueprintLibrary.h
+++ b/Engine/Plugins/Runtime/GameplayAbilities/Source/GameplayAbilities/Public/AbilitySystemBlueprintLibrary.h
@@ -576,6 +576,21 @@ public:
 	UFUNCTION(BlueprintPure, Category = "Ability|GameplayAbility", meta=(DisplayName="Is Active",DefaultToSelf=GameplayAbility))
 	static UE_API bool IsGameplayAbilityActive(const UGameplayAbility* GameplayAbility);
 
+	/** Returns true if the ASC has any abilities that meet a given predicate.
+	 * @param AbilitySystemComponent         The ability system component we are checking abilities on.
+	 * @param Predicate                      The predicate to run on each ability
+	 * @param bOnlyRunPredicateOnAbilityCDOs If true, we will only run the predicate on the ability specs' Ability CDO instead of on its instance(s).
+	 * @return                               Returns true if any abilities passed the predicate.
+	 */
+	static UE_API bool HasAnyAbilitiesByPredicate(
+		const UAbilitySystemComponent* AbilitySystemComponent, 
+		const TFunctionRef<bool(const UGameplayAbility& AbilityRef)> Predicate, 
+		bool bOnlyRunPredicateOnAbilityCDOs = false);
+
+	/** Returns true if the ASC has any abilities that have a specific asset tag.*/
+	UFUNCTION(BlueprintPure, Category = "Ability|GameplayAbility")
+	static UE_API bool HasAnyAbilitiesWithAssetTag(const UAbilitySystemComponent* AbilitySystemComponent, FGameplayTag AssetTag);
+
 	/** Equality operator for two Gameplay Ability Spec Handles */
 	UFUNCTION(BlueprintPure, Category = "Ability|GameplayAbility", meta = (DisplayName = "Equal (Gameplay Ability Spec Handle)", CompactNodeTitle = "==", ScriptOperator = "=="))
 	static UE_API bool EqualEqual_GameplayAbilitySpecHandle(const FGameplayAbilitySpecHandle& A, const FGameplayAbilitySpecHandle& B);
