Adding proper compound multiply functionality to `ComputeStackedModifierMagnitude`.
[RN] Fixed Multiply (Compound) modifiers not calculating correctly with stacked GameplayEffects.

diff --git a/Engine/Plugins/Runtime/GameplayAbilities/Source/GameplayAbilities/Private/GameplayEffectTypes.cpp b/Engine/Plugins/Runtime/GameplayAbilities/Source/GameplayAbilities/Private/GameplayEffectTypes.cpp
--- a/Engine/Plugins/Runtime/GameplayAbilities/Source/GameplayAbilities/Private/GameplayEffectTypes.cpp
+++ b/Engine/Plugins/Runtime/GameplayAbilities/Source/GameplayAbilities/Private/GameplayEffectTypes.cpp
@@ -40,6 +40,14 @@ static FAutoConsoleVariableRef CVarWarnIfTryingToReplicateNotSupportedActorRefer
 }
 #endif
 
+static bool bUseMultiplyCompoundFix = true;
+static FAutoConsoleVariableRef CVarUseMultiplyCompoundFix(
+	TEXT("AbilitySystem.UseMultiplyCompundFix"),
+	bUseMultiplyCompoundFix,
+	TEXT("True/False - Add proper logic for compound multiplication in stackable GEs."),
+	ECVF_Default);
+
+
 FGameplayModEvaluationChannelSettings::FGameplayModEvaluationChannelSettings()
 {
 	static const UEnum* EvalChannelEnum = nullptr;
@@ -123,7 +131,21 @@ float GameplayEffectUtilities::ComputeStackedModifierMagnitude(float BaseCompute
 	if (ModOp != EGameplayModOp::Override)
 	{
 		StackMag -= OperationBias;
-		StackMag *= StackCount;
+		if (bUseMultiplyCompoundFix)
+		{
+			if (ModOp == EGameplayModOp::MultiplyCompound)
+			{
+				StackMag = FMath::Pow(StackMag, StackCount);
+			}
+			else
+			{
+				StackMag *= StackCount;
+			}
+		}
+		else
+		{
+			StackMag *= StackCount;
+		}
 		StackMag += OperationBias;
 	}
 
