When GameplayAbility is configured with a cooldown GE class, validation will check whether that GE applies any tags. CVAr AbilitySystem.WarnCooldownEffectWithoutTags can be used to opt-out of this.

diff --git a/Engine/Plugins/Runtime/GameplayAbilities/Source/GameplayAbilities/Private/Abilities/GameplayAbility.cpp b/Engine/Plugins/Runtime/GameplayAbilities/Source/GameplayAbilities/Private/Abilities/GameplayAbility.cpp
--- a/Engine/Plugins/Runtime/GameplayAbilities/Source/GameplayAbilities/Private/Abilities/GameplayAbility.cpp
+++ b/Engine/Plugins/Runtime/GameplayAbilities/Source/GameplayAbilities/Private/Abilities/GameplayAbility.cpp
@@ -39,6 +39,15 @@ namespace UE::AbilitySystem::Private
 	int32 CVarDeprecateReplicatedPropertiesValue = 2;
 	static FAutoConsoleVariableRef CVarDeprecateReplicatedProperties(TEXT("AbilitySystem.DeprecateReplicatedProperties"), CVarDeprecateReplicatedPropertiesValue,
 		TEXT("Deprecate the use of Replicated Properties. ") TEXT("0: No. 1: Yes. 2: Yes, Except During Automation"), ECVF_Default );
+
+	/** 
+		By default, for a GameplayEffect to be usable as an ability's CooldownGameplayEffectClass class, that GameplayEffect must grant at least one GameplayTag.
+		If the GE grants no tags, the ability doesn't consider itself on cooldown even if the GE is applied. Thus by default we want to warn developers at editor-time
+		when a CooldownGameplayEffectClass is configured that grants no tags. These warnings can be disabled via this CVar, because UGameplayAbility::CheckCooldown
+		can be overridden and thus the CooldownGameplayEffectClass may not require granted tags.
+	*/
+	int32 CVarWarnCooldownEffectWithoutTagsValue = 1;
+	static FAutoConsoleVariableRef CVarWarnCooldownEffectWithoutTags(TEXT("AbilitySystem.WarnCooldownEffectWithoutTags"), CVarWarnCooldownEffectWithoutTagsValue, TEXT("Show an editor warning and runtime log when a GameplayAbility is configured with a CooldownGameplayEffectClass that grants no tags"), ECVF_Default);
 }
 
 namespace FAbilitySystemTweaks
@@ -309,6 +318,21 @@ EDataValidationResult UGameplayAbility::IsDataValid(FDataValidationContext& Cont
 		}
 	}
 
+	// Generate a blueprint warning if a Cooldown GE class is configured that doesn't grant tags.
+	if (UE::AbilitySystem::Private::CVarWarnCooldownEffectWithoutTagsValue > 0)
+	{
+		if (const UGameplayEffect* CooldownGE = GetCooldownGameplayEffect())
+		{
+			// Detect misconfigured cooldown GE class
+			if (CooldownGE->GetGrantedTags().IsEmpty())
+			{
+				FText ErrorText = FText::Format(LOCTEXT("CooldownEffectNoTags", "CooldownGameplayEffectClass '{0}' grants no tags. A GameplayEffect class must grant tags (Component: Grant Tags to Target Actor) to be used as cooldown."), CooldownGameplayEffectClass->GetDisplayNameText());
+				Context.AddError(ErrorText);
+				Result = EDataValidationResult::Invalid;
+			}
+		}
+	}
+
 	return Result;
 }
 #endif
@@ -1055,27 +1079,36 @@ bool UGameplayAbility::CheckCooldown(const FGameplayAbilitySpecHandle Handle, co
 	}
 
 	const FGameplayTagContainer* CooldownTags = GetCooldownTags();
-	if (CooldownTags && !CooldownTags->IsEmpty())
+	if (CooldownTags)
 	{
-		if (UAbilitySystemComponent* AbilitySystemComponent = ActorInfo->AbilitySystemComponent.Get())
+		if (!CooldownTags->IsEmpty())
 		{
-			if (AbilitySystemComponent->HasAnyMatchingGameplayTags(*CooldownTags))
+			if (UAbilitySystemComponent* AbilitySystemComponent = ActorInfo->AbilitySystemComponent.Get())
 			{
-				if (OptionalRelevantTags)
+				if (AbilitySystemComponent->HasAnyMatchingGameplayTags(*CooldownTags))
 				{
-					const FGameplayTag& FailCooldownTag = UAbilitySystemGlobals::Get().ActivateFailCooldownTag;
-					if (FailCooldownTag.IsValid())
+					if (OptionalRelevantTags)
 					{
-						OptionalRelevantTags->AddTag(FailCooldownTag);
+						const FGameplayTag& FailCooldownTag = UAbilitySystemGlobals::Get().ActivateFailCooldownTag;
+						if (FailCooldownTag.IsValid())
+						{
+							OptionalRelevantTags->AddTag(FailCooldownTag);
+						}
+
+						// Let the caller know which tags were blocking
+						OptionalRelevantTags->AppendMatchingTags(AbilitySystemComponent->GetOwnedGameplayTags(), *CooldownTags);
 					}
 
-					// Let the caller know which tags were blocking
-					OptionalRelevantTags->AppendMatchingTags(AbilitySystemComponent->GetOwnedGameplayTags(), *CooldownTags);
+					return false;
 				}
-
-				return false;
 			}
 		}
+		else if (UE::AbilitySystem::Private::CVarWarnCooldownEffectWithoutTagsValue > 0)
+		{
+			// Unless this function is overridden, a CooldownGameplayEffectClass that doesn't grant tags isn't effective as cooldown. 
+			// Log a runtime warning when an ability's cooldown is checked while having this misconfiguration.
+			UE_LOG(LogAbilitySystem, Warning, TEXT("CooldownGameplayEffectClass '%s' grants no tags. A GameplayEffect class must grant tags (Component: Grant Tags to Target Actor) to be used as cooldown."), *CooldownGameplayEffectClass->GetName());
+		}
 	}
 	return true;
 }
