Stability fix for UE-221810, preventing a crash when project overrides GetCooldownEffectClass()

diff --git a/Engine/Plugins/Runtime/GameplayAbilities/Source/GameplayAbilities/Private/Abilities/GameplayAbility.cpp b/Engine/Plugins/Runtime/GameplayAbilities/Source/GameplayAbilities/Private/Abilities/GameplayAbility.cpp
--- a/Engine/Plugins/Runtime/GameplayAbilities/Source/GameplayAbilities/Private/Abilities/GameplayAbility.cpp
+++ b/Engine/Plugins/Runtime/GameplayAbilities/Source/GameplayAbilities/Private/Abilities/GameplayAbility.cpp
@@ -319,12 +319,14 @@ EDataValidationResult UGameplayAbility::IsDataValid(FDataValidationContext& Cont
 	}
 
 	// Generate a blueprint warning if a Cooldown GE class is configured that doesn't grant tags.
-	if (UE::AbilitySystem::Private::CVarWarnCooldownEffectWithoutTagsValue > 0)
+	// If your project overrides GetCooldownGameplayEffect() or CheckCooldown() so that cooldown
+	// checking doesn't depend on a tag, set this CVar to 0.
+	if (UE::AbilitySystem::Private::CVarWarnCooldownEffectWithoutTagsValue > 0 && CooldownGameplayEffectClass != nullptr)
 	{
 		if (const UGameplayEffect* CooldownGE = GetCooldownGameplayEffect())
 		{
-			// Detect misconfigured cooldown GE class
-			if (CooldownGE->GetGrantedTags().IsEmpty())
+			// Detect misconfigured cooldown GE class. Ignore when class doesn't match due to overridden GetCooldownGameplayEffect().
+			if (CooldownGE->GetClass() == CooldownGameplayEffectClass && CooldownGE->GetGrantedTags().IsEmpty())
 			{
 				FText ErrorText = FText::Format(LOCTEXT("CooldownEffectNoTags", "CooldownGameplayEffectClass '{0}' grants no tags. A GameplayEffect class must grant tags (Component: Grant Tags to Target Actor) to be used as cooldown."), CooldownGameplayEffectClass->GetDisplayNameText());
 				Context.AddError(ErrorText);
@@ -1103,7 +1105,7 @@ bool UGameplayAbility::CheckCooldown(const FGameplayAbilitySpecHandle Handle, co
 				}
 			}
 		}
-		else if (UE::AbilitySystem::Private::CVarWarnCooldownEffectWithoutTagsValue > 0)
+		else if (UE::AbilitySystem::Private::CVarWarnCooldownEffectWithoutTagsValue > 0 && CooldownGameplayEffectClass != nullptr)
 		{
 			// Unless this function is overridden, a CooldownGameplayEffectClass that doesn't grant tags isn't effective as cooldown. 
 			// Log a runtime warning when an ability's cooldown is checked while having this misconfiguration.
