- Gameplay Abilities need to Cancel themselves in PreDestroyFromReplication if they're still active, as the server won't make the client calls to do it.
- When removing abilities we need to replicate EndAbility if the GA isn't replicated, but skip it if it is.
- Removing CVarReplicateCancelDuringOnRemove as it is no longer needed with the added GA replication check and Actor cleanup detection.

[RN] Fixed replication issue canceling Gameplay Abilities when destroying the ASC from the server while abilities are active.

diff --git a/Engine/Plugins/Runtime/GameplayAbilities/Source/GameplayAbilities/Private/Abilities/GameplayAbility.cpp b/Engine/Plugins/Runtime/GameplayAbilities/Source/GameplayAbilities/Private/Abilities/GameplayAbility.cpp
--- a/Engine/Plugins/Runtime/GameplayAbilities/Source/GameplayAbilities/Private/Abilities/GameplayAbility.cpp
+++ b/Engine/Plugins/Runtime/GameplayAbilities/Source/GameplayAbilities/Private/Abilities/GameplayAbility.cpp
@@ -242,6 +242,15 @@ bool UGameplayAbility::IsSupportedForNetworking() const
 	return Supported;
 }
 
+void UGameplayAbility::PreDestroyFromReplication()
+{
+	// We're being destroyed, just cancel locally.
+	if (IsActive())
+	{
+		CancelAbility(CurrentSpecHandle, CurrentActorInfo, CurrentActivationInfo, false);
+	}
+}
+
 #if WITH_EDITOR
 EDataValidationResult UGameplayAbility::IsDataValid(FDataValidationContext& Context) const
 {
diff --git a/Engine/Plugins/Runtime/GameplayAbilities/Source/GameplayAbilities/Private/AbilitySystemComponent_Abilities.cpp b/Engine/Plugins/Runtime/GameplayAbilities/Source/GameplayAbilities/Private/AbilitySystemComponent_Abilities.cpp
--- a/Engine/Plugins/Runtime/GameplayAbilities/Source/GameplayAbilities/Private/AbilitySystemComponent_Abilities.cpp
+++ b/Engine/Plugins/Runtime/GameplayAbilities/Source/GameplayAbilities/Private/AbilitySystemComponent_Abilities.cpp
@@ -53,7 +53,6 @@ static TAutoConsoleVariable<bool> CVarGasFixClientSideMontageBlendOutTime(TEXT("
 static TAutoConsoleVariable<bool> CVarUpdateMontageSectionIdToPlay(TEXT("AbilitySystem.UpdateMontageSectionIdToPlay"), true, TEXT("During tick, update the section ID that replicated montages should use"));
 static TAutoConsoleVariable<bool> CVarReplicateMontageNextSectionId(TEXT("AbilitySystem.ReplicateMontageNextSectionId"), true, TEXT("Apply the replicated next section Id to montages when skipping position replication"));
 static TAutoConsoleVariable<bool> CVarEnsureAbilitiesEndGracefully(TEXT("AbilitySystem.EnsureAbilitiesEndGracefully"), true, TEXT("When shutting down (during ClearAllAbilities) we should check if all GameplayAbilities gracefully ended. This should be disabled if you have NonInstanced abilities that are designed for multiple concurrent executions."));
-static TAutoConsoleVariable<bool> CVarReplicateCancelDuringOnRemove(TEXT("AbilitySystem.ReplicateCancelOnRemove"), false, TEXT("When an ability is removed and not yet ended, assume it's cancelled and replicate to ensure the other side receives that before instance destruction. If false, you may miss EndAbility if the Actor is destroyed."));
 
 void UAbilitySystemComponent::InitializeComponent()
 {
@@ -632,14 +631,9 @@ void UAbilitySystemComponent::OnRemoveAbility(FGameplayAbilitySpec& Spec)
 		{
 			if (Instance->IsActive())
 			{
-				// End the ability but don't replicate it, OnRemoveAbility gets replicated
-				bool bReplicateEndAbility = false;
+				// ReplicateEndAbility if the GA itself isn't replicated.
+				bool bReplicateEndAbility = Instance->GetReplicationPolicy() == EGameplayAbilityReplicationPolicy::ReplicateNo;
 				bool bWasCancelled = false;
-				if (CVarReplicateCancelDuringOnRemove.GetValueOnGameThread())
-				{
-					bReplicateEndAbility = true;
-					bWasCancelled = true;
-				}
 				Instance->EndAbility(Instance->CurrentSpecHandle, Instance->CurrentActorInfo, Instance->CurrentActivationInfo, bReplicateEndAbility, bWasCancelled);
 			}
 			else
@@ -1335,14 +1329,14 @@ void UAbilitySystemComponent::CancelAbilitySpec(FGameplayAbilitySpec& Spec, UGam
 		{
 			if (InstanceAbility && Ignore != InstanceAbility)
 			{
-				InstanceAbility->CancelAbility(Spec.Handle, ActorInfo, InstanceAbility->GetCurrentActivationInfoRef(), true);
+				InstanceAbility->CancelAbility(Spec.Handle, ActorInfo, InstanceAbility->GetCurrentActivationInfoRef(), !bDestroyActiveStateInitiated);
 			}
 		}
 	}
 	else
 	{
 		// Try to cancel the non instanced, this may not necessarily work
-		Spec.Ability->CancelAbility(Spec.Handle, ActorInfo, FGameplayAbilityActivationInfo(), true);
+		Spec.Ability->CancelAbility(Spec.Handle, ActorInfo, FGameplayAbilityActivationInfo(), !bDestroyActiveStateInitiated);
 	}
 	MarkAbilitySpecDirty(Spec);
 }
diff --git a/Engine/Plugins/Runtime/GameplayAbilities/Source/GameplayAbilities/Public/Abilities/GameplayAbility.h b/Engine/Plugins/Runtime/GameplayAbilities/Source/GameplayAbilities/Public/Abilities/GameplayAbility.h
--- a/Engine/Plugins/Runtime/GameplayAbilities/Source/GameplayAbilities/Public/Abilities/GameplayAbility.h
+++ b/Engine/Plugins/Runtime/GameplayAbilities/Source/GameplayAbilities/Public/Abilities/GameplayAbility.h
@@ -513,6 +513,7 @@ public:
 	UE_API virtual int32 GetFunctionCallspace(UFunction* Function, FFrame* Stack) override;
 	UE_API virtual bool CallRemoteFunction(UFunction* Function, void* Parameters, FOutParmRec* OutParms, FFrame* Stack) override;
 	UE_API virtual bool IsSupportedForNetworking() const override;
+	UE_API virtual void PreDestroyFromReplication() override;
 
 #if WITH_EDITOR
 	UE_API virtual EDataValidationResult IsDataValid(FDataValidationContext& Context) const override;
