- Adds predictive gameplay effect stacking.
- Adds better handling for when adding gameplay effect stacks will overflow the stack limit.
#13023
[RN] Added predictive gameplay effect stacking and improved stack overflow handling.

diff --git a/Engine/Plugins/Runtime/GameplayAbilities/Source/GameplayAbilities/Private/AbilitySystemComponent.cpp b/Engine/Plugins/Runtime/GameplayAbilities/Source/GameplayAbilities/Private/AbilitySystemComponent.cpp
--- a/Engine/Plugins/Runtime/GameplayAbilities/Source/GameplayAbilities/Private/AbilitySystemComponent.cpp
+++ b/Engine/Plugins/Runtime/GameplayAbilities/Source/GameplayAbilities/Private/AbilitySystemComponent.cpp
@@ -1065,6 +1065,11 @@ void UAbilitySystemComponent::RemoveActiveGameplayEffect_AllowClientRemoval(FAct
 	ActiveGameplayEffects.RemoveActiveGameplayEffect(Handle, StacksToRemove);
 }
 
+void UAbilitySystemComponent::OnPredictiveGameplayEffectStackCaughtUp(FActiveGameplayEffectHandle Handle)
+{
+	ActiveGameplayEffects.OnPredictiveGameplayEffectStackCaughtUp(Handle);
+}
+
 void UAbilitySystemComponent::RemoveActiveGameplayEffectBySourceEffect(TSubclassOf<UGameplayEffect> GameplayEffect, UAbilitySystemComponent* InstigatorAbilitySystemComponent, int32 StacksToRemove /*= -1*/)
 {
 	if (GameplayEffect)
diff --git a/Engine/Plugins/Runtime/GameplayAbilities/Source/GameplayAbilities/Private/GameplayEffect.cpp b/Engine/Plugins/Runtime/GameplayAbilities/Source/GameplayAbilities/Private/GameplayEffect.cpp
--- a/Engine/Plugins/Runtime/GameplayAbilities/Source/GameplayAbilities/Private/GameplayEffect.cpp
+++ b/Engine/Plugins/Runtime/GameplayAbilities/Source/GameplayAbilities/Private/GameplayEffect.cpp
@@ -72,6 +72,24 @@ namespace UE::GameplayEffect
 	FAutoConsoleVariableRef CVarActiveGameplayEffectReplicationFix{ TEXT("AbilitySystem.Fix.ActiveGEReplicationFix"), ActiveGameplayEffectReplicationFix, TEXT("Experimental code mask for fixing Active Gameplay Effects (set to 0 to disable)"), ECVF_Default };
 	inline bool HasActiveGameplayEffectFix(EActiveGameplayEffectFix Flag) { return (ActiveGameplayEffectReplicationFix & static_cast<int32>(Flag)) != 0; }
 
+	enum class EActiveGameplayEffectOverflowBehavior : int32
+	{
+		None = 0,
+		HandleOverflowingStacksAlways		= (1 << 0), // Instead of checking the overflow if the ActiveGE stacks are at the limit, we check it whenever there are overflowing stacks
+		ApplyRemainingStacksWhenOverflow	= (1 << 1), // Instead of returning when we deny the application because of overflowing stacks, we apply the stacks up to the limit and ignore the surplus of stacks (stacks over the StackLimit)
+		
+		Legacy = None,							// Default behavior without those new changes
+		Current = HandleOverflowingStacksAlways | ApplyRemainingStacksWhenOverflow,
+	};
+
+	int32 ActiveGameplayEffectOverflowBehavior = (int32)EActiveGameplayEffectOverflowBehavior::Current;
+	FAutoConsoleVariableRef CVarActiveGameplayEffectOverflowBehavior{ TEXT("AbilitySystem.ActiveGameplayEffectOverflowBehavior"), ActiveGameplayEffectOverflowBehavior, TEXT("Change the behavior of overflowing stacks for Active Gameplay Effects. Use 0 to keep legacy functionality"), ECVF_Default };
+	inline bool HasActiveGameplayEffectOverflowBehavior(EActiveGameplayEffectOverflowBehavior Flag) { return (ActiveGameplayEffectOverflowBehavior & static_cast<int32>(Flag)) != 0; }
+
+	// Prediction of stacking gameplay effects
+	bool bAllowPredictiveApplicationOfStackingGameplayEffects = true;
+	FAutoConsoleVariableRef CVarAllowPredictiveStackingGEs{ TEXT("AbilitySystem.AllowPredictiveStackingGEs"), bAllowPredictiveApplicationOfStackingGameplayEffects, TEXT("Allow the predictive application of stacking gameplay effects. The removal of stacks is not predicted (Only application is)"), ECVF_Default };
+
 	TAutoConsoleVariable<int32> CVarGameplayEffectMaxVersion(TEXT("AbilitySystem.GameplayEffects.MaxVersion"), (int32)EGameplayEffectVersion::Current, TEXT("Override the Gameplay Effect Current Version (disabling upgrade code paths)"), ECVF_Default);
 
 	// Fix introduced in UE5.4
@@ -3445,6 +3463,29 @@ void FActiveGameplayEffectsContainer::OnStackCountChange(FActiveGameplayEffect&
 	ActiveEffect.EventSet.OnStackChanged.Broadcast(ActiveEffect.Handle, ActiveEffect.Spec.GetStackCount(), OldStackCount);
 }
 
+void FActiveGameplayEffectsContainer::OnPredictiveGameplayEffectStackCaughtUp(FActiveGameplayEffectHandle Handle)
+{
+	const FActiveGameplayEffect* ActiveGE = GetActiveGameplayEffect(Handle);
+	if (ActiveGE)
+	{
+		const int32 CurrentStackCount = ActiveGE->Spec.GetStackCount();
+		const int32 LastServerReplicatedStackCount = ActiveGE->ClientCachedStackCount;
+
+		// If there is a delta then remove the predicted stacks. 
+		const int32 DeltaStacks = CurrentStackCount - LastServerReplicatedStackCount;
+		if (DeltaStacks != 0)
+		{
+			// We can't predictively remove stacks, so this should always be true
+			ensure(CurrentStackCount > LastServerReplicatedStackCount);
+
+			// We can assume that if the StackCount is different from the ClientCachedStackedCount then the application of the new stacks has been rejected.
+			// Otherwise, PostReplicatedChange would have been called first and would have set the new StackCount and its ClientCachedStackCount
+			Owner->RemoveActiveGameplayEffect_AllowClientRemoval(Handle, DeltaStacks);
+			UE_LOG(LogGameplayEffects, Log, TEXT("Removing local predicted stacks [%d] in %hs. Last Client's stack count [%d] and Current stack count [%d]"), DeltaStacks, __func__, LastServerReplicatedStackCount, CurrentStackCount);
+		}
+	}
+}
+
 /** Called when the duration or starttime of an AGE has changed */
 void FActiveGameplayEffectsContainer::OnDurationChange(FActiveGameplayEffect& Effect)
 {
@@ -3530,10 +3571,15 @@ FActiveGameplayEffect* FActiveGameplayEffectsContainer::FindStackableActiveGamep
 	return StackableGE;
 }
 
-bool FActiveGameplayEffectsContainer::HandleActiveGameplayEffectStackOverflow(const FActiveGameplayEffect& ActiveStackableGE, const FGameplayEffectSpec& OldSpec, const FGameplayEffectSpec& OverflowingSpec)
+bool FActiveGameplayEffectsContainer::HandleActiveGameplayEffectStackOverflow(const FActiveGameplayEffect& ActiveStackableGE, const FGameplayEffectSpec& OldSpec, const FGameplayEffectSpec& OverflowingSpec, FPredictionKey PredictionKey)
 {
 	const UGameplayEffect* StackedGE = OldSpec.Def;
 	const bool bAllowOverflowApplication = !(StackedGE->bDenyOverflowApplication);
+	
+	// allow refreshing duration and context if the new application would cause us to overflow but we're not currently at max stacks
+	using namespace UE::GameplayEffect;
+	const bool bAtStackLimit = OldSpec.GetStackCount() == StackedGE->StackLimitCount;
+	const bool bRefreshToLimit = !bAtStackLimit && !StackedGE->bClearStackOnOverflow && HasActiveGameplayEffectOverflowBehavior(EActiveGameplayEffectOverflowBehavior::ApplyRemainingStacksWhenOverflow);
 
 	for (TSubclassOf<UGameplayEffect> OverflowEffect : StackedGE->OverflowEffects)
 	{
@@ -3541,16 +3587,20 @@ bool FActiveGameplayEffectsContainer::HandleActiveGameplayEffectStackOverflow(co
 		{
 			FGameplayEffectSpec NewGESpec;
 			NewGESpec.InitializeFromLinkedSpec(CDO, OverflowingSpec);
-			Owner->ApplyGameplayEffectSpecToSelf(NewGESpec);
+			Owner->ApplyGameplayEffectSpecToSelf(NewGESpec, PredictionKey);
 		}
 	}
 
-	if (!bAllowOverflowApplication && StackedGE->bClearStackOnOverflow)
+	// don't allow clients to predict removal of active effects due to overflow.
+	if (!PredictionKey.IsLocalClientKey() || IsNetAuthority())
 	{
-		Owner->RemoveActiveGameplayEffect(ActiveStackableGE.Handle);
+		if (!bAllowOverflowApplication && StackedGE->bClearStackOnOverflow)
+		{
+			Owner->RemoveActiveGameplayEffect(ActiveStackableGE.Handle);
+		}
 	}
 
-	return bAllowOverflowApplication;
+	return bAllowOverflowApplication || bRefreshToLimit;
 }
 
 bool FActiveGameplayEffectsContainer::ShouldUseMinimalReplication()
@@ -4013,18 +4063,26 @@ FActiveGameplayEffect* FActiveGameplayEffectsContainer::ApplyGameplayEffectSpec(
 	int32 StartingStackCount = 0;
 	int32 NewStackCount = 0;
 
+	// For HasActiveGameplayEffectOverflowBehavior() 
+	// and EActiveGameplayEffectOverflowBehavior
+	// and bAllowPredictiveApplicationOfStackingGameplayEffects
+	using namespace UE::GameplayEffect;
+
 	// Check if there's an active GE this application should stack upon
 	if (ExistingStackableGE)
 	{
-		if (!IsNetAuthority())
+		if (bAllowPredictiveApplicationOfStackingGameplayEffects == false)
 		{
-			// Don't allow prediction of stacking for now
-			return nullptr;
-		}
-		else
-		{
-			// Server invalidates the prediction key for this GE since client is not predicting it
-			InPredictionKey = FPredictionKey();
+			if (!IsNetAuthority())
+			{
+				// Don't apply the stacking GE if we don't allow it
+				return nullptr;
+			}
+			else
+			{
+				// Server invalidates the prediction key for this GE since client is not predicting it
+				InPredictionKey = FPredictionKey();
+			}
 		}
 
 		bFoundExistingStackableGE = true;
@@ -4038,17 +4096,24 @@ FActiveGameplayEffect* FActiveGameplayEffectsContainer::ApplyGameplayEffectSpec(
 		// How to apply multiple stacks at once? What if we trigger an overflow which can reject the application?
 		// We still want to apply the stacks that didn't push us over, but we also want to call HandleActiveGameplayEffectStackOverflow.
 		
-		// For now: call HandleActiveGameplayEffectStackOverflow only if we are ALREADY at the limit. Else we just clamp stack limit to max.
-		if (ExistingSpec.GetStackCount() == ExistingSpec.Def->StackLimitCount)
+		NewStackCount = ExistingSpec.GetStackCount() + Spec.GetStackCount();
+
+		// Legacy behavior: call HandleActiveGameplayEffectStackOverflow only if we are ALREADY at the limit.
+		const bool bStackCountAtLimit = ExistingSpec.GetStackCount() == ExistingSpec.Def->StackLimitCount;
+
+		// In this path, we call HandleActiveGameplayEffectStackOverflow if we have any overflowing stacks. 
+		const bool bWillOverflowStacks = NewStackCount > ExistingSpec.Def->StackLimitCount && HasActiveGameplayEffectOverflowBehavior(EActiveGameplayEffectOverflowBehavior::HandleOverflowingStacksAlways);
+
+		// We call HandleActiveGameplayEffectStackOverflow if we want to apply the remaining stacks up to the stack limit or for legacy behavior
+		if (bStackCountAtLimit || bWillOverflowStacks)
 		{
-			if (!HandleActiveGameplayEffectStackOverflow(*ExistingStackableGE, ExistingSpec, Spec))
+			if (!HandleActiveGameplayEffectStackOverflow(*ExistingStackableGE, ExistingSpec, Spec, InPredictionKey))
 			{
 				UE_VLOG_UELOG(Owner, LogGameplayEffects, Log, TEXT("Application of %s denied (StackLimit)"), *Spec.ToSimpleString());
 				return nullptr;
 			}
 		}
 		
-		NewStackCount = ExistingSpec.GetStackCount() + Spec.GetStackCount();
 		if (ExistingSpec.Def->StackLimitCount > 0)
 		{
 			NewStackCount = FMath::Min(NewStackCount, ExistingSpec.Def->StackLimitCount);
@@ -4081,7 +4146,8 @@ PRAGMA_ENABLE_DEPRECATION_WARNINGS
 		const UGameplayEffect* GEDef = ExistingSpec.Def;
 
 		// Make sure the GE actually wants to refresh its duration
-		if (GEDef->StackDurationRefreshPolicy == EGameplayEffectStackingDurationPolicy::NeverRefresh)
+		// Client's predicted GE will be combined with the existing GE and we don't want to set any duration timer in this case
+		if (InPredictionKey.IsLocalClientKey() || GEDef->StackDurationRefreshPolicy == EGameplayEffectStackingDurationPolicy::NeverRefresh)
 		{
 			bSetDuration = false;
 		}
@@ -4262,8 +4328,23 @@ PRAGMA_ENABLE_DEPRECATION_WARNINGS
 		MarkArrayDirty();
 
 		// Once replicated state has caught up to this prediction key, we must remove this gameplay effect.
-		InPredictionKey.NewRejectOrCaughtUpDelegate(FPredictionKeyEvent::CreateUObject(Owner, &UAbilitySystemComponent::RemoveActiveGameplayEffect_AllowClientRemoval, AppliedActiveGE->Handle, -1));
-		
+		// This is either not a stackable GE or the first application of a stackable GE 
+		if (ExistingStackableGE == nullptr)
+		{
+			// Once replicated state has caught up to this prediction key, we must remove this gameplay effect.
+			constexpr int32 RemoveAllStacks = -1;
+			InPredictionKey.NewRejectOrCaughtUpDelegate(FPredictionKeyEvent::CreateUObject(Owner, &UAbilitySystemComponent::RemoveActiveGameplayEffect_AllowClientRemoval, AppliedActiveGE->Handle, RemoveAllStacks));
+		}
+		else // We have an existing Stackable GE
+		{
+			// Once replicated state has caught up to this prediction key, we need to readjust the stack count or at least make sure it has been handled in PostReplicatedChange
+			InPredictionKey.NewRejectOrCaughtUpDelegate(FPredictionKeyEvent::CreateUObject(Owner, &UAbilitySystemComponent::OnPredictiveGameplayEffectStackCaughtUp, AppliedActiveGE->Handle));
+
+			// @note: This should be the number of stacks that would have been removed if the application is rejected but we would have been relying on the FPredictionKey reject delegate.
+			// However the reject delegate is not guaranteed to be called with 'bad' network conditions and we would end up with too many stacks on the client.
+			const int32 StacksAppliedWithinThisApplication = FMath::Min(Spec.GetStackCount(), Spec.Def->StackLimitCount - StartingStackCount);
+			UE_LOG(LogGameplayEffects, Log, TEXT("This Gameplay Effect [%s] has applied [%d] stacks when it had to stack."), *AppliedActiveGE->GetDebugString(), StacksAppliedWithinThisApplication);
+		}
 	}
 
 	// @note @todo: This is currently assuming (potentially incorrectly) that the inhibition state of a stacked GE won't change
diff --git a/Engine/Plugins/Runtime/GameplayAbilities/Source/GameplayAbilities/Public/AbilitySystemComponent.h b/Engine/Plugins/Runtime/GameplayAbilities/Source/GameplayAbilities/Public/AbilitySystemComponent.h
--- a/Engine/Plugins/Runtime/GameplayAbilities/Source/GameplayAbilities/Public/AbilitySystemComponent.h
+++ b/Engine/Plugins/Runtime/GameplayAbilities/Source/GameplayAbilities/Public/AbilitySystemComponent.h
@@ -1876,6 +1876,8 @@ protected:
 	/** A version of RemoveActiveGameplayEffect that allows us to remove it even when we don't have authority */
 	UE_API void RemoveActiveGameplayEffect_AllowClientRemoval(FActiveGameplayEffectHandle Handle, int32 StacksToRemove = -1);
 
+	UE_API void OnPredictiveGameplayEffectStackCaughtUp(FActiveGameplayEffectHandle Handle);
+
 	/** Contains all of the gameplay effects that are currently active on this component */
 	UPROPERTY(Replicated)
 	FActiveGameplayEffectsContainer ActiveGameplayEffects;
diff --git a/Engine/Plugins/Runtime/GameplayAbilities/Source/GameplayAbilities/Public/GameplayEffect.h b/Engine/Plugins/Runtime/GameplayAbilities/Source/GameplayAbilities/Public/GameplayEffect.h
--- a/Engine/Plugins/Runtime/GameplayAbilities/Source/GameplayAbilities/Public/GameplayEffect.h
+++ b/Engine/Plugins/Runtime/GameplayAbilities/Source/GameplayAbilities/Public/GameplayEffect.h
@@ -1957,6 +1957,8 @@ private:
 
 	UE_API void OnStackCountChange(FActiveGameplayEffect& ActiveEffect, int32 OldStackCount, int32 NewStackCount);
 
+	UE_API void OnPredictiveGameplayEffectStackCaughtUp(FActiveGameplayEffectHandle Handle);
+
 	UE_API void OnDurationChange(FActiveGameplayEffect& ActiveEffect);
 
 	UE_API void UpdateAllAggregatorModMagnitudes(FActiveGameplayEffect& ActiveEffect);
@@ -1967,7 +1969,7 @@ private:
 	UE_API FActiveGameplayEffect* FindStackableActiveGameplayEffect(const FGameplayEffectSpec& Spec);
 	
 	/** Helper function to handle the case of same-effect stacking overflow; Returns true if the overflow application should apply, false if it should not */
-	UE_API bool HandleActiveGameplayEffectStackOverflow(const FActiveGameplayEffect& ActiveStackableGE, const FGameplayEffectSpec& OldSpec, const FGameplayEffectSpec& OverflowingSpec);
+	UE_API bool HandleActiveGameplayEffectStackOverflow(const FActiveGameplayEffect& ActiveStackableGE, const FGameplayEffectSpec& OldSpec, const FGameplayEffectSpec& OverflowingSpec, FPredictionKey PredictionKey = {});
 
 	UE_API bool ShouldUseMinimalReplication();
 
