Creating the GE inside the lambda causes it to be cleaned up with the lambda, The modifier stack still needs it around because it has hard pointers to data in the GE--As well as a few other places.

diff --git a/Engine/Plugins/Runtime/GameplayAbilities/Source/GameplayAbilities/Private/Tests/GameplayEffectTests.cpp b/Engine/Plugins/Runtime/GameplayAbilities/Source/GameplayAbilities/Private/Tests/GameplayEffectTests.cpp
--- a/Engine/Plugins/Runtime/GameplayAbilities/Source/GameplayAbilities/Private/Tests/GameplayEffectTests.cpp
+++ b/Engine/Plugins/Runtime/GameplayAbilities/Source/GameplayAbilities/Private/Tests/GameplayEffectTests.cpp
@@ -197,22 +197,21 @@ public: // the tests
 		const float ManaBaseValue = DestComponent->GetSet<UAbilitySystemTestAttributeSet>()->Mana.GetBaseValue();
 
 		// Define a common lambda for applying a GE with a specific GameplayModOp
-		auto ApplyGameplayModOp = [this, BuffValue](EGameplayModOp::Type GameplayModOp)
+		auto ApplyGameplayModOp = [this, BuffValue](EGameplayModOp::Type GameplayModOp, UGameplayEffect* InEffect)
 		{
-			CONSTRUCT_CLASS(UGameplayEffect, DamageBuffEffect);
-			AddModifier(DamageBuffEffect, GET_FIELD_CHECKED(UAbilitySystemTestAttributeSet, Mana), GameplayModOp, FScalableFloat(BuffValue));
-			DamageBuffEffect->DurationPolicy = EGameplayEffectDurationType::Infinite;
+			AddModifier(InEffect, GET_FIELD_CHECKED(UAbilitySystemTestAttributeSet, Mana), GameplayModOp, FScalableFloat(BuffValue));
+			InEffect->DurationPolicy = EGameplayEffectDurationType::Infinite;
 
-			return SourceComponent->ApplyGameplayEffectToTarget(DamageBuffEffect, DestComponent, 1.f);
+			return SourceComponent->ApplyGameplayEffectToTarget(InEffect, DestComponent, 1.f);
 		};
 
 		// Define a common lambda for testing if an op applied correctly
-		auto TestGameplayModOp = [this, &ApplyGameplayModOp, BuffValue](EGameplayModOp::Type GameplayModOp)
+		auto TestGameplayModOp = [this, &ApplyGameplayModOp, BuffValue](EGameplayModOp::Type GameplayModOp, UGameplayEffect* InEffect)
 		{
 			const float PrevValue = DestComponent->GetSet<UAbilitySystemTestAttributeSet>()->Mana.GetCurrentValue();
 			const float ExpectedValue = FAggregator::StaticExecModOnBaseValue(PrevValue, GameplayModOp, BuffValue);
 
-			FActiveGameplayEffectHandle AGEHandle = ApplyGameplayModOp(GameplayModOp);
+			FActiveGameplayEffectHandle AGEHandle = ApplyGameplayModOp(GameplayModOp, InEffect);
 			const float CurrentValue = DestComponent->GetSet<UAbilitySystemTestAttributeSet>()->Mana.GetCurrentValue();
 			Test->TestEqual(FString::Printf(TEXT("Attribute GameplayModOp %s"), *UEnum::GetValueAsString(GameplayModOp)), CurrentValue, ExpectedValue);
 			
@@ -220,15 +219,21 @@ public: // the tests
 		};
 
 		// Test all of the ops (order matters here due to the implementation of the above StaticExecModOnBaseValue)
-		FActiveGameplayEffectHandle AdditiveHandle = TestGameplayModOp(EGameplayModOp::Additive);
-		FActiveGameplayEffectHandle MultiplicativeHandle = TestGameplayModOp(EGameplayModOp::Multiplicitive);
-		FActiveGameplayEffectHandle DivisionHandle = TestGameplayModOp(EGameplayModOp::Division);
-		FActiveGameplayEffectHandle CompoundHandle = TestGameplayModOp(EGameplayModOp::MultiplyCompound);
-		FActiveGameplayEffectHandle FinalAddHandle = TestGameplayModOp(EGameplayModOp::AddFinal);
+		CONSTRUCT_CLASS(UGameplayEffect, AdditiveEffect);
+		FActiveGameplayEffectHandle AdditiveHandle = TestGameplayModOp(EGameplayModOp::Additive, AdditiveEffect);
+		CONSTRUCT_CLASS(UGameplayEffect, MultiplyEffect);
+		FActiveGameplayEffectHandle MultiplicativeHandle = TestGameplayModOp(EGameplayModOp::Multiplicitive, MultiplyEffect);
+		CONSTRUCT_CLASS(UGameplayEffect, DivisionEffect);
+		FActiveGameplayEffectHandle DivisionHandle = TestGameplayModOp(EGameplayModOp::Division, DivisionEffect);
+		CONSTRUCT_CLASS(UGameplayEffect, CompoundEffect);
+		FActiveGameplayEffectHandle CompoundHandle = TestGameplayModOp(EGameplayModOp::MultiplyCompound, CompoundEffect);
+		CONSTRUCT_CLASS(UGameplayEffect, FinalEffect);
+		FActiveGameplayEffectHandle FinalAddHandle = TestGameplayModOp(EGameplayModOp::AddFinal, FinalEffect);
 
 		// Test the override quickly here before testing the aggregation is as expected
 		{
-			FActiveGameplayEffectHandle OverrideHandle = TestGameplayModOp(EGameplayModOp::Override);
+			CONSTRUCT_CLASS(UGameplayEffect, OverrideEffect);
+			FActiveGameplayEffectHandle OverrideHandle = TestGameplayModOp(EGameplayModOp::Override, OverrideEffect);
 			DestComponent->RemoveActiveGameplayEffect(OverrideHandle);
 		}
 
@@ -236,10 +241,14 @@ public: // the tests
 
 		// Add some compounding GameplayEffects and manually test our equation to see if it's all compounded correctly.
 		{
-			FActiveGameplayEffectHandle BaseAdd2 = ApplyGameplayModOp(EGameplayModOp::Additive);
-			FActiveGameplayEffectHandle BaseMultiply2 = ApplyGameplayModOp(EGameplayModOp::Multiplicitive);
-			FActiveGameplayEffectHandle Compound2 = ApplyGameplayModOp(EGameplayModOp::MultiplyCompound);
-			FActiveGameplayEffectHandle FinalAdd2 = ApplyGameplayModOp(EGameplayModOp::AddFinal);
+			CONSTRUCT_CLASS(UGameplayEffect, Add2Effect);
+			FActiveGameplayEffectHandle BaseAdd2 = ApplyGameplayModOp(EGameplayModOp::Additive, Add2Effect);
+			CONSTRUCT_CLASS(UGameplayEffect, Multiply2Effect);
+			FActiveGameplayEffectHandle BaseMultiply2 = ApplyGameplayModOp(EGameplayModOp::Multiplicitive, Multiply2Effect);
+			CONSTRUCT_CLASS(UGameplayEffect, Compound2Effect);
+			FActiveGameplayEffectHandle Compound2 = ApplyGameplayModOp(EGameplayModOp::MultiplyCompound, Compound2Effect);
+			CONSTRUCT_CLASS(UGameplayEffect, Final2Effect);
+			FActiveGameplayEffectHandle FinalAdd2 = ApplyGameplayModOp(EGameplayModOp::AddFinal, Final2Effect);
 			
 			// ExpectedResult = ((BaseValue + Additive) * Multiplicative / Division * CompoundMultiply) + FinalAdd;
 			// Multiplicative and Division Compound as: 1.0f + ForEachValue(Value - 1.0f). E.g. two applications of 1.5 = 2.
