Fixing issue in AbilitySystemBlueprintLibrary with checking if Ability is instantiated when it will always return false because it's asking a CDO

#rnx

diff --git a/Engine/Plugins/Runtime/GameplayAbilities/Source/GameplayAbilities/Private/AbilitySystemBlueprintLibrary.cpp b/Engine/Plugins/Runtime/GameplayAbilities/Source/GameplayAbilities/Private/AbilitySystemBlueprintLibrary.cpp
--- a/Engine/Plugins/Runtime/GameplayAbilities/Source/GameplayAbilities/Private/AbilitySystemBlueprintLibrary.cpp
+++ b/Engine/Plugins/Runtime/GameplayAbilities/Source/GameplayAbilities/Private/AbilitySystemBlueprintLibrary.cpp
@@ -1480,7 +1480,7 @@ bool UAbilitySystemBlueprintLibrary::HasAnyAbilitiesByPredicate(
 
 		TArray<UGameplayAbility*> AbilitiesToCheck;
 
-		if (Spec.Ability->IsInstantiated() && !bOnlyRunPredicateOnAbilityCDOs)
+		if (!bOnlyRunPredicateOnAbilityCDOs && Spec.Ability->GetInstancingPolicy() != EGameplayAbilityInstancingPolicy::NonInstanced)
 		{
 			AbilitiesToCheck.Append(Spec.GetAbilityInstances());
 		}
