GE stacks can be predicted so the tests need to be updated.
[RN] Gameplay Cues are properly triggered when a local GameplayEffect stack is predicted.

diff --git a/Engine/Plugins/Runtime/GameplayAbilities/Source/GameplayAbilities/Private/AbilitySystemComponent.cpp b/Engine/Plugins/Runtime/GameplayAbilities/Source/GameplayAbilities/Private/AbilitySystemComponent.cpp
--- a/Engine/Plugins/Runtime/GameplayAbilities/Source/GameplayAbilities/Private/AbilitySystemComponent.cpp
+++ b/Engine/Plugins/Runtime/GameplayAbilities/Source/GameplayAbilities/Private/AbilitySystemComponent.cpp
@@ -925,12 +925,18 @@ FActiveGameplayEffectHandle UAbilitySystemComponent::ApplyGameplayEffectSpecToSe
 	if (!bSuppressGameplayCues && !Spec.Def->bSuppressStackingCues && bFoundExistingStackableGE && AppliedEffect && !AppliedEffect->bIsInhibited)
 	{
 		ensureMsgf(OurCopyOfSpec, TEXT("OurCopyOfSpec will always be valid if bFoundExistingStackableGE"));
-		if (OurCopyOfSpec && OurCopyOfSpec->GetStackCount() > Spec.GetStackCount())
+		if (OurCopyOfSpec && OurCopyOfSpec->GetStackCount() >= Spec.GetStackCount())
 		{
-			// Because PostReplicatedChange will get called from modifying the stack count
-			// (and not PostReplicatedAdd) we won't know which GE was modified.
-			// So instead we need to explicitly RPC the client so it knows the GC needs updating
-			UAbilitySystemGlobals::Get().GetGameplayCueManager()->InvokeGameplayCueAddedAndWhileActive_FromSpec(this, *OurCopyOfSpec, PredictionKey);
+			// If we're predicting we want to activate on ourselves
+			if (!IsOwnerActorAuthoritative() && PredictionKey.IsLocalClientKey())
+			{
+				InvokeGameplayCueEvent(*OurCopyOfSpec, EGameplayCueEvent::OnActive);
+				InvokeGameplayCueEvent(*OurCopyOfSpec, EGameplayCueEvent::WhileActive);
+			}
+			else
+			{
+				UAbilitySystemGlobals::Get().GetGameplayCueManager()->InvokeGameplayCueAddedAndWhileActive_FromSpec(this, *OurCopyOfSpec, PredictionKey);
+			}
 		}
 	}
 	
diff --git a/Engine/Plugins/Runtime/GameplayAbilities/Source/GameplayAbilities/Private/Tests/GameplayEffectTests.cpp b/Engine/Plugins/Runtime/GameplayAbilities/Source/GameplayAbilities/Private/Tests/GameplayEffectTests.cpp
--- a/Engine/Plugins/Runtime/GameplayAbilities/Source/GameplayAbilities/Private/Tests/GameplayEffectTests.cpp
+++ b/Engine/Plugins/Runtime/GameplayAbilities/Source/GameplayAbilities/Private/Tests/GameplayEffectTests.cpp
@@ -600,7 +600,7 @@ public: // the tests
 				// Unexpected behavior:  We can never predict a stacked application (why?!) so it ends up with 1.
 				const bool bPredicting = !DestComponent->IsOwnerActorAuthoritative();
 				bool bExpectedAdded = (!DestComponent->bSuppressGameplayCues);
-				int ExpectedOnActive = bExpectedAdded ? (bPredicting ? 1 : StackLimitCount) : 0; // can only predict the first cue
+				int ExpectedOnActive = bExpectedAdded ? StackLimitCount : 0;
 				int ExpectedOnExecute = 0;
 				int ExpectedWhileActive = ExpectedOnActive;
 				int ExpectedOnRemove = (ExpectedOnActive > 0) ? 1 : 0;		// Even with stacking, the GameplayCue is actually only applied once (and therefore removed once when fully unstacked)
