Improve the worst-case perf hit from the common async actions used by Minifigs during MakeNewObject with a default "None" name passed in
- By passing no name in, this calls MakeUniqueObjectName with just the class of the async object which then probes the name table iteratively to find a unique object name with the next free_X.
- This causes a lot of spinning time in this infinite loop when there have been a lot of transient actors created and it has to iterate up to a large value of _X to find a unique name.
- Instead, we generate our own unique name, using the WorldContextObj (or passed in Actor) name, which for Actors like Minifig is itself already a unique name.  We concatenate our async action classname to that, and then call MakeUniqueObjectName with that as the base name.  This will search much fewer names, since we are dividing the namespace up by each actor's name, resulting in lower _X names overall.

diff --git a/Engine/Plugins/Runtime/GameplayAbilities/Source/GameplayAbilities/Private/Abilities/Async/AbilityAsync_WaitGameplayTagCountChanged.cpp b/Engine/Plugins/Runtime/GameplayAbilities/Source/GameplayAbilities/Private/Abilities/Async/AbilityAsync_WaitGameplayTagCountChanged.cpp
--- a/Engine/Plugins/Runtime/GameplayAbilities/Source/GameplayAbilities/Private/Abilities/Async/AbilityAsync_WaitGameplayTagCountChanged.cpp
+++ b/Engine/Plugins/Runtime/GameplayAbilities/Source/GameplayAbilities/Private/Abilities/Async/AbilityAsync_WaitGameplayTagCountChanged.cpp
@@ -3,6 +3,7 @@
 #include "Abilities/Async/AbilityAsync_WaitGameplayTagCountChanged.h"
 #include "AbilitySystemComponent.h"
 #include "AbilitySystemLog.h"
+#include "UObject/Package.h"
 
 void UAbilityAsync_WaitGameplayTagCountChanged::Activate()
 {
@@ -44,8 +45,21 @@ void UAbilityAsync_WaitGameplayTagCountChanged::GameplayTagCallback(const FGamep
 }
 
 UAbilityAsync_WaitGameplayTagCountChanged* UAbilityAsync_WaitGameplayTagCountChanged::WaitGameplayTagCountChangedOnActor(AActor* TargetActor, FGameplayTag Tag)
-{
-	UAbilityAsync_WaitGameplayTagCountChanged* MyObj = NewObject<UAbilityAsync_WaitGameplayTagCountChanged>();
+{	
+	UPackage* TransientPackage = GetTransientPackage();
+	FName UniqueActionName = NAME_None;
+
+	// Use the name of the passed in TargetActor to generate a more unique basename for this Action 
+	// by concatenating that AActor's name with our classname.  This is to avoid spinning inside MakeUniqueObjectName 
+	// probing for a valid name since they're all just ClassName_x, where x can get pretty large
+	if (TargetActor != nullptr)
+	{
+		const UClass* ThisClass = UAbilityAsync_WaitGameplayTagCountChanged::StaticClass();
+		const FName ActionNameConcat = FName(*FString::Printf(TEXT("%s_%s"), *TargetActor->GetName(), *ThisClass->GetName()));
+		UniqueActionName = ::MakeUniqueObjectName(TransientPackage, ThisClass, ActionNameConcat);
+	}
+
+	UAbilityAsync_WaitGameplayTagCountChanged* MyObj = NewObject<UAbilityAsync_WaitGameplayTagCountChanged>(TransientPackage, UniqueActionName);
 	MyObj->SetAbilityActor(TargetActor);
 	MyObj->Tag = Tag;
 
