SetTagMapCount should call OnTagUpdated like update tag map does.
Fixed TODO in FMJam.
#12283
[RN] Fixed issue where SetTagMapCount didn't call OnTagUpdated like the other tag update functions. This will now be called when Gameplay Cue tags and Loose Tags are added and removed (No change for changing values that don't add or remove the tag)

diff --git a/Engine/Plugins/Runtime/GameplayAbilities/Source/GameplayAbilities/Public/AbilitySystemComponent.h b/Engine/Plugins/Runtime/GameplayAbilities/Source/GameplayAbilities/Public/AbilitySystemComponent.h
--- a/Engine/Plugins/Runtime/GameplayAbilities/Source/GameplayAbilities/Public/AbilitySystemComponent.h
+++ b/Engine/Plugins/Runtime/GameplayAbilities/Source/GameplayAbilities/Public/AbilitySystemComponent.h
@@ -614,7 +614,10 @@ class UAbilitySystemComponent : public UGameplayTasksComponent, public IGameplay
 	/** Forcibly sets the number of instances of a given tag */
 	FORCEINLINE void SetTagMapCount(const FGameplayTag& Tag, int32 NewCount)
 	{
-		GameplayTagCountContainer.SetTagCount(Tag, NewCount);
+		if (GameplayTagCountContainer.SetTagCount(Tag, NewCount))
+		{
+			OnTagUpdated(Tag, NewCount > 0);
+		}
 	}
 
 	/** Update the number of instances of a given tag and calls callback */
