Clarify comments for AbilitySystemInterface to explain why you shouldn't call it directly

diff --git a/Engine/Plugins/Runtime/GameplayAbilities/Source/GameplayAbilities/Public/AbilitySystemInterface.h b/Engine/Plugins/Runtime/GameplayAbilities/Source/GameplayAbilities/Public/AbilitySystemInterface.h
--- a/Engine/Plugins/Runtime/GameplayAbilities/Source/GameplayAbilities/Public/AbilitySystemInterface.h
+++ b/Engine/Plugins/Runtime/GameplayAbilities/Source/GameplayAbilities/Public/AbilitySystemInterface.h
@@ -11,7 +11,11 @@
 
 class UAbilitySystemComponent;
 
-/** Interface for actors that expose access to an ability system component */
+/** 
+ * Interface for native actors to expose access to an ability system component. 
+ * This is used by native actors to efficiently return a component and handle more complicated cases. This interface should never be called directly, 
+ * instead call GetAbilitySystemComponentFromActor or GetAbilitySystemComponent on the AbilitySystemBlueprintLibrary as that also handles looking up by component type.
+ */
 UINTERFACE(MinimalAPI, meta = (CannotImplementInterfaceInBlueprint))
 class UAbilitySystemInterface : public UInterface
 {
