PR #13311: Fix AttributeSet.h comment typos

diff --git a/Engine/Plugins/Runtime/GameplayAbilities/Source/GameplayAbilities/Public/AttributeSet.h b/Engine/Plugins/Runtime/GameplayAbilities/Source/GameplayAbilities/Public/AttributeSet.h
--- a/Engine/Plugins/Runtime/GameplayAbilities/Source/GameplayAbilities/Public/AttributeSet.h
+++ b/Engine/Plugins/Runtime/GameplayAbilities/Source/GameplayAbilities/Public/AttributeSet.h
@@ -16,7 +16,7 @@ class UCurveTable;
 struct FGameplayAbilityActorInfo;
 struct FAggregator;
 
-/** Place in an AttributeSet to create an attribute that can be accesed using FGameplayAttribute. It is strongly encouraged to use this instead of raw float attributes */
+/** Place in an AttributeSet to create an attribute that can be accessed using FGameplayAttribute. It is strongly encouraged to use this instead of raw float attributes */
 USTRUCT(BlueprintType)
 struct FGameplayAttributeData
 {
@@ -193,13 +193,13 @@ public:
 
 	/**
 	 *	Called just before modifying the value of an attribute. AttributeSet can make additional modifications here. Return true to continue, or false to throw out the modification.
-	 *	Note this is only called during an 'execute'. E.g., a modification to the 'base value' of an attribute. It is not called during an application of a GameplayEffect, such as a 5 ssecond +10 movement speed buff.
+	 *	Note this is only called during an 'execute'. E.g., a modification to the 'base value' of an attribute. It is not called during an application of a GameplayEffect, such as a 5 second +10 movement speed buff.
 	 */	
 	virtual bool PreGameplayEffectExecute(struct FGameplayEffectModCallbackData &Data) { return true; }
 	
 	/**
 	 *	Called just after a GameplayEffect is executed to modify the base value of an attribute. No more changes can be made.
-	 *	Note this is only called during an 'execute'. E.g., a modification to the 'base value' of an attribute. It is not called during an application of a GameplayEffect, such as a 5 ssecond +10 movement speed buff.
+	 *	Note this is only called during an 'execute'. E.g., a modification to the 'base value' of an attribute. It is not called during an application of a GameplayEffect, such as a 5 second +10 movement speed buff.
 	 */
 	virtual void PostGameplayEffectExecute(const struct FGameplayEffectModCallbackData &Data) { }
 
@@ -409,7 +409,7 @@ private:
 
 /**
  * This defines a set of helper functions for accessing and initializing attributes, to avoid having to manually write these functions.
- * It would creates the following functions, for attribute Health
+ * It would create the following functions, for attribute Health
  *
  *	static FGameplayAttribute UMyHealthSet::GetHealthAttribute();
  *	FORCEINLINE float UMyHealthSet::GetHealth() const;
