PR #13324: Fix some comments inside AbilityTypes.h. Add better naming for triggers

diff --git a/Engine/Plugins/Runtime/GameplayAbilities/Source/GameplayAbilities/Public/Abilities/GameplayAbilityTypes.h b/Engine/Plugins/Runtime/GameplayAbilities/Source/GameplayAbilities/Public/Abilities/GameplayAbilityTypes.h
--- a/Engine/Plugins/Runtime/GameplayAbilities/Source/GameplayAbilities/Public/Abilities/GameplayAbilityTypes.h
+++ b/Engine/Plugins/Runtime/GameplayAbilities/Source/GameplayAbilities/Public/Abilities/GameplayAbilityTypes.h
@@ -41,14 +41,16 @@ namespace EGameplayAbilityInstancingPolicy
 	 */
 	enum Type : int
 	{
-		// This ability is never instanced. Anything that executes the ability is operating on the CDO.
+		/** This ability is never instanced. Anything that executes the ability is operating on the CDO. */
 		NonInstanced UE_DEPRECATED_FORGAME(5.5, "Use InstancedPerActor as the default to avoid confusing corner cases"),
 
-		// Each actor gets their own instance of this ability. State can be saved, replication is possible.
-		InstancedPerActor,
+		/** Each actor gets it's own single instance of this ability, which supports replication to its owning actor to allow RPC's
+		 * within the ability. As it is a single instance, only one ability maybe active at any time. State is saved between activations. */
+		InstancedPerActor 		UMETA(DisplayName = "Instanced Per Actor"),
 
-		// We instance this ability each time it is executed. Replication currently unsupported.
-		InstancedPerExecution,
+		/** We instance this ability each time it is executed. Replication currently unsupported. These can have multiple running
+		 * at the same time, and no state is saved between execution. */
+		InstancedPerExecution		UMETA(DisplayName = "Instanced Per Execution"),
 	};
 }
 
@@ -58,16 +60,16 @@ namespace EGameplayAbilityNetExecutionPolicy
 	/** Where does an ability execute on the network. Does a client "ask and predict", "ask and wait", "don't ask (just do it)" */
 	enum Type : int
 	{
-		// Part of this ability runs predictively on the local client if there is one
+		/** Part of this ability runs predictively on the local client if there is one */
 		LocalPredicted		UMETA(DisplayName = "Local Predicted"),
 
-		// This ability will only run on the client or server that has local control
+		/** This ability will only run on the client or server that has local control */
 		LocalOnly			UMETA(DisplayName = "Local Only"),
 
-		// This ability is initiated by the server, but will also run on the local client if one exists
+		/** This ability is initiated by the server, but will also run on the local client if one exists */
 		ServerInitiated		UMETA(DisplayName = "Server Initiated"),
 
-		// This ability will only run on the server
+		/** This ability will only run on the server */
 		ServerOnly			UMETA(DisplayName = "Server Only"),
 	};
 }
@@ -78,16 +80,16 @@ namespace EGameplayAbilityNetSecurityPolicy
 	/** What protections does this ability have? Should the client be allowed to request changes to the execution of the ability? */
 	enum Type : int
 	{
-		// No security requirements. Client or server can trigger execution and termination of this ability freely.
+		/** No security requirements. Client or server can trigger execution and termination of this ability freely. */
 		ClientOrServer			UMETA(DisplayName = "Client Or Server"),
 
-		// A client requesting execution of this ability will be ignored by the server. Clients can still request that the server cancel or end this ability.
+		/** A client requesting execution of this ability will be ignored by the server. Clients can still request that the server cancel or end this ability. */
 		ServerOnlyExecution		UMETA(DisplayName = "Server Only Execution"),
 
-		// A client requesting cancellation or ending of this ability will be ignored by the server. Clients can still request execution of the ability.
+		/** A client requesting cancellation or ending of this ability will be ignored by the server. Clients can still request execution of the ability. */
 		ServerOnlyTermination	UMETA(DisplayName = "Server Only Termination"),
 
-		// Server controls both execution and termination of this ability. A client making any requests will be ignored.
+		/** Server controls both execution and termination of this ability. A client making any requests will be ignored. */
 		ServerOnly				UMETA(DisplayName = "Server Only"),
 	};
 }
@@ -98,10 +100,10 @@ namespace EGameplayAbilityReplicationPolicy
 	/** How an ability replicates state/events to everyone on the network */
 	enum Type : int
 	{
-		// We don't replicate the instance of the ability to anyone.
+		/** We don't replicate the instance of the ability to anyone. */
 		ReplicateNo			UMETA(DisplayName = "Do Not Replicate"),
 
-		// We replicate the instance of the ability to the owner.
+		/** We replicate the instance of the ability to the owner. */
 		ReplicateYes		UMETA(DisplayName = "Replicate"),
 	};
 }
@@ -112,14 +114,14 @@ namespace EGameplayAbilityTriggerSource
 	/**	Defines what type of trigger will activate the ability, paired to a tag */
 	enum Type : int
 	{
-		// Triggered from a gameplay event, will come with payload
-		GameplayEvent,
-
-		// Triggered if the ability's owner gets a tag added, triggered once whenever it's added
-		OwnedTagAdded,
-
-		// Triggered if the ability's owner gets tag added, removed when the tag is removed
-		OwnedTagPresent,
+		/** Triggered by an external gameplay event. The ability will receive a GameplayEvent payload. */
+		GameplayEvent      UMETA(DisplayName = "On Gameplay Event"),
+		
+		/** Triggered when the owner gains the specified tag. Will not cancel when the tag is removed. */
+		OwnedTagAdded      UMETA(DisplayName = "When Tag Is Added"),
+		
+		/** Triggered when the owner has the specified tag. The ability will be canceled if the tag is later removed. */
+		OwnedTagPresent    UMETA(DisplayName = "When Tag Is Present"),
 	};
 }
 
@@ -593,7 +595,7 @@ struct FAbilityReplicatedDataCache
 
 /** 
  *	Associative container of GameplayAbilitySpecs + PredictionKeys --> FAbilityReplicatedDataCache. Basically, it holds replicated data on the ability system component that abilities access in their scripting.
- *	This was refactored from a normal TMap. This mainly servers to:
+ *	This was refactored from a normal TMap. This mainly serves to:
  *		1. Return shared ptrs to the cached data so that callsites are not vulnerable to the underlying map shifting around (E.g invoking a replicated event ends the ability or activates a new one and causes memory to move, invalidating the pointer).
  *		2. Data is cleared on ability end via ::Remove.
  *		3. The FAbilityReplicatedDataCache instances are recycled rather than allocated each time via ::FreeData.
