Gameplay Effects stacking options have better names and comments
#13325

diff --git a/Engine/Plugins/Runtime/GameplayAbilities/Source/GameplayAbilities/Public/GameplayEffectTypes.h b/Engine/Plugins/Runtime/GameplayAbilities/Source/GameplayAbilities/Public/GameplayEffectTypes.h
--- a/Engine/Plugins/Runtime/GameplayAbilities/Source/GameplayAbilities/Public/GameplayEffectTypes.h
+++ b/Engine/Plugins/Runtime/GameplayAbilities/Source/GameplayAbilities/Public/GameplayEffectTypes.h
@@ -171,16 +171,18 @@ namespace GameplayEffectUtilities
 	GAMEPLAYABILITIES_API float ComputeStackedModifierMagnitude(float BaseComputedMagnitude, int32 StackCount, EGameplayModOp::Type ModOp);
 }
 
-/** Enumeration for ways a single GameplayEffect asset can stack. */
+/** Defines how multiple instances of the same GameplayEffect definition stack when applied (non-instant GE's). */
 UENUM()
 enum class EGameplayEffectStackingType : uint8
 {
-	/** No stacking. Multiple applications of this GameplayEffect are treated as separate instances. */
-	None,
-	/** Each caster has its own stack. */
-	AggregateBySource,
-	/** Each target has its own stack. */
-	AggregateByTarget,
+	/** No stacking. Each application is a separate instance. */
+	None UMETA(DisplayName = "No Stacking"),
+
+	/** Stacks only when the same source reapplies the effect. */
+	AggregateBySource UMETA(DisplayName = "Stack Per Source"),
+
+	/** Stacks only when the effect is reapplied to the same target. */
+	AggregateByTarget UMETA(DisplayName = "Stack Per Target"),
 };
 
 
