Warn when Instant GE has a stack type set.
Deprecate StackType to make it private.
#13326
[RN] StackType deprecated to make private in the future, use new accessors instead.

diff --git a/Engine/Plugins/Runtime/GameplayAbilities/Source/GameplayAbilities/Private/AbilitySystemComponent.cpp b/Engine/Plugins/Runtime/GameplayAbilities/Source/GameplayAbilities/Private/AbilitySystemComponent.cpp
--- a/Engine/Plugins/Runtime/GameplayAbilities/Source/GameplayAbilities/Private/AbilitySystemComponent.cpp
+++ b/Engine/Plugins/Runtime/GameplayAbilities/Source/GameplayAbilities/Private/AbilitySystemComponent.cpp
@@ -2696,7 +2696,7 @@ void UAbilitySystemComponent::Debug_Internal(FAbilitySystemComponentDebugInfo& I
 			if (ActiveGE.Spec.GetStackCount() > 1)
 			{
 
-				if (ActiveGE.Spec.Def->StackingType == EGameplayEffectStackingType::AggregateBySource)
+				if (ActiveGE.Spec.Def->GetStackingType() == EGameplayEffectStackingType::AggregateBySource)
 				{
 					StackString = FString::Printf(TEXT("(Stacks: %d. From: %s) "), ActiveGE.Spec.GetStackCount(), *GetNameSafe(ActiveGE.Spec.GetContext().GetInstigatorAbilitySystemComponent()->GetAvatarActor_Direct()));
 				}
diff --git a/Engine/Plugins/Runtime/GameplayAbilities/Source/GameplayAbilities/Private/GameplayEffect.cpp b/Engine/Plugins/Runtime/GameplayAbilities/Source/GameplayAbilities/Private/GameplayEffect.cpp
--- a/Engine/Plugins/Runtime/GameplayAbilities/Source/GameplayAbilities/Private/GameplayEffect.cpp
+++ b/Engine/Plugins/Runtime/GameplayAbilities/Source/GameplayAbilities/Private/GameplayEffect.cpp
@@ -310,6 +310,13 @@ EDataValidationResult UGameplayEffect::IsDataValid(FDataValidationContext& Conte
 		}
 	}
 
+	PRAGMA_DISABLE_DEPRECATION_WARNINGS
+	if (DurationPolicy == EGameplayEffectDurationType::Instant && StackingType != EGameplayEffectStackingType::None)
+	{
+		Context.AddWarning(LOCTEXT("InstantNonStackingBehaviour", "Instant GameplayEffect's will not Stack."));
+	}
+	PRAGMA_ENABLE_DEPRECATION_WARNINGS
+
 	TArray<FText> Warnings, Errors;
 	Context.SplitIssues(Warnings, Errors);
 
@@ -836,6 +843,23 @@ void UGameplayEffect::ConvertUIComponent()
 	}
 }
 
+EGameplayEffectStackingType UGameplayEffect::GetStackingType() const
+{
+	if (DurationPolicy == EGameplayEffectDurationType::Instant)
+	{
+		return EGameplayEffectStackingType::None;
+	}
+	
+	return StackingType;
+}
+
+#if WITH_EDITOR
+void UGameplayEffect::SetStackingType(EGameplayEffectStackingType InType)
+{
+	StackingType = InType;
+}
+#endif
+
 void UGameplayEffect::PreSave(FObjectPreSaveContext SaveContext)
 {
 	Super::PreSave(SaveContext);
@@ -3549,7 +3573,7 @@ FActiveGameplayEffect* FActiveGameplayEffectsContainer::FindStackableActiveGamep
 {
 	FActiveGameplayEffect* StackableGE = nullptr;
 	const UGameplayEffect* GEDef = Spec.Def;
-	EGameplayEffectStackingType StackingType = GEDef->StackingType;
+	EGameplayEffectStackingType StackingType = GEDef->GetStackingType();
 
 	if ((StackingType != EGameplayEffectStackingType::None) && (GEDef->DurationPolicy != EGameplayEffectDurationType::Instant))
 	{
diff --git a/Engine/Plugins/Runtime/GameplayAbilities/Source/GameplayAbilities/Private/Tests/GameplayEffectTests.cpp b/Engine/Plugins/Runtime/GameplayAbilities/Source/GameplayAbilities/Private/Tests/GameplayEffectTests.cpp
--- a/Engine/Plugins/Runtime/GameplayAbilities/Source/GameplayAbilities/Private/Tests/GameplayEffectTests.cpp
+++ b/Engine/Plugins/Runtime/GameplayAbilities/Source/GameplayAbilities/Private/Tests/GameplayEffectTests.cpp
@@ -341,7 +341,7 @@ public: // the tests
 		StackingEffect->DurationPolicy = EGameplayEffectDurationType::HasDuration;
 		StackingEffect->DurationMagnitude = FGameplayEffectModifierMagnitude(FScalableFloat(Duration));
 		StackingEffect->StackLimitCount = StackLimit;
-		StackingEffect->StackingType = EGameplayEffectStackingType::AggregateByTarget;
+		StackingEffect->SetStackingType(EGameplayEffectStackingType::AggregateByTarget);
 		StackingEffect->StackDurationRefreshPolicy = EGameplayEffectStackingDurationPolicy::NeverRefresh;
 		StackingEffect->StackExpirationPolicy = EGameplayEffectStackingExpirationPolicy::ClearEntireStack;
 
@@ -372,7 +372,7 @@ public: // the tests
 		StackingEffect->DurationPolicy = EGameplayEffectDurationType::HasDuration;
 		StackingEffect->DurationMagnitude = FGameplayEffectModifierMagnitude(SetByCallerDuration);
 		StackingEffect->StackLimitCount = StackLimit;
-		StackingEffect->StackingType = EGameplayEffectStackingType::AggregateByTarget;
+		StackingEffect->SetStackingType(EGameplayEffectStackingType::AggregateByTarget);
 		StackingEffect->StackDurationRefreshPolicy = EGameplayEffectStackingDurationPolicy::NeverRefresh;
 		StackingEffect->StackExpirationPolicy = EGameplayEffectStackingExpirationPolicy::RemoveSingleStackAndRefreshDuration;
 
@@ -421,7 +421,7 @@ public: // the tests
 		CONSTRUCT_CLASS(UGameplayEffect, TestGameplayCuesGE);
 		TestGameplayCuesGE->GameplayCues.Emplace(GameplayEffectCue);
 		TestGameplayCuesGE->DurationPolicy = EGameplayEffectDurationType::Infinite;
-		TestGameplayCuesGE->StackingType = EGameplayEffectStackingType::AggregateByTarget;
+		TestGameplayCuesGE->SetStackingType(EGameplayEffectStackingType::AggregateByTarget);
 		TestGameplayCuesGE->StackLimitCount = 5; // Make it high enough to not confuse between "expected 1 and got 2" etc.
 
 		UGameplayCueNotify_UnitTest* GCNotify_Test_CDO = GetMutableDefault<UGameplayCueNotify_UnitTest>();
@@ -541,7 +541,7 @@ public: // the tests
 		CONSTRUCT_CLASS(UGameplayEffect, TestGameplayCuesGE);
 		TestGameplayCuesGE->GameplayCues.Emplace(GameplayEffectCue);
 		TestGameplayCuesGE->DurationPolicy = EGameplayEffectDurationType::Infinite;
-		TestGameplayCuesGE->StackingType = EGameplayEffectStackingType::AggregateByTarget;
+		TestGameplayCuesGE->SetStackingType(EGameplayEffectStackingType::AggregateByTarget);
 		TestGameplayCuesGE->StackDurationRefreshPolicy = EGameplayEffectStackingDurationPolicy::NeverRefresh;
 		TestGameplayCuesGE->StackPeriodResetPolicy = EGameplayEffectStackingPeriodPolicy::NeverReset;
 		TestGameplayCuesGE->StackLimitCount = 5; // Make it high enough to not confuse between "expected 1 and got 2" etc.
diff --git a/Engine/Plugins/Runtime/GameplayAbilities/Source/GameplayAbilities/Public/GameplayEffect.h b/Engine/Plugins/Runtime/GameplayAbilities/Source/GameplayAbilities/Public/GameplayEffect.h
--- a/Engine/Plugins/Runtime/GameplayAbilities/Source/GameplayAbilities/Public/GameplayEffect.h
+++ b/Engine/Plugins/Runtime/GameplayAbilities/Source/GameplayAbilities/Public/GameplayEffect.h
@@ -2079,7 +2079,7 @@ public:
 	// These are deprecated but remain for backwards compat, please use FGameplayEffectConstants:: instead.
 	static UE_API const float INFINITE_DURATION;
 	static UE_API const float INSTANT_APPLICATION;
-	static UE_API const float NO_PERIOD;	
+	static UE_API const float NO_PERIOD;
 	static UE_API const float INVALID_LEVEL;
 
 	UE_DEPRECATED(5.3, "The implementation and method name did not match.  Use GetGrantedTags() to get the tags Granted to the Actor this GameplayEffect is applied to.")
@@ -2371,7 +2371,9 @@ public:
 	//	Stacking
 	// ----------------------------------------------------------------------
 
+	// TODO: 5.11 Add EditCondition DurationPolicy to Stacking rules
 	/** How this GameplayEffect stacks with other instances of this same GameplayEffect */
+	UE_DEPRECATED(5.7, "Stacking Type will be made private, Please use GetStackingType.")
 	UPROPERTY(EditDefaultsOnly, BlueprintReadOnly, Category = Stacking)
 	EGameplayEffectStackingType	StackingType;
 
@@ -2395,10 +2397,18 @@ public:
 	UPROPERTY(EditDefaultsOnly, BlueprintReadOnly, Category = Stacking, meta = (EditConditionHides, EditCondition = "StackingType != EGameplayEffectStackingType::None"))
 	bool bFactorInStackCount;
 
+	/** return the stacking type defined from the data asset */
+	EGameplayEffectStackingType GetStackingType() const;
+
+#if WITH_EDITOR
+	/** Set stacking type */
+	void SetStackingType(EGameplayEffectStackingType InType);
+#endif
+
 	// ----------------------------------------------------------------------
 	//	Granted abilities
 	// ----------------------------------------------------------------------
-	
+
 	/** Policy for what abilities this GE will grant. */
 	UE_DEPRECATED(5.3, "GrantedAbilities are deprecated in favor of AbilitiesGameplayEffectComponent")
 	UPROPERTY(BlueprintReadOnly, Category = "Granted Abilities")
