GameplayAbilities: Avoid nullptr dereference in PostReplicatedAdd()

[RN] Fixed a crash in FActiveGameplayEffect::PostReplicatedAdd()

diff --git a/Engine/Plugins/Runtime/GameplayAbilities/Source/GameplayAbilities/Private/GameplayEffect.cpp b/Engine/Plugins/Runtime/GameplayAbilities/Source/GameplayAbilities/Private/GameplayEffect.cpp
--- a/Engine/Plugins/Runtime/GameplayAbilities/Source/GameplayAbilities/Private/GameplayEffect.cpp
+++ b/Engine/Plugins/Runtime/GameplayAbilities/Source/GameplayAbilities/Private/GameplayEffect.cpp
@@ -2832,7 +2832,15 @@ void FActiveGameplayEffect::PostReplicatedAdd(const struct FActiveGameplayEffect
 	// Handles are not replicated, so create a new one.
 	Handle = FActiveGameplayEffectHandle::GenerateNewHandle(InArray.Owner);
 
-	UE_VLOG_UELOG(InArray.Owner->GetOwnerActor(), LogGameplayEffects, Verbose, TEXT("%s (Non-Auth): %s. Pending( OnActive: %d WhileActive: %d )"), ANSI_TO_TCHAR(__func__), *GetDebugString(), bPendingRepOnActiveGC, bPendingRepWhileActiveGC);
+	UE_VLOG_UELOG(
+		IsValid(InArray.Owner) ? InArray.Owner->GetOwnerActor() : nullptr,
+		LogGameplayEffects,
+		Verbose,
+		TEXT("%s (Non-Auth): %s. Pending( OnActive: %d WhileActive: %d )"),
+		ANSI_TO_TCHAR(__func__),
+		*GetDebugString(),
+		bPendingRepOnActiveGC,
+		bPendingRepWhileActiveGC);
 
 	// Do stuff for adding GEs (add mods, tags, *invoke callbacks*).  But do NOT invoke the GameplayCues as we don't know if this GE ends up inhibited or not (thus the bPendingRepOnActiveGC variables).
 	constexpr bool bInvokeGameplayCueEvents = false;
@@ -4411,7 +4419,10 @@ void FActiveGameplayEffectsContainer::InternalOnActiveGameplayEffectAdded(FActiv
 	Effect.bIsInhibited = true; // Effect has to start inhibited, so our call to Inhibit will trigger if we should be active
 
 	FActiveGameplayEffectHandle EffectHandle = Effect.Handle;
-	Owner->SetActiveGameplayEffectInhibit(MoveTemp(EffectHandle), !bActive, bInvokeGameplayCueEvents);
+	if (IsValid(Owner))
+	{
+		Owner->SetActiveGameplayEffectInhibit(MoveTemp(EffectHandle), !bActive, bInvokeGameplayCueEvents);
+	}
 }
 
 void FActiveGameplayEffectsContainer::AddActiveGameplayEffectGrantedTagsAndModifiers(FActiveGameplayEffect& Effect, bool bInvokeGameplayCueEvents)
@@ -5107,7 +5118,7 @@ PRAGMA_ENABLE_DEPRECATION_WARNINGS
 
 void FActiveGameplayEffectsContainer::Uninitialize()
 {
-	UWorld* World = Owner->GetWorld();
+	UWorld* World = IsValid(Owner) ? Owner->GetWorld() : nullptr;
 	for (FActiveGameplayEffect& CurEffect : this)
 	{
 		RemoveCustomMagnitudeExternalDependencies(CurEffect);
@@ -5130,8 +5141,16 @@ void FActiveGameplayEffectsContainer::Uninitialize()
 
 bool FActiveGameplayEffectsContainer::IsServerWorldTimeAvailable() const
 {
+	if (!IsValid(Owner))
+	{
+		return false;
+	}
+
 	UWorld* World = Owner->GetWorld();
-	check(World);
+	if (!World)
+	{
+		return false;
+	}
 
 	AGameStateBase* GameState = World->GetGameState();
 	return (GameState != nullptr);
