Fix GetStackingType being in a WITH_EDITOR block
#rnx

diff --git a/Engine/Plugins/Runtime/GameplayAbilities/Source/GameplayAbilities/Private/GameplayEffect.cpp b/Engine/Plugins/Runtime/GameplayAbilities/Source/GameplayAbilities/Private/GameplayEffect.cpp
--- a/Engine/Plugins/Runtime/GameplayAbilities/Source/GameplayAbilities/Private/GameplayEffect.cpp
+++ b/Engine/Plugins/Runtime/GameplayAbilities/Source/GameplayAbilities/Private/GameplayEffect.cpp
@@ -842,6 +842,7 @@ void UGameplayEffect::ConvertUIComponent()
 		UIData = nullptr;
 	}
 }
+#endif //WITH_EDITOR
 
 EGameplayEffectStackingType UGameplayEffect::GetStackingType() const
 {
@@ -858,7 +859,6 @@ void UGameplayEffect::SetStackingType(EGameplayEffectStackingType InType)
 {
 	StackingType = InType;
 }
-#endif
 
 void UGameplayEffect::PreSave(FObjectPreSaveContext SaveContext)
 {
