Moving new StackingType functionality out of the giant deprecation and editor blocks.
#rnx

diff --git a/Engine/Plugins/Runtime/GameplayAbilities/Source/GameplayAbilities/Private/GameplayEffect.cpp b/Engine/Plugins/Runtime/GameplayAbilities/Source/GameplayAbilities/Private/GameplayEffect.cpp
--- a/Engine/Plugins/Runtime/GameplayAbilities/Source/GameplayAbilities/Private/GameplayEffect.cpp
+++ b/Engine/Plugins/Runtime/GameplayAbilities/Source/GameplayAbilities/Private/GameplayEffect.cpp
@@ -842,23 +842,6 @@ void UGameplayEffect::ConvertUIComponent()
 		UIData = nullptr;
 	}
 }
-#endif //WITH_EDITOR
-
-EGameplayEffectStackingType UGameplayEffect::GetStackingType() const
-{
-	if (DurationPolicy == EGameplayEffectDurationType::Instant)
-	{
-		return EGameplayEffectStackingType::None;
-	}
-	
-	return StackingType;
-}
-
-#if WITH_EDITOR
-void UGameplayEffect::SetStackingType(EGameplayEffectStackingType InType)
-{
-	StackingType = InType;
-}
 
 void UGameplayEffect::PreSave(FObjectPreSaveContext SaveContext)
 {
@@ -920,6 +903,27 @@ void UGameplayEffect::PostEditChangeProperty(FPropertyChangedEvent& PropertyChan
 
 #endif // #if WITH_EDITOR
 
+EGameplayEffectStackingType UGameplayEffect::GetStackingType() const
+{
+	if (DurationPolicy == EGameplayEffectDurationType::Instant)
+	{
+		return EGameplayEffectStackingType::None;
+	}
+	
+	PRAGMA_DISABLE_DEPRECATION_WARNINGS
+	return StackingType;
+	PRAGMA_ENABLE_DEPRECATION_WARNINGS
+}
+
+#if WITH_EDITOR
+void UGameplayEffect::SetStackingType(EGameplayEffectStackingType InType)
+{
+	PRAGMA_DISABLE_DEPRECATION_WARNINGS
+	StackingType = InType;
+	PRAGMA_ENABLE_DEPRECATION_WARNINGS
+}
+#endif
+
 bool UGameplayEffect::CanApply(const FActiveGameplayEffectsContainer& ActiveGEContainer, const FGameplayEffectSpec& GESpec) const
 {
 	for (const UGameplayEffectComponent* GEComponent : GEComponents)
