PR #13388: fix crash when debugging gameplay effect
[RN] Fixed crash when debugging gameplay effect with aggregate by source stacks.

diff --git a/Engine/Plugins/Runtime/GameplayAbilities/Source/GameplayAbilities/Private/AbilitySystemComponent.cpp b/Engine/Plugins/Runtime/GameplayAbilities/Source/GameplayAbilities/Private/AbilitySystemComponent.cpp
--- a/Engine/Plugins/Runtime/GameplayAbilities/Source/GameplayAbilities/Private/AbilitySystemComponent.cpp
+++ b/Engine/Plugins/Runtime/GameplayAbilities/Source/GameplayAbilities/Private/AbilitySystemComponent.cpp
@@ -2698,7 +2698,7 @@ void UAbilitySystemComponent::Debug_Internal(FAbilitySystemComponentDebugInfo& I
 
 				if (ActiveGE.Spec.Def->GetStackingType() == EGameplayEffectStackingType::AggregateBySource)
 				{
-					StackString = FString::Printf(TEXT("(Stacks: %d. From: %s) "), ActiveGE.Spec.GetStackCount(), *GetNameSafe(ActiveGE.Spec.GetContext().GetInstigatorAbilitySystemComponent()->GetAvatarActor_Direct()));
+					StackString = FString::Printf(TEXT("(Stacks: %d. From: %s) "), ActiveGE.Spec.GetStackCount(), ActiveGE.Spec.GetContext().GetInstigatorAbilitySystemComponent() ? *GetNameSafe(ActiveGE.Spec.GetContext().GetInstigatorAbilitySystemComponent()->GetAvatarActor_Direct()) : TEXT("Invalid"));
 				}
 				else
 				{
