Fixed target data recalculating source tags and attributes when applying Effect Spec.
[RN] Fixed target data recalculating source tags and attributes when applying Effect Spec.
#13445

diff --git a/Engine/Plugins/Runtime/GameplayAbilities/Source/GameplayAbilities/Private/GameplayAbilityTargetTypes.cpp b/Engine/Plugins/Runtime/GameplayAbilities/Source/GameplayAbilities/Private/GameplayAbilityTargetTypes.cpp
--- a/Engine/Plugins/Runtime/GameplayAbilities/Source/GameplayAbilities/Private/GameplayAbilityTargetTypes.cpp
+++ b/Engine/Plugins/Runtime/GameplayAbilities/Source/GameplayAbilities/Private/GameplayAbilityTargetTypes.cpp
@@ -10,6 +10,16 @@
 
 #include UE_INLINE_GENERATED_CPP_BY_NAME(GameplayAbilityTargetTypes)
 
+namespace UE::AbilitySystem::Private
+{
+	int32 CVarRecalcuateTargetDataSourceOnApply = 0;
+	FAutoConsoleVariableRef CVarRecalcuateTargetDataSourceOnSpecApplication(
+		TEXT("AbilitySystem.Fix.RecalcuateTargetDataSourceOnApply"),
+		CVarRecalcuateTargetDataSourceOnApply,
+		TEXT("Whether to force recalculation of Source Tags and Attributes on application. Prior to 5.7 this recalculated source tags and attributes when applying from target data."),
+		ECVF_Default);
+}
+
 TArray<FActiveGameplayEffectHandle> FGameplayAbilityTargetData::ApplyGameplayEffect(const UGameplayEffect* GameplayEffect, const FGameplayEffectContextHandle& InEffectContext, float Level, FPredictionKey PredictionKey)
 {
 	// Make a temp spec and call the spec function. This ends up cloning the spec per target
@@ -38,9 +48,13 @@ TArray<FActiveGameplayEffectHandle> FGameplayAbilityTargetData::ApplyGameplayEff
 		if (TargetComponent)
 		{
 			// We have to make a new effect spec and context here, because otherwise the targeting info gets accumulated and things take damage multiple times
-			FGameplayEffectSpec	SpecToApply(InSpec);
-			FGameplayEffectContextHandle EffectContext = SpecToApply.GetContext().Duplicate();
-			SpecToApply.SetContext(EffectContext);
+			FGameplayEffectContextHandle EffectContext = InSpec.GetContext().Duplicate();
+			FGameplayEffectSpec	SpecToApply(InSpec, EffectContext);
+			if (UE::AbilitySystem::Private::CVarRecalcuateTargetDataSourceOnApply > 0)
+			{
+				// Force recalculation of the source tags and attributes
+				SpecToApply.SetContext(EffectContext);
+			}
 
 			AddTargetDataToContext(EffectContext, false);
 
