fix unreachable code warning

#rnx

diff --git a/Engine/Plugins/Runtime/GameplayAbilities/Source/GameplayAbilitiesEditor/Private/AttributeDetails.cpp b/Engine/Plugins/Runtime/GameplayAbilities/Source/GameplayAbilitiesEditor/Private/AttributeDetails.cpp
--- a/Engine/Plugins/Runtime/GameplayAbilities/Source/GameplayAbilitiesEditor/Private/AttributeDetails.cpp
+++ b/Engine/Plugins/Runtime/GameplayAbilities/Source/GameplayAbilitiesEditor/Private/AttributeDetails.cpp
@@ -418,9 +418,6 @@ TSharedRef<SWidget> FScalableFloatDetails::CreateCurveTableWidget()
 			.ToolTipText(this, &FScalableFloatDetails::GetCurveTableTooltip)
 			.Font(FAppStyle::GetFontStyle(TEXT("PropertyWindow.NormalFont")))
 		];
-
-	// Need to make the buttons go away and show custom prompt
-	return CurveTableProperty->CreatePropertyValueWidget(false);
 }
 
 TSharedRef<SWidget> FScalableFloatDetails::CreateRegistryTypeWidget()
