Fixed AbilityTask_ApplyRootMotionJumpForce not removing its root motion source because OnDestroy wasn't being called.

A recent change caused EndTask to not be called on clients that are simulating. This changes PreDestroyFromReplication to call EndTask directly instead.

#tests Entered a match with two players and used an ability that applied a jump root motion force to the player. After landing the player's root motion was removed and other players continued to see them in the correct location.
#rnx

diff --git a/Engine/Plugins/Runtime/GameplayAbilities/Source/GameplayAbilities/Private/Abilities/Tasks/AbilityTask_ApplyRootMotionJumpForce.cpp b/Engine/Plugins/Runtime/GameplayAbilities/Source/GameplayAbilities/Private/Abilities/Tasks/AbilityTask_ApplyRootMotionJumpForce.cpp
--- a/Engine/Plugins/Runtime/GameplayAbilities/Source/GameplayAbilities/Private/Abilities/Tasks/AbilityTask_ApplyRootMotionJumpForce.cpp
+++ b/Engine/Plugins/Runtime/GameplayAbilities/Source/GameplayAbilities/Private/Abilities/Tasks/AbilityTask_ApplyRootMotionJumpForce.cpp
@@ -195,7 +195,8 @@ void UAbilityTask_ApplyRootMotionJumpForce::GetLifetimeReplicatedProps(TArray< F
 
 void UAbilityTask_ApplyRootMotionJumpForce::PreDestroyFromReplication()
 {
-	Finish();
+	bIsFinished = true;
+	EndTask();
 }
 
 void UAbilityTask_ApplyRootMotionJumpForce::OnDestroy(bool AbilityIsEnding)
