Fix attributeset possibly broadcasting the wrong post update value if SetCurrentValue is overriden

diff --git a/Engine/Plugins/Runtime/GameplayAbilities/Source/GameplayAbilities/Private/AttributeSet.cpp b/Engine/Plugins/Runtime/GameplayAbilities/Source/GameplayAbilities/Private/AttributeSet.cpp
--- a/Engine/Plugins/Runtime/GameplayAbilities/Source/GameplayAbilities/Private/AttributeSet.cpp
+++ b/Engine/Plugins/Runtime/GameplayAbilities/Source/GameplayAbilities/Private/AttributeSet.cpp
@@ -99,7 +99,7 @@ void FGameplayAttribute::SetNumericValueChecked(float& NewValue, class UAttribut
 		OldValue = DataPtr->GetCurrentValue();
 		Dest->PreAttributeChange(*this, NewValue);
 		DataPtr->SetCurrentValue(NewValue);
-		Dest->PostAttributeChange(*this, OldValue, NewValue);
+		Dest->PostAttributeChange(*this, OldValue, DataPtr->GetCurrentValue());
 
 		MARK_PROPERTY_DIRTY(Dest, StructProperty);
 	}
diff --git a/Engine/Plugins/Runtime/GameplayAbilities/Source/GameplayAbilities/Private/GameplayEffect.cpp b/Engine/Plugins/Runtime/GameplayAbilities/Source/GameplayAbilities/Private/GameplayEffect.cpp
--- a/Engine/Plugins/Runtime/GameplayAbilities/Source/GameplayAbilities/Private/GameplayEffect.cpp
+++ b/Engine/Plugins/Runtime/GameplayAbilities/Source/GameplayAbilities/Private/GameplayEffect.cpp
@@ -3939,7 +3939,7 @@ void FActiveGameplayEffectsContainer::SetAttributeBaseValue(FGameplayAttribute A
 		InternalUpdateNumericalAttribute(Attribute, NewBaseValue, nullptr);
 	}
 
-	Set->PostAttributeBaseChange(Attribute, OldBaseValue, NewBaseValue);
+	Set->PostAttributeBaseChange(Attribute, OldBaseValue, GetAttributeBaseValue(Attribute));
 }
 
 float FActiveGameplayEffectsContainer::GetAttributeBaseValue(FGameplayAttribute Attribute) const
