Make sure the server still has authority to refresh GE duration when stacking.

diff --git a/Engine/Plugins/Runtime/GameplayAbilities/Source/GameplayAbilities/Private/GameplayEffect.cpp b/Engine/Plugins/Runtime/GameplayAbilities/Source/GameplayAbilities/Private/GameplayEffect.cpp
--- a/Engine/Plugins/Runtime/GameplayAbilities/Source/GameplayAbilities/Private/GameplayEffect.cpp
+++ b/Engine/Plugins/Runtime/GameplayAbilities/Source/GameplayAbilities/Private/GameplayEffect.cpp
@@ -4183,7 +4183,7 @@ PRAGMA_ENABLE_DEPRECATION_WARNINGS
 
 		// Make sure the GE actually wants to refresh its duration
 		// Client's predicted GE will be combined with the existing GE and we don't want to set any duration timer in this case
-		if (InPredictionKey.IsLocalClientKey() || GEDef->StackDurationRefreshPolicy == EGameplayEffectStackingDurationPolicy::NeverRefresh)
+		if ((InPredictionKey.IsLocalClientKey() && InPredictionKey.WasLocallyGenerated()) || GEDef->StackDurationRefreshPolicy == EGameplayEffectStackingDurationPolicy::NeverRefresh)
 		{
 			bSetDuration = false;
 		}
