Fix possibly wrong NewValue of an attribute being broadcast
#rnx

diff --git a/Engine/Plugins/Runtime/GameplayAbilities/Source/GameplayAbilities/Private/GameplayEffect.cpp b/Engine/Plugins/Runtime/GameplayAbilities/Source/GameplayAbilities/Private/GameplayEffect.cpp
--- a/Engine/Plugins/Runtime/GameplayAbilities/Source/GameplayAbilities/Private/GameplayEffect.cpp
+++ b/Engine/Plugins/Runtime/GameplayAbilities/Source/GameplayAbilities/Private/GameplayEffect.cpp
@@ -3856,6 +3856,9 @@ void FActiveGameplayEffectsContainer::InternalUpdateNumericalAttribute(FGameplay
 	
 	if (!bFromRecursiveCall)
 	{
+		// Grab the actual value to broadcast from the attribute since SetCurrentValue could have been overriden
+		NewValue = Owner->GetNumericAttribute(Attribute);
+
 		// We should only have one: either cached CurrentModcallbackData, or explicit callback data passed directly in.
 		if (ModData && CurrentModcallbackData)
 		{
