Improves usability of WaitGameplayEffectApplied_Self+Target ability tasks and the async variant, allowing asset tag and granted tag based requirements and allowing recursion of GE application while still avoiding infinite loops.
* Allowing handling code to apply additional GameplayEffects without triggering the "recursion detected! ignoring" error unless it actually applies the same GE class recursively.
* Caller can specify asset tag and granted tag requirements for which GE to fire these on. This handles many common use cases, like "Wait for a freeze GameplayEffect to be applied", "Wait for fire damage to be applied".

diff --git a/Engine/Plugins/Runtime/GameplayAbilities/Source/GameplayAbilities/Private/Abilities/Async/AbilityAsync_WaitGameplayEffectApplied.cpp b/Engine/Plugins/Runtime/GameplayAbilities/Source/GameplayAbilities/Private/Abilities/Async/AbilityAsync_WaitGameplayEffectApplied.cpp
--- a/Engine/Plugins/Runtime/GameplayAbilities/Source/GameplayAbilities/Private/Abilities/Async/AbilityAsync_WaitGameplayEffectApplied.cpp
+++ b/Engine/Plugins/Runtime/GameplayAbilities/Source/GameplayAbilities/Private/Abilities/Async/AbilityAsync_WaitGameplayEffectApplied.cpp
@@ -7,13 +7,15 @@
 
 #include UE_INLINE_GENERATED_CPP_BY_NAME(AbilityAsync_WaitGameplayEffectApplied)
 
-UAbilityAsync_WaitGameplayEffectApplied* UAbilityAsync_WaitGameplayEffectApplied::WaitGameplayEffectAppliedToActor(AActor* TargetActor, const FGameplayTargetDataFilterHandle SourceFilter, FGameplayTagRequirements SourceTagRequirements, FGameplayTagRequirements TargetTagRequirements, bool TriggerOnce, bool ListenForPeriodicEffect)
+UAbilityAsync_WaitGameplayEffectApplied* UAbilityAsync_WaitGameplayEffectApplied::WaitGameplayEffectAppliedToActor(AActor* TargetActor, const FGameplayTargetDataFilterHandle SourceFilter, FGameplayTagRequirements SourceTagRequirements, FGameplayTagRequirements TargetTagRequirements, FGameplayTagRequirements AssetTagRequirements, FGameplayTagRequirements GrantedTagRequirements, bool TriggerOnce, bool ListenForPeriodicEffect)
 {
 	UAbilityAsync_WaitGameplayEffectApplied* MyObj = NewObject<UAbilityAsync_WaitGameplayEffectApplied>();
 	MyObj->SetAbilityActor(TargetActor);
 	MyObj->Filter = SourceFilter;
 	MyObj->SourceTagRequirements = SourceTagRequirements;
 	MyObj->TargetTagRequirements = TargetTagRequirements;
+	MyObj->AssetTagRequirements = AssetTagRequirements;
+	MyObj->GrantedTagRequirements = GrantedTagRequirements;
 	MyObj->TriggerOnce = TriggerOnce;
 	MyObj->ListenForPeriodicEffects = ListenForPeriodicEffect;
 	return MyObj;
@@ -53,19 +55,43 @@ void UAbilityAsync_WaitGameplayEffectApplied::OnApplyGameplayEffectCallback(UAbi
 	{
 		return;
 	}
+	FGameplayTagContainer AllAssetTags;
+	SpecApplied.GetAllAssetTags(AllAssetTags);
+	if (!AssetTagRequirements.RequirementsMet(AllAssetTags))
+	{
+		ABILITY_LOG(Verbose, TEXT("AbilityAsync_WaitGameplayEffectApplied: Not triggering for Applied Spec: %s, AssetTagRequirements not passed: %s"), *SpecApplied.ToSimpleString(), *AssetTagRequirements.ToString());
+		return;
+	}
+
+	FGameplayTagContainer AllGrantedTags;
+	SpecApplied.GetAllGrantedTags(AllGrantedTags);
+	if (!GrantedTagRequirements.RequirementsMet(AllGrantedTags))
+	{
+		ABILITY_LOG(Verbose, TEXT("AbilityAsync_WaitGameplayEffectApplied: Not triggering for Applied Spec: %s, GrantedTagRequirements not passed: %s"), *SpecApplied.ToSimpleString(), *GrantedTagRequirements.ToString());
+		return;
+	}
 
-	if (bLocked)
+	// We allow GameplayEffect application to trigger other GameplayEffect application (as of UE 5.7). However, in order to prevent infinite recursions, 
+	// we don't allow any effect class to be applied recursively. Keep track of effect classes as this function is entered recursively and abort if the
+	// current applied effect class is already being applied higher in the callstack.
+	const TSubclassOf<UGameplayEffect> EffectClass = SpecApplied.Def ? SpecApplied.Def->GetClass() : nullptr;
+	if (BroadcastingEffectStack.Contains(EffectClass))
 	{
-		ABILITY_LOG(Error, TEXT("WaitGameplayEffectApplied recursion detected. Action: %s. Applied Spec: %s. This could cause an infinite loop! Ignoring"), *GetPathName(), *SpecApplied.ToSimpleString());
+		ABILITY_LOG(Error, TEXT("AbilityAsync_WaitGameplayEffectApplied: recursive application of '%s' detected. This could cause an infinite loop, ignoring broadcast. Current stack: %s"), *GetNameSafe(EffectClass), *GetBroadcastingEffectStackString());
 		return;
 	}
+	else if (UE_LOG_ACTIVE(LogAbilitySystem, VeryVerbose) && !BroadcastingEffectStack.IsEmpty())
+	{
+		ABILITY_LOG(VeryVerbose, TEXT("AbilityAsync_WaitGameplayEffectApplied: application of '%s' causes no recursion. Current stack: %s"), *GetNameSafe(EffectClass), *GetBroadcastingEffectStackString());
+	}
 
 	FGameplayEffectSpecHandle SpecHandle(new FGameplayEffectSpec(SpecApplied));
 
 	if (ShouldBroadcastDelegates())
 	{
-		TGuardValue<bool> GuardValue(bLocked, true);
+		BroadcastingEffectStack.Add(EffectClass);
 		OnApplied.Broadcast(AvatarActor, SpecHandle, ActiveHandle);
+		BroadcastingEffectStack.RemoveSingle(EffectClass);
 
 		if (TriggerOnce)
 		{
@@ -91,3 +117,13 @@ void UAbilityAsync_WaitGameplayEffectApplied::EndAction()
 	Super::EndAction();
 }
 
+FString UAbilityAsync_WaitGameplayEffectApplied::GetBroadcastingEffectStackString() const
+{
+	TArray<FString> AllEffectNames;
+	for (const TSubclassOf<UGameplayEffect>& ExistingEffectClass : BroadcastingEffectStack)
+	{
+		AllEffectNames.Add(GetNameSafe(ExistingEffectClass));
+
+	}
+	return FString::Join(AllEffectNames, TEXT(","));
+}
diff --git a/Engine/Plugins/Runtime/GameplayAbilities/Source/GameplayAbilities/Private/Abilities/Tasks/AbilityTask_WaitGameplayEffectApplied.cpp b/Engine/Plugins/Runtime/GameplayAbilities/Source/GameplayAbilities/Private/Abilities/Tasks/AbilityTask_WaitGameplayEffectApplied.cpp
--- a/Engine/Plugins/Runtime/GameplayAbilities/Source/GameplayAbilities/Private/Abilities/Tasks/AbilityTask_WaitGameplayEffectApplied.cpp
+++ b/Engine/Plugins/Runtime/GameplayAbilities/Source/GameplayAbilities/Private/Abilities/Tasks/AbilityTask_WaitGameplayEffectApplied.cpp
@@ -10,7 +10,6 @@
 UAbilityTask_WaitGameplayEffectApplied::UAbilityTask_WaitGameplayEffectApplied(const FObjectInitializer& ObjectInitializer)
 	: Super(ObjectInitializer)
 {
-	Locked = false;
 }
 
 void UAbilityTask_WaitGameplayEffectApplied::Activate()
@@ -31,42 +30,80 @@ void UAbilityTask_WaitGameplayEffectApplied::OnApplyGameplayEffectCallback(UAbil
 	{
 		return;
 	}
+
+	// Check tag queries if the ability task was created using with tag requirement params.
 	if (!SourceTagRequirements.RequirementsMet(*SpecApplied.CapturedSourceTags.GetAggregatedTags()))
 	{
+		ABILITY_LOG(Verbose, TEXT("WaitGameplayEffectApplied in Ability '%s': Not triggering for Applied Spec: %s, SourceTagRequirements not passed: %s"), *GetNameSafe(Ability), *SpecApplied.ToSimpleString(), *SourceTagRequirements.ToString());
 		return;
 	}
+
 	if (!TargetTagRequirements.RequirementsMet(*SpecApplied.CapturedTargetTags.GetAggregatedTags()))
 	{
+		ABILITY_LOG(Verbose, TEXT("WaitGameplayEffectApplied in Ability '%s': Not triggering for Applied Spec: %s, TargetTagRequirements not passed: %s"), *GetNameSafe(Ability), *SpecApplied.ToSimpleString(), *TargetTagRequirements.ToString());
+		return;
+	}
+
+	FGameplayTagContainer AllAssetTags;
+	SpecApplied.GetAllAssetTags(AllAssetTags);
+	if (!AssetTagRequirements.RequirementsMet(AllAssetTags))
+	{
+		ABILITY_LOG(Verbose, TEXT("WaitGameplayEffectApplied in Ability '%s': Not triggering for Applied Spec: %s, AssetTagRequirements not passed: %s"), *GetNameSafe(Ability), *SpecApplied.ToSimpleString(), *AssetTagRequirements.ToString());
 		return;
 	}
 
-	if (SourceTagQuery.IsEmpty() == false)
+	FGameplayTagContainer AllGrantedTags;
+	SpecApplied.GetAllGrantedTags(AllGrantedTags);
+	if (!GrantedTagRequirements.RequirementsMet(AllGrantedTags))
 	{
-		if (!SourceTagQuery.Matches(*SpecApplied.CapturedSourceTags.GetAggregatedTags()))
-		{
-			return;
-		}
+		ABILITY_LOG(Verbose, TEXT("WaitGameplayEffectApplied in Ability '%s': Not triggering for Applied Spec: %s, GrantedTagRequirements not passed: %s"), *GetNameSafe(Ability), *SpecApplied.ToSimpleString(), *GrantedTagRequirements.ToString());
+		return;
 	}
 
-	if (TargetTagQuery.IsEmpty() == false)
+	// Check tag queries if the ability task was created using with tag query params.
+	if (!SourceTagQuery.IsEmpty() && !SourceTagQuery.Matches(*SpecApplied.CapturedSourceTags.GetAggregatedTags()))
 	{
-		if (!TargetTagQuery.Matches(*SpecApplied.CapturedTargetTags.GetAggregatedTags()))
-		{
-			return;
-		}
+		ABILITY_LOG(Verbose, TEXT("WaitGameplayEffectApplied in Ability '%s': Not triggering for Applied Spec: %s, SourceTagQuery not passed: %s"), *GetNameSafe(Ability), *SpecApplied.ToSimpleString(), *SourceTagQuery.GetDescription());
+		return;
 	}
 
-	if (Locked)
+	if (!TargetTagQuery.IsEmpty() && !TargetTagQuery.Matches(*SpecApplied.CapturedTargetTags.GetAggregatedTags()))
 	{
-		ABILITY_LOG(Error, TEXT("WaitGameplayEffectApplied recursion detected. Ability: %s. Applied Spec: %s. This could cause an infinite loop! Ignoring"), *GetNameSafe(Ability), *SpecApplied.ToSimpleString());
+		ABILITY_LOG(Verbose, TEXT("WaitGameplayEffectApplied in Ability '%s': Not triggering for Applied Spec: %s, TargetTagQuery not passed: %s"), *GetNameSafe(Ability), *SpecApplied.ToSimpleString(), *TargetTagQuery.GetDescription());
 		return;
 	}
+
+	if (!AssetTagQuery.IsEmpty() && !AssetTagQuery.Matches(AllAssetTags))
+	{
+		ABILITY_LOG(Verbose, TEXT("WaitGameplayEffectApplied in Ability '%s': Not triggering for Applied Spec: %s, AssetTagQuery not passed: %s"), *GetNameSafe(Ability), *SpecApplied.ToSimpleString(), *AssetTagQuery.GetDescription());
+		return;
+	}
+
+	if (!GrantedTagQuery.IsEmpty() && !GrantedTagQuery.Matches(AllGrantedTags))
+	{
+		ABILITY_LOG(Verbose, TEXT("WaitGameplayEffectApplied in Ability '%s': Not triggering for Applied Spec: %s, GrantedTagQuery not passed: %s"), *GetNameSafe(Ability), *SpecApplied.ToSimpleString(), *GrantedTagQuery.GetDescription());
+		return;
+	}
+
+	// We allow GameplayEffect application to trigger other GameplayEffect application (as of UE 5.7). However, in order to prevent infinite recursions, 
+	// we don't allow any effect class to be applied recursively. Keep track of effect classes as this function is entered recursively and abort if the
+	// current applied effect class is already being applied higher in the callstack.
+	const TSubclassOf<UGameplayEffect> EffectClass = SpecApplied.Def ? SpecApplied.Def->GetClass() : nullptr;
+	if (BroadcastingEffectStack.Contains(EffectClass))
+	{
+		ABILITY_LOG(Error, TEXT("WaitGameplayEffectApplied in Ability '%s': recursive application of '%s' detected. This could cause an infinite loop, ignoring broadcast. Current stack: %s"), *GetNameSafe(Ability), *GetNameSafe(EffectClass), *GetBroadcastingEffectStackString());
+		return;
+	}
+	else if (UE_LOG_ACTIVE(LogAbilitySystem, VeryVerbose) && !BroadcastingEffectStack.IsEmpty())
+	{
+		ABILITY_LOG(VeryVerbose, TEXT("WaitGameplayEffectApplied in Ability '%s': application of '%s' causes no recursion. Current stack: %s"), *GetNameSafe(Ability), *GetNameSafe(EffectClass), *GetBroadcastingEffectStackString());
+	}
 	
 	FGameplayEffectSpecHandle	SpecHandle(new FGameplayEffectSpec(SpecApplied));
-
 	{
-		TGuardValue<bool> GuardValue(Locked, true);	
+		BroadcastingEffectStack.Add(EffectClass);
 		BroadcastDelegate(AvatarActor, SpecHandle, ActiveHandle);
+		BroadcastingEffectStack.RemoveSingle(EffectClass);
 	}
 
 	if (TriggerOnce)
@@ -75,6 +112,17 @@ void UAbilityTask_WaitGameplayEffectApplied::OnApplyGameplayEffectCallback(UAbil
 	}
 }
 
+FString UAbilityTask_WaitGameplayEffectApplied::GetBroadcastingEffectStackString() const
+{
+	TArray<FString> AllEffectNames;
+	for (const TSubclassOf<UGameplayEffect>& ExistingEffectClass : BroadcastingEffectStack)
+	{
+		AllEffectNames.Add(GetNameSafe(ExistingEffectClass));
+
+	}
+	return FString::Join(AllEffectNames, TEXT(","));
+}
+
 void UAbilityTask_WaitGameplayEffectApplied::OnDestroy(bool AbilityEnded)
 {
 	if (GetASC())
diff --git a/Engine/Plugins/Runtime/GameplayAbilities/Source/GameplayAbilities/Private/Abilities/Tasks/AbilityTask_WaitGameplayEffectApplied_Self.cpp b/Engine/Plugins/Runtime/GameplayAbilities/Source/GameplayAbilities/Private/Abilities/Tasks/AbilityTask_WaitGameplayEffectApplied_Self.cpp
--- a/Engine/Plugins/Runtime/GameplayAbilities/Source/GameplayAbilities/Private/Abilities/Tasks/AbilityTask_WaitGameplayEffectApplied_Self.cpp
+++ b/Engine/Plugins/Runtime/GameplayAbilities/Source/GameplayAbilities/Private/Abilities/Tasks/AbilityTask_WaitGameplayEffectApplied_Self.cpp
@@ -12,24 +12,28 @@ UAbilityTask_WaitGameplayEffectApplied_Self::UAbilityTask_WaitGameplayEffectAppl
 {
 }
 
-UAbilityTask_WaitGameplayEffectApplied_Self* UAbilityTask_WaitGameplayEffectApplied_Self::WaitGameplayEffectAppliedToSelf(UGameplayAbility* OwningAbility, const FGameplayTargetDataFilterHandle InFilter, FGameplayTagRequirements InSourceTagRequirements, FGameplayTagRequirements InTargetTagRequirements, bool InTriggerOnce, AActor* OptionalExternalOwner, bool InListenForPeriodicEffect)
+UAbilityTask_WaitGameplayEffectApplied_Self* UAbilityTask_WaitGameplayEffectApplied_Self::WaitGameplayEffectAppliedToSelf(UGameplayAbility* OwningAbility, const FGameplayTargetDataFilterHandle InFilter, FGameplayTagRequirements InSourceTagRequirements, FGameplayTagRequirements InTargetTagRequirements, FGameplayTagRequirements InAssetTagRequirements, FGameplayTagRequirements InGrantedTagRequirements, bool InTriggerOnce, AActor* OptionalExternalOwner, bool InListenForPeriodicEffect)
 {
 	UAbilityTask_WaitGameplayEffectApplied_Self* MyObj = NewAbilityTask<UAbilityTask_WaitGameplayEffectApplied_Self>(OwningAbility);
 	MyObj->Filter = InFilter;
 	MyObj->SourceTagRequirements = InSourceTagRequirements;
 	MyObj->TargetTagRequirements = InTargetTagRequirements;
+	MyObj->AssetTagRequirements = InAssetTagRequirements;
+	MyObj->GrantedTagRequirements = InGrantedTagRequirements;
 	MyObj->TriggerOnce = InTriggerOnce;
 	MyObj->SetExternalActor(OptionalExternalOwner);
 	MyObj->ListenForPeriodicEffects = InListenForPeriodicEffect;
 	return MyObj;
 }
 
-UAbilityTask_WaitGameplayEffectApplied_Self* UAbilityTask_WaitGameplayEffectApplied_Self::WaitGameplayEffectAppliedToSelf_Query(UGameplayAbility* OwningAbility, const FGameplayTargetDataFilterHandle InFilter, FGameplayTagQuery InSourceTagQuery, FGameplayTagQuery InTargetTagQuery, bool InTriggerOnce, AActor* OptionalExternalOwner, bool InListenForPeriodicEffect)
+UAbilityTask_WaitGameplayEffectApplied_Self* UAbilityTask_WaitGameplayEffectApplied_Self::WaitGameplayEffectAppliedToSelf_Query(UGameplayAbility* OwningAbility, const FGameplayTargetDataFilterHandle InFilter, FGameplayTagQuery InSourceTagQuery, FGameplayTagQuery InTargetTagQuery, FGameplayTagQuery InAssetTagQuery, FGameplayTagQuery InGrantedTagQuery, bool InTriggerOnce, AActor* OptionalExternalOwner, bool InListenForPeriodicEffect)
 {
 	UAbilityTask_WaitGameplayEffectApplied_Self* MyObj = NewAbilityTask<UAbilityTask_WaitGameplayEffectApplied_Self>(OwningAbility);
 	MyObj->Filter = InFilter;
 	MyObj->SourceTagQuery = InSourceTagQuery;
 	MyObj->TargetTagQuery = InTargetTagQuery;
+	MyObj->AssetTagQuery = InAssetTagQuery;
+	MyObj->GrantedTagQuery = InGrantedTagQuery;
 	MyObj->TriggerOnce = InTriggerOnce;
 	MyObj->SetExternalActor(OptionalExternalOwner);
 	MyObj->ListenForPeriodicEffects = InListenForPeriodicEffect;
diff --git a/Engine/Plugins/Runtime/GameplayAbilities/Source/GameplayAbilities/Private/Abilities/Tasks/AbilityTask_WaitGameplayEffectApplied_Target.cpp b/Engine/Plugins/Runtime/GameplayAbilities/Source/GameplayAbilities/Private/Abilities/Tasks/AbilityTask_WaitGameplayEffectApplied_Target.cpp
--- a/Engine/Plugins/Runtime/GameplayAbilities/Source/GameplayAbilities/Private/Abilities/Tasks/AbilityTask_WaitGameplayEffectApplied_Target.cpp
+++ b/Engine/Plugins/Runtime/GameplayAbilities/Source/GameplayAbilities/Private/Abilities/Tasks/AbilityTask_WaitGameplayEffectApplied_Target.cpp
@@ -12,24 +12,28 @@ UAbilityTask_WaitGameplayEffectApplied_Target::UAbilityTask_WaitGameplayEffectAp
 {
 }
 
-UAbilityTask_WaitGameplayEffectApplied_Target* UAbilityTask_WaitGameplayEffectApplied_Target::WaitGameplayEffectAppliedToTarget(UGameplayAbility* OwningAbility, const FGameplayTargetDataFilterHandle InFilter, FGameplayTagRequirements InSourceTagRequirements, FGameplayTagRequirements InTargetTagRequirements, bool InTriggerOnce, AActor* OptionalExternalOwner, bool InListenForPeriodicEffect)
+UAbilityTask_WaitGameplayEffectApplied_Target* UAbilityTask_WaitGameplayEffectApplied_Target::WaitGameplayEffectAppliedToTarget(UGameplayAbility* OwningAbility, const FGameplayTargetDataFilterHandle InFilter, FGameplayTagRequirements InSourceTagRequirements, FGameplayTagRequirements InTargetTagRequirements, FGameplayTagRequirements InAssetTagRequirements, FGameplayTagRequirements InGrantedTagRequirements, bool InTriggerOnce, AActor* OptionalExternalOwner, bool InListenForPeriodicEffect)
 {
 	UAbilityTask_WaitGameplayEffectApplied_Target* MyObj = NewAbilityTask<UAbilityTask_WaitGameplayEffectApplied_Target>(OwningAbility);
 	MyObj->Filter = InFilter;
 	MyObj->SourceTagRequirements = InSourceTagRequirements;
 	MyObj->TargetTagRequirements = InTargetTagRequirements;
+	MyObj->AssetTagRequirements = InAssetTagRequirements;
+	MyObj->GrantedTagRequirements = InGrantedTagRequirements;
 	MyObj->TriggerOnce = InTriggerOnce;
 	MyObj->SetExternalActor(OptionalExternalOwner);
 	MyObj->ListenForPeriodicEffects = InListenForPeriodicEffect;
 	return MyObj;
 }
 
-UAbilityTask_WaitGameplayEffectApplied_Target* UAbilityTask_WaitGameplayEffectApplied_Target::WaitGameplayEffectAppliedToTarget_Query(UGameplayAbility* OwningAbility, const FGameplayTargetDataFilterHandle InFilter, FGameplayTagQuery SourceTagQuery, FGameplayTagQuery TargetTagQuery, bool InTriggerOnce, AActor* OptionalExternalOwner, bool InListenForPeriodicEffect)
+UAbilityTask_WaitGameplayEffectApplied_Target* UAbilityTask_WaitGameplayEffectApplied_Target::WaitGameplayEffectAppliedToTarget_Query(UGameplayAbility* OwningAbility, const FGameplayTargetDataFilterHandle InFilter, FGameplayTagQuery SourceTagQuery, FGameplayTagQuery TargetTagQuery, FGameplayTagQuery AssetTagQuery, FGameplayTagQuery GrantedTagQuery, bool InTriggerOnce, AActor* OptionalExternalOwner, bool InListenForPeriodicEffect)
 {
 	UAbilityTask_WaitGameplayEffectApplied_Target* MyObj = NewAbilityTask<UAbilityTask_WaitGameplayEffectApplied_Target>(OwningAbility);
 	MyObj->Filter = InFilter;
 	MyObj->SourceTagQuery = SourceTagQuery;
 	MyObj->TargetTagQuery = TargetTagQuery;
+	MyObj->AssetTagQuery = AssetTagQuery;
+	MyObj->GrantedTagQuery = GrantedTagQuery;
 	MyObj->TriggerOnce = InTriggerOnce;
 	MyObj->SetExternalActor(OptionalExternalOwner);
 	MyObj->ListenForPeriodicEffects = InListenForPeriodicEffect;
diff --git a/Engine/Plugins/Runtime/GameplayAbilities/Source/GameplayAbilities/Public/Abilities/Async/AbilityAsync_WaitGameplayEffectApplied.h b/Engine/Plugins/Runtime/GameplayAbilities/Source/GameplayAbilities/Public/Abilities/Async/AbilityAsync_WaitGameplayEffectApplied.h
--- a/Engine/Plugins/Runtime/GameplayAbilities/Source/GameplayAbilities/Public/Abilities/Async/AbilityAsync_WaitGameplayEffectApplied.h
+++ b/Engine/Plugins/Runtime/GameplayAbilities/Source/GameplayAbilities/Public/Abilities/Async/AbilityAsync_WaitGameplayEffectApplied.h
@@ -25,12 +25,20 @@ class UAbilityAsync_WaitGameplayEffectApplied : public UAbilityAsync
 
 public:
 	/**
-	 * Wait until a GameplayEffect is applied to a target actor
-	 * If TriggerOnce is true, this action will only activate one time. Otherwise it will return every time a GE is applied that meets the requirements over the life of the ability
-	 * If used in an ability graph, this async action will wait even after activation ends. It's recommended to use WaitGameplayEffectApplied instead.
+	 * Wait until a GameplayEffect is applied to a target actor.
+	 * 
+	 * @param TargetActor The actor on which to listen for a GameplayEffect being applied to it
+	 * @param SourceFilter Actor filter for the source actor applying this effect that must be passed to trigger this task
+	 * @param SourceTagRequirements Gameplay tag requirements for the source actor applying this effect
+	 * @param TargetTagRequirements Gameplay tag requirements for the target actor receiving this effect (self in this case)
+	 * @param AssetTagRequirements Requirements for the asset tags of the applied effect: 'Tags this Effect Has' on the GameplayEffect class
+	 * @param GrantedTagRequirements Requirements for the tags granted by the applied effect: 'Grant Tags to Target Actor' on the GameplayEffect class
+	 * @param TriggerOnce If TriggerOnce is true, this task will only return one time. Otherwise it will return everytime a GE is applied that meets the requirements over the life of the ability.
+	 * @param OptionalExternalOwner Can be used to run this task on someone else (not the owner of the ability). By default you can leave this empty.
+	 * @param ListenForPeriodicEffects Whether to also fire on periodic ticks of the applied effect
 	 */
 	UFUNCTION(BlueprintCallable, Category = "Ability|Async", meta = (DefaultToSelf = "TargetActor", BlueprintInternalUseOnly = "TRUE"))
-	static UE_API UAbilityAsync_WaitGameplayEffectApplied* WaitGameplayEffectAppliedToActor(AActor* TargetActor, const FGameplayTargetDataFilterHandle SourceFilter, FGameplayTagRequirements SourceTagRequirements, FGameplayTagRequirements TargetTagRequirements, bool TriggerOnce = false, bool ListenForPeriodicEffect = false);
+	static UE_API UAbilityAsync_WaitGameplayEffectApplied* WaitGameplayEffectAppliedToActor(AActor* TargetActor, const FGameplayTargetDataFilterHandle SourceFilter, FGameplayTagRequirements SourceTagRequirements, FGameplayTagRequirements TargetTagRequirements, FGameplayTagRequirements AssetTagRequirements, FGameplayTagRequirements GrantedTagRequirements, bool TriggerOnce = false, bool ListenForPeriodicEffect = false);
 
 	DECLARE_DYNAMIC_MULTICAST_DELEGATE_ThreeParams(FOnAppliedDelegate, AActor*, Source, FGameplayEffectSpecHandle, SpecHandle, FActiveGameplayEffectHandle, ActiveHandle);
 	UPROPERTY(BlueprintAssignable)
@@ -40,18 +48,29 @@ protected:
 	UE_API virtual void Activate() override;
 	UE_API virtual void EndAction() override;
 
+	UE_API FString GetBroadcastingEffectStackString() const;
+
 	UE_API void OnApplyGameplayEffectCallback(UAbilitySystemComponent* Target, const FGameplayEffectSpec& SpecApplied, FActiveGameplayEffectHandle ActiveHandle);
 
+	// Actor filter for the source actor applying this effect that must be passed to trigger this task
 	FGameplayTargetDataFilterHandle Filter;
+	// Tag requirements for source actor that applies the GameplayEffect
 	FGameplayTagRequirements SourceTagRequirements;
+	// Tag requirements for target actor to which the GameplayEffect is applied
 	FGameplayTagRequirements TargetTagRequirements;
+	// Tag requirements for asset tags ('Tags this Effect Has') on the GameplayEffect being applied
+	FGameplayTagRequirements AssetTagRequirements;
+	// Tag requirements for granted tags ('Grant Tags to Target Actor') that the GameplayEffect applies to the target actor
+	FGameplayTagRequirements GrantedTagRequirements;
 
 	bool TriggerOnce = false;
 	bool ListenForPeriodicEffects = false;
 
 	FDelegateHandle OnApplyGameplayEffectCallbackDelegateHandle;
 	FDelegateHandle OnPeriodicGameplayEffectExecuteCallbackDelegateHandle;
-	bool bLocked = false;
+	
+	// If we are in the process of broadcasting, we do not allow a GameplayEffect class to trigger application of another instance of the same class.
+	TArray<TSubclassOf<class UGameplayEffect>> BroadcastingEffectStack;
 };
 
 #undef UE_API
diff --git a/Engine/Plugins/Runtime/GameplayAbilities/Source/GameplayAbilities/Public/Abilities/Tasks/AbilityTask_WaitGameplayEffectApplied.h b/Engine/Plugins/Runtime/GameplayAbilities/Source/GameplayAbilities/Public/Abilities/Tasks/AbilityTask_WaitGameplayEffectApplied.h
--- a/Engine/Plugins/Runtime/GameplayAbilities/Source/GameplayAbilities/Public/Abilities/Tasks/AbilityTask_WaitGameplayEffectApplied.h
+++ b/Engine/Plugins/Runtime/GameplayAbilities/Source/GameplayAbilities/Public/Abilities/Tasks/AbilityTask_WaitGameplayEffectApplied.h
@@ -22,12 +22,25 @@ class UAbilityTask_WaitGameplayEffectApplied : public UAbilityTask
 
 	UE_API virtual void Activate() override;
 
+	// Actor filter for the source actor applying this effect that must be passed to trigger this task
 	FGameplayTargetDataFilterHandle Filter;
+	// Tag requirements for source actor that applies the GameplayEffect
 	FGameplayTagRequirements SourceTagRequirements;
+	// Tag requirements for target actor to which the GameplayEffect is applied
 	FGameplayTagRequirements TargetTagRequirements;
+	// Tag requirements for asset tags ('Tags this Effect Has') on the GameplayEffect being applied
+	FGameplayTagRequirements AssetTagRequirements;
+	// Tag requirements for granted tags ('Grant Tags to Target Actor') that the GameplayEffect applies to the target actor
+	FGameplayTagRequirements GrantedTagRequirements;
 
+	// Tag query that must be passed for source actor that applies the GameplayEffect
 	FGameplayTagQuery SourceTagQuery;
+	// Tag query that must be passed for target actor to which the GameplayEffect is applied
 	FGameplayTagQuery TargetTagQuery;
+	// Tag query that must be passed for asset tags ('Tags this Effect Has') on the GameplayEffect being applied
+	FGameplayTagQuery AssetTagQuery;
+	// Tag query that must be passed for granted tags ('Grant Tags to Target Actor') that the GameplayEffect applies to the target actor
+	FGameplayTagQuery GrantedTagQuery;
 
 	bool TriggerOnce;
 	bool ListenForPeriodicEffects;
@@ -42,6 +55,8 @@ protected:
 	virtual void RegisterDelegate() { }
 	virtual void RemoveDelegate() { }
 
+	FString GetBroadcastingEffectStackString() const;
+
 	UE_API virtual void OnDestroy(bool AbilityEnded) override;
 
 	bool RegisteredCallback;
@@ -50,8 +65,8 @@ protected:
 	UPROPERTY()
 	TObjectPtr<UAbilitySystemComponent> ExternalOwner;
 
-	// If we are in the process of broadcasting and should not accept additional GE callbacks
-	bool Locked;
+	// If we are in the process of broadcasting, we do not allow a GameplayEffect class to trigger application of another instance of the same class.
+	TArray<TSubclassOf<class UGameplayEffect>> BroadcastingEffectStack;
 };
 
 #undef UE_API
diff --git a/Engine/Plugins/Runtime/GameplayAbilities/Source/GameplayAbilities/Public/Abilities/Tasks/AbilityTask_WaitGameplayEffectApplied_Self.h b/Engine/Plugins/Runtime/GameplayAbilities/Source/GameplayAbilities/Public/Abilities/Tasks/AbilityTask_WaitGameplayEffectApplied_Self.h
--- a/Engine/Plugins/Runtime/GameplayAbilities/Source/GameplayAbilities/Public/Abilities/Tasks/AbilityTask_WaitGameplayEffectApplied_Self.h
+++ b/Engine/Plugins/Runtime/GameplayAbilities/Source/GameplayAbilities/Public/Abilities/Tasks/AbilityTask_WaitGameplayEffectApplied_Self.h
@@ -21,20 +21,37 @@ class UAbilityTask_WaitGameplayEffectApplied_Self : public UAbilityTask_WaitGame
 	FGameplayEffectAppliedSelfDelegate OnApplied;
 
 	/**	 
-	 * Wait until the owner *receives* a GameplayEffect from a given source (the source may be the owner too!). If TriggerOnce is true, this task will only return one time. Otherwise it will return everytime a GE is applied that meets the requirements over the life of the ability
-	 * Optional External Owner can be used to run this task on someone else (not the owner of the ability). By default you can leave this empty.
+	 * Wait until the owner *receives* a GameplayEffect from a given source (the source may be the owner too!).
+	 *
+	 * @param OwningAbility The ability that executes this task
+	 * @param SourceFilter Actor filter for the source actor applying this effect that must be passed to trigger this task
+	 * @param SourceTagRequirements Gameplay tag requirements for the source actor applying this effect
+	 * @param TargetTagRequirements Gameplay tag requirements for the target actor receiving this effect (self in this case)
+	 * @param AssetTagRequirements Requirements for the asset tags of the applied effect: 'Tags this Effect Has' on the GameplayEffect class
+	 * @param GrantedTagRequirements Requirements for the tags granted by the applied effect: 'Grant Tags to Target Actor' on the GameplayEffect class
+	 * @param TriggerOnce If TriggerOnce is true, this task will only return one time. Otherwise it will return everytime a GE is applied that meets the requirements over the life of the ability.
+	 * @param OptionalExternalOwner Can be used to run this task on someone else (not the owner of the ability). By default you can leave this empty.
+	 * @param ListenForPeriodicEffects Whether to also fire on periodic ticks of the applied effect
 	 */	 	 
 	UFUNCTION(BlueprintCallable, Category="Ability|Tasks", meta = (HidePin = "OwningAbility", DefaultToSelf = "OwningAbility", BlueprintInternalUseOnly = "TRUE"))
-	static UE_API UAbilityTask_WaitGameplayEffectApplied_Self* WaitGameplayEffectAppliedToSelf(UGameplayAbility* OwningAbility, const FGameplayTargetDataFilterHandle SourceFilter, FGameplayTagRequirements SourceTagRequirements, FGameplayTagRequirements TargetTagRequirements, bool TriggerOnce=false, AActor* OptionalExternalOwner=nullptr, bool ListenForPeriodicEffect=false );
+	static UE_API UAbilityTask_WaitGameplayEffectApplied_Self* WaitGameplayEffectAppliedToSelf(UGameplayAbility* OwningAbility, const FGameplayTargetDataFilterHandle SourceFilter, FGameplayTagRequirements SourceTagRequirements, FGameplayTagRequirements TargetTagRequirements, FGameplayTagRequirements AssetTagRequirements, FGameplayTagRequirements GrantedTagRequirements, bool TriggerOnce=false, AActor* OptionalExternalOwner=nullptr, bool ListenForPeriodicEffect=false );
 
 	/**	 
-	 * Wait until the owner *receives* a GameplayEffect from a given source (the source may be the owner too!). If TriggerOnce is true, this task will only return one time. Otherwise it will return everytime a GE is applied that meets the requirements over the life of the ability
-	 * Optional External Owner can be used to run this task on someone else (not the owner of the ability). By default you can leave this empty.
-	 *
-	 * This version uses FGameplayTagQuery (more power) instead of FGameplayTagRequirements (faster)
+	 * Wait until the owner *receives* a GameplayEffect from a given source (the source may be the owner too!).
+	 * This version uses FGameplayTagQuery (more power) instead of FGameplayTagRequirements (faster).
+	 * 
+	 * @param OwningAbility The ability that executes this task
+	 * @param SourceFilter Actor filter for the source actor applying this effect that must be passed to trigger this task
+	 * @param SourceTagQuery Gameplay tag requirements for the source actor applying this effect
+	 * @param TargetTagQuery Gameplay tag requirements for the target actor receiving this effect (self in this case)
+	 * @param AssetTagQuery Requirements for the asset tags of the applied effect: 'Tags this Effect Has' on the GameplayEffect class
+	 * @param GrantedTagQuery Requirements for the tags granted by the applied effect: 'Grant Tags to Target Actor' on the GameplayEffect class
+	 * @param TriggerOnce If TriggerOnce is true, this task will only return one time. Otherwise it will return everytime a GE is applied that meets the requirements over the life of the ability.
+	 * @param OptionalExternalOwner Can be used to run this task on someone else (not the owner of the ability). By default you can leave this empty.
+	 * @param ListenForPeriodicEffects Whether to also fire on periodic ticks of the applied effect
 	 */	
 	UFUNCTION(BlueprintCallable, Category="Ability|Tasks", meta = (HidePin = "OwningAbility", DefaultToSelf = "OwningAbility", BlueprintInternalUseOnly = "TRUE"))
-	static UE_API UAbilityTask_WaitGameplayEffectApplied_Self* WaitGameplayEffectAppliedToSelf_Query(UGameplayAbility* OwningAbility, const FGameplayTargetDataFilterHandle SourceFilter, FGameplayTagQuery SourceTagQuery, FGameplayTagQuery TargetTagQuery, bool TriggerOnce=false, AActor* OptionalExternalOwner=nullptr, bool ListenForPeriodicEffect=false );
+	static UE_API UAbilityTask_WaitGameplayEffectApplied_Self* WaitGameplayEffectAppliedToSelf_Query(UGameplayAbility* OwningAbility, const FGameplayTargetDataFilterHandle SourceFilter, FGameplayTagQuery SourceTagQuery, FGameplayTagQuery TargetTagQuery, FGameplayTagQuery AssetTagQuery, FGameplayTagQuery GrantedTagQuery, bool TriggerOnce=false, AActor* OptionalExternalOwner=nullptr, bool ListenForPeriodicEffect=false );
 	
 protected:
 
diff --git a/Engine/Plugins/Runtime/GameplayAbilities/Source/GameplayAbilities/Public/Abilities/Tasks/AbilityTask_WaitGameplayEffectApplied_Target.h b/Engine/Plugins/Runtime/GameplayAbilities/Source/GameplayAbilities/Public/Abilities/Tasks/AbilityTask_WaitGameplayEffectApplied_Target.h
--- a/Engine/Plugins/Runtime/GameplayAbilities/Source/GameplayAbilities/Public/Abilities/Tasks/AbilityTask_WaitGameplayEffectApplied_Target.h
+++ b/Engine/Plugins/Runtime/GameplayAbilities/Source/GameplayAbilities/Public/Abilities/Tasks/AbilityTask_WaitGameplayEffectApplied_Target.h
@@ -21,18 +21,37 @@ class UAbilityTask_WaitGameplayEffectApplied_Target : public UAbilityTask_WaitGa
 	FGameplayEffectAppliedTargetDelegate OnApplied;
 
 	/** 
-	 * Wait until the owner (or External Owner) applies a GameplayEffect to a Target (the target may be the owner too!). If TriggerOnce is true, this task will only return one time. Otherwise it will return everytime a GE is applied that meets the requirements over the life of the ability 
-	 * Optional External Owner can be used to run this task on someone else (not the owner of the ability). By default you can leave this empty.
+	 * Wait until the owner (or External Owner) applies a GameplayEffect to a Target (the target may be the owner too!).
+	 * 
+	 * @param OwningAbility The ability that executes this task
+	 * @param SourceFilter Actor filter for the source actor applying this effect that must be passed to trigger this task
+	 * @param SourceTagRequirements Gameplay tag requirements for the source actor applying this effect
+	 * @param TargetTagRequirements Gameplay tag requirements for the target actor receiving this effect
+	 * @param AssetTagRequirements Requirements for the asset tags of the applied effect: 'Tags this Effect Has' on the GameplayEffect class.
+	 * @param GrantedTagRequirements Requirements for the tags granted by the applied effect: 'Grant Tags to Target Actor' on the GameplayEffect class.
+	 * @param TriggerOnce If TriggerOnce is true, this task will only return one time. Otherwise it will return everytime a GE is applied that meets the requirements over the life of the ability.
+	 * @param OptionalExternalOwner Can be used to run this task on someone else (not the owner of the ability). By default you can leave this empty.
+	 * @param ListenForPeriodicEffects Whether to also fire on periodic ticks of the applied effect
 	 */
 	UFUNCTION(BlueprintCallable, Category="Ability|Tasks", meta = (HidePin = "OwningAbility", DefaultToSelf = "OwningAbility", BlueprintInternalUseOnly = "TRUE"))
-	static UE_API UAbilityTask_WaitGameplayEffectApplied_Target* WaitGameplayEffectAppliedToTarget(UGameplayAbility* OwningAbility, const FGameplayTargetDataFilterHandle TargetFilter, FGameplayTagRequirements SourceTagRequirements, FGameplayTagRequirements TargetTagRequirements, bool TriggerOnce=false, AActor* OptionalExternalOwner=nullptr, bool ListenForPeriodicEffects=false);
+	static UE_API UAbilityTask_WaitGameplayEffectApplied_Target* WaitGameplayEffectAppliedToTarget(UGameplayAbility* OwningAbility, const FGameplayTargetDataFilterHandle TargetFilter, FGameplayTagRequirements SourceTagRequirements, FGameplayTagRequirements TargetTagRequirements, FGameplayTagRequirements AssetTagRequirements, FGameplayTagRequirements GrantedTagRequirements, bool TriggerOnce=false, AActor* OptionalExternalOwner=nullptr, bool ListenForPeriodicEffects=false);
 
 	/** 
-	 * Wait until the owner (or External Owner) applies a GameplayEffect to a Target (the target may be the owner too!). If TriggerOnce is true, this task will only return one time. Otherwise it will return everytime a GE is applied that meets the requirements over the life of the ability 
-	 * Optional External Owner can be used to run this task on someone else (not the owner of the ability). By default you can leave this empty.
+	 * Wait until the owner (or External Owner) applies a GameplayEffect to a Target (the target may be the owner too!).
+	 * This version uses FGameplayTagQuery (more power) instead of FGameplayTagRequirements (faster).
+	 * 
+	 * @param OwningAbility The ability that executes this task
+	 * @param SourceFilter Actor filter for the source actor applying this effect that must be passed to trigger this task
+	 * @param SourceTagQuery Gameplay tag requirements for the source actor applying this effect
+	 * @param TargetTagQuery Gameplay tag requirements for the target actor receiving this effect
+	 * @param AssetTagQuery Requirements for the asset tags of the applied effect: 'Tags this Effect Has' on the GameplayEffect class
+	 * @param GrantedTagQuery Requirements for the tags granted by the applied effect: 'Grant Tags to Target Actor' on the GameplayEffect class
+	 * @param TriggerOnce If TriggerOnce is true, this task will only return one time. Otherwise it will return everytime a GE is applied that meets the requirements over the life of the ability.
+	 * @param OptionalExternalOwner Can be used to run this task on someone else (not the owner of the ability). By default you can leave this empty.
+	 * @param ListenForPeriodicEffects Whether to also fire on periodic ticks of the applied effect
 	 */
 	UFUNCTION(BlueprintCallable, Category="Ability|Tasks", meta = (HidePin = "OwningAbility", DefaultToSelf = "OwningAbility", BlueprintInternalUseOnly = "TRUE"))
-	static UE_API UAbilityTask_WaitGameplayEffectApplied_Target* WaitGameplayEffectAppliedToTarget_Query(UGameplayAbility* OwningAbility, const FGameplayTargetDataFilterHandle SourceFilter, FGameplayTagQuery SourceTagQuery, FGameplayTagQuery TargetTagQuery, bool TriggerOnce=false, AActor* OptionalExternalOwner=nullptr, bool ListenForPeriodicEffect=false );
+	static UE_API UAbilityTask_WaitGameplayEffectApplied_Target* WaitGameplayEffectAppliedToTarget_Query(UGameplayAbility* OwningAbility, const FGameplayTargetDataFilterHandle SourceFilter, FGameplayTagQuery SourceTagQuery, FGameplayTagQuery TargetTagQuery, FGameplayTagQuery AssetTagQuery, FGameplayTagQuery GrantedTagQuery, bool TriggerOnce=false, AActor* OptionalExternalOwner=nullptr, bool ListenForPeriodicEffect=false );
 	
 protected:
 
