GAS: Exposing common GameplayEffectSpecHandle and ActiveGameplayEffectHandle queries to blueprint.
* Getting GE (definition) from GE Spec Handle
* Getting AssetTags and GrantedTags from GE Spec Handle
* Getting Duration Policy from GE Spec Handle
* Checking Active GE Handle validity
* Checking whether Active GE Handle is still active
* Getting ASC from Active GE Handle

diff --git a/Engine/Plugins/Runtime/GameplayAbilities/Source/GameplayAbilities/Private/AbilitySystemBlueprintLibrary.cpp b/Engine/Plugins/Runtime/GameplayAbilities/Source/GameplayAbilities/Private/AbilitySystemBlueprintLibrary.cpp
--- a/Engine/Plugins/Runtime/GameplayAbilities/Source/GameplayAbilities/Private/AbilitySystemBlueprintLibrary.cpp
+++ b/Engine/Plugins/Runtime/GameplayAbilities/Source/GameplayAbilities/Private/AbilitySystemBlueprintLibrary.cpp
@@ -1061,6 +1061,17 @@ FGameplayEffectSpecHandle UAbilitySystemBlueprintLibrary::AddGrantedTags(FGamepl
 	return SpecHandle;
 }
 
+FGameplayTagContainer UAbilitySystemBlueprintLibrary::GetGrantedTags(FGameplayEffectSpecHandle SpecHandle)
+{
+	FGameplayTagContainer AllGrantedTags;
+	FGameplayEffectSpec* Spec = SpecHandle.Data.Get();
+	if (Spec)
+	{
+		Spec->GetAllGrantedTags(AllGrantedTags);
+	}
+	return AllGrantedTags;
+}
+
 FGameplayEffectSpecHandle UAbilitySystemBlueprintLibrary::AddAssetTag(FGameplayEffectSpecHandle SpecHandle, FGameplayTag NewGameplayTag)
 {
 	FGameplayEffectSpec* Spec = SpecHandle.Data.Get();
@@ -1090,6 +1101,17 @@ FGameplayEffectSpecHandle UAbilitySystemBlueprintLibrary::AddAssetTags(FGameplay
 
 	return SpecHandle;
 }
+
+FGameplayTagContainer UAbilitySystemBlueprintLibrary::GetAssetTags(FGameplayEffectSpecHandle SpecHandle)
+{
+	FGameplayTagContainer AllAssetTags;
+	FGameplayEffectSpec* Spec = SpecHandle.Data.Get();
+	if (Spec)
+	{
+		Spec->GetAllAssetTags(AllAssetTags);
+	}
+	return AllAssetTags;
+}
 	
 FGameplayEffectSpecHandle UAbilitySystemBlueprintLibrary::AddLinkedGameplayEffectSpec(FGameplayEffectSpecHandle SpecHandle, FGameplayEffectSpecHandle LinkedGameplayEffectSpec)
 {
@@ -1191,6 +1213,21 @@ PRAGMA_DISABLE_DEPRECATION_WARNINGS
 PRAGMA_ENABLE_DEPRECATION_WARNINGS
 }
 
+bool UAbilitySystemBlueprintLibrary::IsActiveGameplayEffectHandleValid(FActiveGameplayEffectHandle Handle)
+{
+	return Handle.IsValid();
+}
+
+bool UAbilitySystemBlueprintLibrary::IsActiveGameplayEffectHandleActive(FActiveGameplayEffectHandle Handle)
+{
+	return Handle.IsValid() && Handle.GetOwningAbilitySystemComponent() && Handle.GetOwningAbilitySystemComponent()->GetActiveGameplayEffect(Handle) != nullptr;
+}
+
+UAbilitySystemComponent* UAbilitySystemBlueprintLibrary::GetAbilitySystemComponentFromActiveGameplayEffectHandle(FActiveGameplayEffectHandle Handle)
+{
+	return Handle.GetOwningAbilitySystemComponent();
+}
+
 int32 UAbilitySystemBlueprintLibrary::GetActiveGameplayEffectStackCount(FActiveGameplayEffectHandle ActiveHandle)
 {
 	UAbilitySystemComponent* ASC = ActiveHandle.GetOwningAbilitySystemComponent();
@@ -1341,6 +1378,33 @@ bool UAbilitySystemBlueprintLibrary::RemoveLooseGameplayTags(AActor* Actor, cons
 	return false;
 }
 
+const UGameplayEffect* UAbilitySystemBlueprintLibrary::GetGameplayEffectFromSpecHandle(FGameplayEffectSpecHandle Handle)
+{
+	return Handle.Data ? Handle.Data->Def.Get() : nullptr;
+}
+
+bool UAbilitySystemBlueprintLibrary::IsDurationGameplayEffectSpecHandle(FGameplayEffectSpecHandle Handle)
+{
+	if (Handle.Data && Handle.Data->Def)
+	{
+		return Handle.Data->Def->DurationPolicy != EGameplayEffectDurationType::Instant;
+	}
+
+	ABILITY_LOG(Warning, TEXT("IsDurationGameplayEffectSpecHandle called on invalid handle"));
+	return false;
+}
+
+EGameplayEffectDurationType UAbilitySystemBlueprintLibrary::GetDurationPolicyFromGameplayEffectSpecHandle(FGameplayEffectSpecHandle Handle)
+{
+	if (Handle.Data && Handle.Data->Def)
+	{
+		return Handle.Data->Def->DurationPolicy;
+	}
+
+	ABILITY_LOG(Warning, TEXT("GetDurationPolicyFromGameplayEffectSpecHandle called on invalid handle"));
+	return EGameplayEffectDurationType::Instant;
+}
+
 const UGameplayEffectUIData* UAbilitySystemBlueprintLibrary::GetGameplayEffectUIData(TSubclassOf<UGameplayEffect> EffectClass, TSubclassOf<UGameplayEffectUIData> DataType)
 {
 	if (const UGameplayEffect* Effect = EffectClass.GetDefaultObject())
diff --git a/Engine/Plugins/Runtime/GameplayAbilities/Source/GameplayAbilities/Public/AbilitySystemBlueprintLibrary.h b/Engine/Plugins/Runtime/GameplayAbilities/Source/GameplayAbilities/Public/AbilitySystemBlueprintLibrary.h
--- a/Engine/Plugins/Runtime/GameplayAbilities/Source/GameplayAbilities/Public/AbilitySystemBlueprintLibrary.h
+++ b/Engine/Plugins/Runtime/GameplayAbilities/Source/GameplayAbilities/Public/AbilitySystemBlueprintLibrary.h
@@ -438,6 +438,10 @@ public:
 	UFUNCTION(BlueprintCallable, Category = "Ability|GameplayEffect")
 	static UE_API FGameplayEffectSpecHandle AddGrantedTags(FGameplayEffectSpecHandle SpecHandle, FGameplayTagContainer NewGameplayTags);
 
+	/** Return all tags granted by this GameplayEffectSpec: both from the GE asset and dynamic granted tags. */
+	UFUNCTION(BlueprintCallable, Category = "Ability|GameplayEffect")
+	static UE_API FGameplayTagContainer GetGrantedTags(FGameplayEffectSpecHandle SpecHandle);
+
 	/** Adds NewGameplayTag to this instance of the effect */
 	UFUNCTION(BlueprintCallable, Category = "Ability|GameplayEffect")
 	static UE_API FGameplayEffectSpecHandle AddAssetTag(FGameplayEffectSpecHandle SpecHandle, FGameplayTag NewGameplayTag);
@@ -446,6 +450,10 @@ public:
 	UFUNCTION(BlueprintCallable, Category = "Ability|GameplayEffect")
 	static UE_API FGameplayEffectSpecHandle AddAssetTags(FGameplayEffectSpecHandle SpecHandle, FGameplayTagContainer NewGameplayTags);
 
+	/** Return all asset tags of this GameplayEffectSpec: both from the GE asset and dynamic asset tags. */
+	UFUNCTION(BlueprintCallable, Category = "Ability|GameplayEffect")
+	static UE_API FGameplayTagContainer GetAssetTags(FGameplayEffectSpecHandle SpecHandle);
+
 	/** Adds LinkedGameplayEffectSpec to SpecHandles. LinkedGameplayEffectSpec will be applied when/if SpecHandle is applied successfully. LinkedGameplayEffectSpec will not be modified here. Returns the ORIGINAL SpecHandle (legacy decision) */
 	UE_DEPRECATED(5.3, "Linked GameplayEffects aren't replicated.  The new UAdditionalGameplayEffectsComponent renders this functionality obsolete (and since it's configured in the Asset, it is properly synced).")
 	UFUNCTION(BlueprintCallable, Category = "Ability|GameplayEffect", meta = (DeprecatedFunction, DeprecatedMessage="Linked GameplayEffects aren't replicated.  Configure the GameplayEffect asset with a suitable GameplayEffectComponent."))
@@ -481,6 +489,18 @@ public:
 	UFUNCTION(BlueprintCallable, Category="Ability|GameplayEffect")
 	static UE_API bool RemoveLooseGameplayTags(AActor* Actor, const FGameplayTagContainer& GameplayTags, bool bShouldReplicate=false);
 
+	/** Get the GameplayEffect definition from a GameplayEffectSpecHandle. */
+	UFUNCTION(BlueprintCallable, Category = "Ability|GameplayEffect")
+	static UE_API const UGameplayEffect* GetGameplayEffectFromSpecHandle(FGameplayEffectSpecHandle Handle);
+
+	/** Returns whether the GameplayEffectSpecHandle represents a non-instant GameplayEffect. */
+	UFUNCTION(BlueprintCallable, Category = "Ability|GameplayEffect")
+	static UE_API bool IsDurationGameplayEffectSpecHandle(FGameplayEffectSpecHandle Handle);
+
+	/** Returns the Duration Policy of a GameplayEffectSpec. */
+	UFUNCTION(BlueprintCallable, Category = "Ability|GameplayEffect")
+	static UE_API EGameplayEffectDurationType GetDurationPolicyFromGameplayEffectSpecHandle(FGameplayEffectSpecHandle Handle);
+
 	// -------------------------------------------------------------------------------
 	//		GameplayEffectSpec
 	// -------------------------------------------------------------------------------
@@ -495,7 +515,19 @@ public:
 	// -------------------------------------------------------------------------------
 	//		FActiveGameplayEffectHandle
 	// -------------------------------------------------------------------------------
-	
+
+	/** Returns whether the active gameplay effect handle has a valid handle value, regardless of whether the effect expired. */
+	UFUNCTION(BlueprintCallable, Category = "Ability|GameplayEffect")
+	static bool IsActiveGameplayEffectHandleValid(FActiveGameplayEffectHandle Handle);
+
+	/** Returns whether the active gameplay effect handle represents a gameplay effect that is valid and still active. */
+	UFUNCTION(BlueprintCallable, Category = "Ability|GameplayEffect")
+	static bool IsActiveGameplayEffectHandleActive(FActiveGameplayEffectHandle Handle);
+
+	/** Returns AbilitySystemComponent from an ActiveGameplayEffectHandle. */
+	UFUNCTION(BlueprintCallable, Category = "Ability|GameplayEffect")
+	static class UAbilitySystemComponent* GetAbilitySystemComponentFromActiveGameplayEffectHandle(FActiveGameplayEffectHandle Handle);
+
 	/** Returns current stack count of an active Gameplay Effect. Will return 0 if the GameplayEffect is no longer valid. */
 	UFUNCTION(BlueprintCallable, Category = "Ability|GameplayEffect")
 	static UE_API int32 GetActiveGameplayEffectStackCount(FActiveGameplayEffectHandle ActiveHandle);
