[GameplayDebugger] fixed nullptr accesses when using player controller in pure-editor mode (i.e., UGameplayDebuggerUserSettings::bEnableGameplayDebuggerInEditor)

diff --git a/Engine/Plugins/Runtime/GameplayAbilities/Source/GameplayAbilities/Private/GameplayDebuggerCategory_Abilities.cpp b/Engine/Plugins/Runtime/GameplayAbilities/Source/GameplayAbilities/Private/GameplayDebuggerCategory_Abilities.cpp
--- a/Engine/Plugins/Runtime/GameplayAbilities/Source/GameplayAbilities/Private/GameplayDebuggerCategory_Abilities.cpp
+++ b/Engine/Plugins/Runtime/GameplayAbilities/Source/GameplayAbilities/Private/GameplayDebuggerCategory_Abilities.cpp
@@ -205,7 +205,7 @@ TArray<FGameplayDebuggerCategory_Abilities::FRepData::FGameplayEffectDebug> FGam
 		ItemData.Stacks = EffectSpec.GetStackCount();
 		ItemData.Level = EffectSpec.GetLevel();
 
-		ItemData.NetworkStatus = OwnerPC->HasAuthority() ? ENetworkStatus::ServerOnly : ENetworkStatus::LocalOnly;
+		ItemData.NetworkStatus = OwnerPC != nullptr && OwnerPC->HasAuthority() ? ENetworkStatus::ServerOnly : ENetworkStatus::LocalOnly;
 	}
 
 	return DebugEffects;
@@ -234,7 +234,7 @@ TArray<FGameplayDebuggerCategory_Abilities::FRepData::FGameplayAttributeDebug> F
 		if (const UNetworkSubsystem* NetConditionGroupSubsystem = Actor->GetWorld()->GetSubsystem<UNetworkSubsystem>())
 		{
 			FReplicationFlags RepFlags;
-			RepFlags.bNetOwner = (Actor->GetNetConnection() == OwnerPC->GetNetConnection());
+			RepFlags.bNetOwner = OwnerPC != nullptr && (Actor->GetNetConnection() == OwnerPC->GetNetConnection());
 			RepFlags.bNetSimulated = !RepFlags.bNetOwner && Actor->GetNetConnection();
 			RepFlags.bRolesOnly = true;
 			const TStaticBitArray<COND_Max> ConditionMap = UE::Net::BuildConditionMapFromRepFlags(RepFlags);
@@ -520,7 +520,7 @@ void FGameplayDebuggerCategory_Abilities::DrawGameplayEffects(FGameplayDebuggerC
 
 	// Merge together the server and local data
 	const UAbilitySystemComponent* LocalASC = UAbilitySystemGlobals::GetAbilitySystemComponentFromActor(FindLocalDebugActor());
-	const bool bConsiderLocalStatus = !OwnerPC->IsNetMode(ENetMode::NM_Standalone) && LocalASC && (LocalASC->ReplicationMode != EGameplayEffectReplicationMode::Full);
+	const bool bConsiderLocalStatus = OwnerPC != nullptr && !OwnerPC->IsNetMode(ENetMode::NM_Standalone) && LocalASC && (LocalASC->ReplicationMode != EGameplayEffectReplicationMode::Full);
 	const TArray<FGameplayEffectDebug>& ServerEffects = DataPack.GameplayEffects;
 	TArray<FGameplayEffectDebug> LocalEffects = bConsiderLocalStatus ? CollectEffectsData(OwnerPC, LocalASC) : ServerEffects;
 
@@ -701,7 +701,7 @@ void FGameplayDebuggerCategory_Abilities::DrawGameplayAttributes(FGameplayDebugg
 		ENetworkStatus NetworkStatus = ENetworkStatus::LocalOnly;
 	};
 
-	const bool bConsiderNetworkStatus = !OwnerPC->IsNetMode(ENetMode::NM_Standalone);
+	const bool bConsiderNetworkStatus = OwnerPC != nullptr && !OwnerPC->IsNetMode(ENetMode::NM_Standalone);
 	const UAbilitySystemComponent* LocalASC = bConsiderNetworkStatus ? UAbilitySystemGlobals::GetAbilitySystemComponentFromActor(FindLocalDebugActor()) : nullptr;
 	const TArray<FRepData::FGameplayAttributeDebug> LocalAttributes = LocalASC ? CollectAttributeData(OwnerPC, LocalASC) : TArray<FRepData::FGameplayAttributeDebug>{};
 	TArray<FRepData::FGameplayAttributeDebug> ServerAttributes = DataPack.Attributes;
diff --git a/Engine/Source/Runtime/AIModule/Private/GameplayDebugger/GameplayDebuggerCategory_EQS.cpp b/Engine/Source/Runtime/AIModule/Private/GameplayDebugger/GameplayDebuggerCategory_EQS.cpp
--- a/Engine/Source/Runtime/AIModule/Private/GameplayDebugger/GameplayDebuggerCategory_EQS.cpp
+++ b/Engine/Source/Runtime/AIModule/Private/GameplayDebugger/GameplayDebuggerCategory_EQS.cpp
@@ -57,8 +57,10 @@ void FGameplayDebuggerCategory_EQS::FRepData::Serialize(FArchive& Ar)
 void FGameplayDebuggerCategory_EQS::CollectData(APlayerController* OwnerPC, AActor* DebugActor)
 {
 #if USE_EQS_DEBUGGER
-	UWorld* World = OwnerPC->GetWorld();
-	UEnvQueryManager* QueryManager = World ? UEnvQueryManager::GetCurrent(World) : nullptr;
+	UWorld* World = GetDataWorld(OwnerPC, DebugActor);
+	check(World);
+
+	UEnvQueryManager* QueryManager = UEnvQueryManager::GetCurrent(World);
 	if (QueryManager == nullptr)
 	{
 		return;
