Deprecated GetSubobjectsWithStableNamesForNetworking function

diff --git a/Engine/Plugins/Runtime/GameplayAbilities/Source/GameplayAbilities/Private/AbilitySystemComponent.cpp b/Engine/Plugins/Runtime/GameplayAbilities/Source/GameplayAbilities/Private/AbilitySystemComponent.cpp
--- a/Engine/Plugins/Runtime/GameplayAbilities/Source/GameplayAbilities/Private/AbilitySystemComponent.cpp
+++ b/Engine/Plugins/Runtime/GameplayAbilities/Source/GameplayAbilities/Private/AbilitySystemComponent.cpp
@@ -1750,17 +1750,6 @@ bool UAbilitySystemComponent::ReplicateSubobjects(UActorChannel* Channel, FOutBu
 	return WroteSomething;
 }
 
-void UAbilitySystemComponent::GetSubobjectsWithStableNamesForNetworking(TArray<UObject*>& Objs)
-{
-	for (const UAttributeSet* Set : GetSpawnedAttributes())
-	{
-		if (Set && Set->IsNameStableForNetworking())
-		{
-			Objs.Add(const_cast<UAttributeSet*>(Set));
-		}
-	}
-}
-
 void UAbilitySystemComponent::PreNetReceive()
 {
 	// Update the cached IsNetSimulated value here if this component is still considered authority.
diff --git a/Engine/Plugins/Runtime/GameplayAbilities/Source/GameplayAbilities/Public/AbilitySystemComponent.h b/Engine/Plugins/Runtime/GameplayAbilities/Source/GameplayAbilities/Public/AbilitySystemComponent.h
--- a/Engine/Plugins/Runtime/GameplayAbilities/Source/GameplayAbilities/Public/AbilitySystemComponent.h
+++ b/Engine/Plugins/Runtime/GameplayAbilities/Source/GameplayAbilities/Public/AbilitySystemComponent.h
@@ -1667,7 +1667,6 @@ public:
 	UE_API virtual void OnComponentDestroyed(bool bDestroyingHierarchy) override;
 	UE_API virtual bool GetShouldTick() const override;
 	UE_API virtual void TickComponent(float DeltaTime, enum ELevelTick TickType, FActorComponentTickFunction *ThisTickFunction) override;
-	UE_API virtual void GetSubobjectsWithStableNamesForNetworking(TArray<UObject*>& Objs) override;
 	UE_API virtual bool ReplicateSubobjects(class UActorChannel *Channel, class FOutBunch *Bunch, FReplicationFlags *RepFlags) override;
 	/** Force owning actor to update it's replication, to make sure that gameplay cues get sent down quickly. Override to change how aggressive this is */
 	UE_API virtual void ForceReplication() override;
diff --git a/Engine/Source/Runtime/CoreUObject/Public/UObject/Object.h b/Engine/Source/Runtime/CoreUObject/Public/UObject/Object.h
--- a/Engine/Source/Runtime/CoreUObject/Public/UObject/Object.h
+++ b/Engine/Source/Runtime/CoreUObject/Public/UObject/Object.h
@@ -1057,6 +1057,7 @@ public:
 	COREUOBJECT_API virtual bool IsSupportedForNetworking() const;
 
 	/** Returns a list of sub-objects that have stable names for networking */
+	UE_DEPRECATED(5.7, "This function is useless now and will be deleted")
 	virtual void GetSubobjectsWithStableNamesForNetworking(TArray<UObject*> &ObjList) {}
 
 	/** Called right before receiving a bunch */
diff --git a/Engine/Source/Runtime/Engine/Classes/GameFramework/Actor.h b/Engine/Source/Runtime/Engine/Classes/GameFramework/Actor.h
--- a/Engine/Source/Runtime/Engine/Classes/GameFramework/Actor.h
+++ b/Engine/Source/Runtime/Engine/Classes/GameFramework/Actor.h
@@ -2948,9 +2948,6 @@ public:
 	/** IsSupportedForNetworking means an object can be referenced over the network */
 	ENGINE_API virtual bool IsSupportedForNetworking() const override;
 
-	/** Returns a list of sub-objects that have stable names for networking */
-	ENGINE_API virtual void GetSubobjectsWithStableNamesForNetworking(TArray<UObject*> &ObjList) override;
-
 	/** Always called immediately after spawning and reading in replicated properties */
 	ENGINE_API virtual void PostNetInit();
 
diff --git a/Engine/Source/Runtime/Engine/Private/ActorReplication.cpp b/Engine/Source/Runtime/Engine/Private/ActorReplication.cpp
--- a/Engine/Source/Runtime/Engine/Private/ActorReplication.cpp
+++ b/Engine/Source/Runtime/Engine/Private/ActorReplication.cpp
@@ -984,22 +984,6 @@ void AActor::BuildReplicatedComponentsInfo()
 	}
 }
 
-void AActor::GetSubobjectsWithStableNamesForNetworking(TArray<UObject*> &ObjList)
-{	
-	// For experimenting with replicating ALL stably named components initially
-	for (UActorComponent* Component : OwnedComponents)
-	{
-		if (IsValid(Component) && Component->IsNameStableForNetworking())
-		{
-			ObjList.Add(Component);
-			Component->GetSubobjectsWithStableNamesForNetworking(ObjList);
-		}
-	}
-
-	// Sort the list so that we generate the same list on client/server
-	Algo::SortBy(ObjList, &UObject::GetFName, FNameLexicalLess());
-}
-
 void AActor::OnSubobjectCreatedFromReplication(UObject *NewSubobject)
 {
 	check(NewSubobject);
