Changes to Gameplay Tag Count Replication
These changes allow more fine grained control of replication per gameplay tag in a single unified container.

diff --git a/Engine/Plugins/Runtime/GameplayAbilities/Source/GameplayAbilities/Private/Abilities/GameplayAbility.cpp b/Engine/Plugins/Runtime/GameplayAbilities/Source/GameplayAbilities/Private/Abilities/GameplayAbility.cpp
--- a/Engine/Plugins/Runtime/GameplayAbilities/Source/GameplayAbilities/Private/Abilities/GameplayAbility.cpp
+++ b/Engine/Plugins/Runtime/GameplayAbilities/Source/GameplayAbilities/Private/Abilities/GameplayAbility.cpp
@@ -871,19 +871,6 @@ void UGameplayAbility::EndAbility(const FGameplayAbilitySpecHandle Handle, const
 			// Remove tags we added to owner
 			AbilitySystemComponent->RemoveLooseGameplayTags(ActivationOwnedTags);
 
-			if (UAbilitySystemGlobals::Get().ShouldReplicateActivationOwnedTags())
-			{
-				if (GetNetExecutionPolicy() == EGameplayAbilityNetExecutionPolicy::LocalPredicted || GetNetExecutionPolicy() == EGameplayAbilityNetExecutionPolicy::ServerInitiated)
-				{
-					// If this ability also executes on the client, then don't communicate the tags to the client (it already used RemoveLooseGameplayTags above)
-					AbilitySystemComponent->RemoveMinimalReplicationGameplayTags(ActivationOwnedTags);
-				}
-				else
-				{
-					AbilitySystemComponent->RemoveReplicatedLooseGameplayTags(ActivationOwnedTags);
-				}
-			}
-
 			// Remove tracked GameplayCues that we added
 			for (FGameplayTag& GameplayCueTag : TrackedGameplayCues)
 			{
@@ -995,20 +982,7 @@ void UGameplayAbility::PreActivate(const FGameplayAbilitySpecHandle Handle, cons
 
 	Comp->HandleChangeAbilityCanBeCanceled(GetAssetTags(), this, true);
 
-	Comp->AddLooseGameplayTags(ActivationOwnedTags);
-
-	if (UAbilitySystemGlobals::Get().ShouldReplicateActivationOwnedTags())
-	{
-		if (GetNetExecutionPolicy() == EGameplayAbilityNetExecutionPolicy::LocalPredicted || GetNetExecutionPolicy() == EGameplayAbilityNetExecutionPolicy::ServerInitiated)
-		{
-			// If this ability also executes on the client, then don't communicate the tags to the client (it already used AddLooseGameplayTags above)
-			Comp->AddMinimalReplicationGameplayTags(ActivationOwnedTags);
-		}
-		else
-		{
-			Comp->AddReplicatedLooseGameplayTags(ActivationOwnedTags);
-		}
-	}
+	Comp->AddLooseGameplayTags(ActivationOwnedTags, 1, UAbilitySystemGlobals::Get().ShouldReplicateActivationOwnedTags() ? EGameplayTagReplicationState::CountToOwner : EGameplayTagReplicationState::None);
 
 	if (OnGameplayAbilityEndedDelegate)
 	{
diff --git a/Engine/Plugins/Runtime/GameplayAbilities/Source/GameplayAbilities/Private/AbilitySystemBlueprintLibrary.cpp b/Engine/Plugins/Runtime/GameplayAbilities/Source/GameplayAbilities/Private/AbilitySystemBlueprintLibrary.cpp
--- a/Engine/Plugins/Runtime/GameplayAbilities/Source/GameplayAbilities/Private/AbilitySystemBlueprintLibrary.cpp
+++ b/Engine/Plugins/Runtime/GameplayAbilities/Source/GameplayAbilities/Private/AbilitySystemBlueprintLibrary.cpp
@@ -1348,12 +1348,7 @@ bool UAbilitySystemBlueprintLibrary::AddLooseGameplayTags(AActor* Actor, const F
 {
 	if (UAbilitySystemComponent* AbilitySysComp = GetAbilitySystemComponent(Actor))
 	{
-		AbilitySysComp->AddLooseGameplayTags(GameplayTags);
-
-		if (bShouldReplicate)
-		{
-			AbilitySysComp->AddReplicatedLooseGameplayTags(GameplayTags);
-		}
+		AbilitySysComp->AddLooseGameplayTags(GameplayTags, 1, bShouldReplicate ? EGameplayTagReplicationState::CountToOwner : EGameplayTagReplicationState::None);
 
 		return true;
 	}
@@ -1365,12 +1360,31 @@ bool UAbilitySystemBlueprintLibrary::RemoveLooseGameplayTags(AActor* Actor, cons
 {
 	if (UAbilitySystemComponent* AbilitySysComp = GetAbilitySystemComponent(Actor))
 	{
-		AbilitySysComp->RemoveLooseGameplayTags(GameplayTags);
+		AbilitySysComp->RemoveLooseGameplayTags(GameplayTags, 1);
 
-		if (bShouldReplicate)
-		{
-			AbilitySysComp->RemoveReplicatedLooseGameplayTags(GameplayTags);
-		}
+		return true;
+	}
+
+	return false;
+}
+
+bool UAbilitySystemBlueprintLibrary::AddGameplayTags(AActor* Actor, const FGameplayTagContainer& GameplayTags, EGameplayTagReplicationState ReplicationRule)
+{
+	if (UAbilitySystemComponent* AbilitySysComp = GetAbilitySystemComponent(Actor))
+	{
+		AbilitySysComp->AddLooseGameplayTags(GameplayTags, 1, ReplicationRule);
+
+		return true;
+	}
+
+	return false;
+}
+
+bool UAbilitySystemBlueprintLibrary::RemoveGameplayTags(AActor* Actor, const FGameplayTagContainer& GameplayTags)
+{
+	if (UAbilitySystemComponent* AbilitySysComp = GetAbilitySystemComponent(Actor))
+	{
+		AbilitySysComp->RemoveLooseGameplayTags(GameplayTags, 1);
 
 		return true;
 	}
diff --git a/Engine/Plugins/Runtime/GameplayAbilities/Source/GameplayAbilities/Private/AbilitySystemComponent.cpp b/Engine/Plugins/Runtime/GameplayAbilities/Source/GameplayAbilities/Private/AbilitySystemComponent.cpp
--- a/Engine/Plugins/Runtime/GameplayAbilities/Source/GameplayAbilities/Private/AbilitySystemComponent.cpp
+++ b/Engine/Plugins/Runtime/GameplayAbilities/Source/GameplayAbilities/Private/AbilitySystemComponent.cpp
@@ -214,12 +214,11 @@ void UAbilitySystemComponent::OnRegister()
 	UpdateActiveGameplayEffectsReplicationCondition();
 	UpdateMinimalReplicationGameplayCuesCondition();
 
+	GameplayTagCountContainer.SetOwner(this);
 PRAGMA_DISABLE_DEPRECATION_WARNINGS
-	// This field is not replicated (MinimalReplicationTags has a custom serializer),
-	// so we don't need to mark it dirty.
 	MinimalReplicationTags.Owner = this;
-PRAGMA_ENABLE_DEPRECATION_WARNINGS
 	ReplicatedLooseTags.Owner = this;
+PRAGMA_ENABLE_DEPRECATION_WARNINGS
 
 	/** Allocate an AbilityActorInfo. Note: this goes through a global function and is a SharedPtr so projects can make their own AbilityActorInfo */
 	if(!AbilityActorInfo.IsValid())
@@ -638,7 +637,7 @@ FOnGivenActiveGameplayEffectRemoved& UAbilitySystemComponent::OnAnyGameplayEffec
 	return ActiveGameplayEffects.OnActiveGameplayEffectRemovedDelegate;
 }
 
-void UAbilitySystemComponent::UpdateTagMap_Internal(const FGameplayTagContainer& Container, const int32 CountDelta)
+void UAbilitySystemComponent::UpdateTagMap_Internal(const FGameplayTagContainer& Container, const int32 CountDelta, EGameplayTagReplicationState TagRepState)
 {
 	// For removal, reorder calls so that FillParentTags is only called once
 	if (CountDelta > 0)
@@ -646,7 +645,7 @@ void UAbilitySystemComponent::UpdateTagMap_Internal(const FGameplayTagContainer&
 		for (auto TagIt = Container.CreateConstIterator(); TagIt; ++TagIt)
 		{
 			const FGameplayTag& Tag = *TagIt;
-			if (GameplayTagCountContainer.UpdateTagCount(Tag, CountDelta))
+			if (GameplayTagCountContainer.UpdateTagCount(Tag, CountDelta, TagRepState))
 			{
 				OnTagUpdated(Tag, true);
 			}
@@ -1650,6 +1649,7 @@ void UAbilitySystemComponent::GetLifetimeReplicatedProps(TArray< FLifetimeProper
 	DOREPLIFETIME_WITH_PARAMS_FAST(UAbilitySystemComponent, RepAnimMontageInfo, Params);
 	DOREPLIFETIME_WITH_PARAMS_FAST(UAbilitySystemComponent, OwnerActor, Params);
 	DOREPLIFETIME_WITH_PARAMS_FAST(UAbilitySystemComponent, AvatarActor, Params);
+	DOREPLIFETIME_WITH_PARAMS_FAST(UAbilitySystemComponent, GameplayTagCountContainer, Params);
 	DOREPLIFETIME_WITH_PARAMS_FAST(UAbilitySystemComponent, ReplicatedLooseTags, Params);
 
 	Params.Condition = COND_ReplayOrOwner;
diff --git a/Engine/Plugins/Runtime/GameplayAbilities/Source/GameplayAbilities/Private/GameplayCueInterface.cpp b/Engine/Plugins/Runtime/GameplayAbilities/Source/GameplayAbilities/Private/GameplayCueInterface.cpp
--- a/Engine/Plugins/Runtime/GameplayAbilities/Source/GameplayAbilities/Private/GameplayCueInterface.cpp
+++ b/Engine/Plugins/Runtime/GameplayAbilities/Source/GameplayAbilities/Private/GameplayCueInterface.cpp
@@ -268,7 +268,7 @@ void FActiveGameplayCue::PreReplicatedRemove(const struct FActiveGameplayCueCont
 	if (bPredictivelyRemoved == false)
 	{
 		// If predicted ignore the add/remove
-		InArray.Owner->UpdateTagMap(GameplayCueTag, -1);
+		InArray.Owner->UpdateTagMap(GameplayCueTag, -1, EGameplayTagReplicationState::None);
 		InArray.Owner->InvokeGameplayCueEvent(GameplayCueTag, EGameplayCueEvent::Removed, Parameters);
 	}
 }
@@ -288,7 +288,7 @@ void FActiveGameplayCue::PostReplicatedAdd(const struct FActiveGameplayCueContai
 		UE_VLOG(PropertyOwner, LogGameplayCues, VeryVerbose, TEXT("%s::%hs %s %s"), *PropertyName, __func__, *GameplayCueTag.ToString(), *PredictionKey.ToString());
 	}
 
-	InArray.Owner->UpdateTagMap(GameplayCueTag, 1);
+	InArray.Owner->UpdateTagMap(GameplayCueTag, 1, EGameplayTagReplicationState::None);
 
 	if (PredictionKey.IsLocalClientKey() == false)
 	{
@@ -324,7 +324,7 @@ void FActiveGameplayCueContainer::AddCue(const FGameplayTag& Tag, const FPredict
 	NewCue.Parameters = Parameters;
 	MarkItemDirty(NewCue);
 	
-	Owner->UpdateTagMap(Tag, 1);
+	Owner->UpdateTagMap(Tag, 1, EGameplayTagReplicationState::None);
 }
 
 void FActiveGameplayCueContainer::RemoveCue(const FGameplayTag& Tag)
@@ -355,7 +355,7 @@ void FActiveGameplayCueContainer::RemoveCue(const FGameplayTag& Tag)
 
 		if (CountDelta > 0)
 		{
-			Owner->UpdateTagMap(Tag, -CountDelta);
+			Owner->UpdateTagMap(Tag, -CountDelta, EGameplayTagReplicationState::None);
 		}
 	}
 	
@@ -369,7 +369,7 @@ void FActiveGameplayCueContainer::RemoveCue(const FGameplayTag& Tag)
 			if (GameplayCueInterfacePrivate::bUseEqualTagCountAndRemovalCallbacks)
 			{
 				++CountDelta;
-				Owner->UpdateTagMap(Tag, -1);
+				Owner->UpdateTagMap(Tag, -1, EGameplayTagReplicationState::None);
 			}
 
 			Owner->InvokeGameplayCueEvent(Tag, EGameplayCueEvent::Removed, Cue.Parameters);
@@ -394,7 +394,7 @@ void FActiveGameplayCueContainer::RemoveAllCues()
 	for (int32 idx=0; idx < GameplayCues.Num(); ++idx)
 	{
 		FActiveGameplayCue& Cue = GameplayCues[idx];
-		Owner->UpdateTagMap(Cue.GameplayCueTag, -1);
+		Owner->UpdateTagMap(Cue.GameplayCueTag, -1, EGameplayTagReplicationState::None);
 		Owner->InvokeGameplayCueEvent(Cue.GameplayCueTag, EGameplayCueEvent::Removed, Cue.Parameters);
 	}
 }
@@ -417,7 +417,7 @@ void FActiveGameplayCueContainer::PredictiveRemove(const FGameplayTag& Tag)
 		if (Cue.GameplayCueTag == Tag && !Cue.bPredictivelyRemoved)
 		{
 			Cue.bPredictivelyRemoved = true;
-			Owner->UpdateTagMap(Tag, -1);
+			Owner->UpdateTagMap(Tag, -1, EGameplayTagReplicationState::None);
 			Owner->InvokeGameplayCueEvent(Tag, EGameplayCueEvent::Removed, Cue.Parameters);	
 			return;
 		}
@@ -431,7 +431,7 @@ void FActiveGameplayCueContainer::PredictiveAdd(const FGameplayTag& Tag, FPredic
 		return;
 	}
 
-	Owner->UpdateTagMap(Tag, 1);	
+	Owner->UpdateTagMap(Tag, 1, EGameplayTagReplicationState::None);
 	PredictionKey.NewRejectOrCaughtUpDelegate(FPredictionKeyEvent::CreateUObject(ToRawPtr(Owner), &UAbilitySystemComponent::OnPredictiveGameplayCueCatchup, Tag));
 }
 
diff --git a/Engine/Plugins/Runtime/GameplayAbilities/Source/GameplayAbilities/Private/GameplayEffect.cpp b/Engine/Plugins/Runtime/GameplayAbilities/Source/GameplayAbilities/Private/GameplayEffect.cpp
--- a/Engine/Plugins/Runtime/GameplayAbilities/Source/GameplayAbilities/Private/GameplayEffect.cpp
+++ b/Engine/Plugins/Runtime/GameplayAbilities/Source/GameplayAbilities/Private/GameplayEffect.cpp
@@ -4493,20 +4493,12 @@ void FActiveGameplayEffectsContainer::AddActiveGameplayEffectGrantedTagsAndModif
 
 
 	// Update our owner with the tags this GameplayEffect grants them
-	Owner->UpdateTagMap(Effect.Spec.Def->GetGrantedTags(), 1);
-	Owner->UpdateTagMap(Effect.Spec.DynamicGrantedTags, 1);
+	Owner->UpdateTagMap(Effect.Spec.Def->GetGrantedTags(), 1, ShouldUseMinimalReplication() ? EGameplayTagReplicationState::TagOnly : EGameplayTagReplicationState::None);
+	Owner->UpdateTagMap(Effect.Spec.DynamicGrantedTags, 1, ShouldUseMinimalReplication() ? EGameplayTagReplicationState::TagOnly : EGameplayTagReplicationState::None);
 
 	// Update our owner with the blocked ability tags this GameplayEffect adds to them
 	Owner->BlockAbilitiesWithTags(Effect.Spec.Def->GetBlockedAbilityTags());
 
-
-	// Update minimal replication if needed.
-	if (ShouldUseMinimalReplication())
-	{
-		Owner->AddMinimalReplicationGameplayTags(Effect.Spec.Def->GetGrantedTags());
-		Owner->AddMinimalReplicationGameplayTags(Effect.Spec.DynamicGrantedTags);
-	}
-
 PRAGMA_DISABLE_DEPRECATION_WARNINGS
 	// Grant abilities
 	if (IsNetAuthority() && !Owner->bSuppressGrantAbility)
@@ -4550,7 +4542,7 @@ PRAGMA_ENABLE_DEPRECATION_WARNINGS
 			else // ActiveGameplayEffects are replicating to everyone (this path can also execute on client)
 			{
 				// Do a pseudo-AddGameplayCue (but don't add to ActiveGameplayCues so it doesn't replicate in addition to the AGE we're replicating).
-				Owner->UpdateTagMap(Cue.GameplayCueTags, 1);
+				Owner->UpdateTagMap(Cue.GameplayCueTags, 1, EGameplayTagReplicationState::None);
 
 				if (bInvokeGameplayCueEvents)
 				{
@@ -4790,19 +4782,12 @@ void FActiveGameplayEffectsContainer::RemoveActiveGameplayEffectGrantedTagsAndMo
 	}
 
 	// Update gameplaytag count and broadcast delegate if we are at 0
-	Owner->UpdateTagMap(Effect.Spec.Def->GetGrantedTags(), -1);
-	Owner->UpdateTagMap(Effect.Spec.DynamicGrantedTags, -1);
+	Owner->UpdateTagMap(Effect.Spec.Def->GetGrantedTags(), -1, ShouldUseMinimalReplication() ? EGameplayTagReplicationState::TagOnly : EGameplayTagReplicationState::None);
+	Owner->UpdateTagMap(Effect.Spec.DynamicGrantedTags, -1, ShouldUseMinimalReplication() ? EGameplayTagReplicationState::TagOnly : EGameplayTagReplicationState::None);
 
 	// Update our owner with the blocked ability tags this GameplayEffect adds to them
 	Owner->UnBlockAbilitiesWithTags(Effect.Spec.Def->GetBlockedAbilityTags());
 
-	// Update minimal replication if needed.
-	if (ShouldUseMinimalReplication())
-	{
-		Owner->RemoveMinimalReplicationGameplayTags(Effect.Spec.Def->GetGrantedTags());
-		Owner->RemoveMinimalReplicationGameplayTags(Effect.Spec.DynamicGrantedTags);
-	}
-
 PRAGMA_DISABLE_DEPRECATION_WARNINGS
 	// Cancel/remove granted abilities
 	if (IsNetAuthority())
@@ -4850,7 +4835,7 @@ PRAGMA_ENABLE_DEPRECATION_WARNINGS
 			else
 			{
 				// Perform pseudo-RemoveCue (without affecting ActiveGameplayCues, as we were not inserted there - see AddActiveGameplayEffectGrantedTagsAndModifiers)
-				Owner->UpdateTagMap(Cue.GameplayCueTags, -1);
+				Owner->UpdateTagMap(Cue.GameplayCueTags, -1, EGameplayTagReplicationState::None);
 
 				if (bInvokeGameplayCueEvents)
 				{
diff --git a/Engine/Plugins/Runtime/GameplayAbilities/Source/GameplayAbilities/Private/GameplayEffectTypes.cpp b/Engine/Plugins/Runtime/GameplayAbilities/Source/GameplayAbilities/Private/GameplayEffectTypes.cpp
--- a/Engine/Plugins/Runtime/GameplayAbilities/Source/GameplayAbilities/Private/GameplayEffectTypes.cpp
+++ b/Engine/Plugins/Runtime/GameplayAbilities/Source/GameplayAbilities/Private/GameplayEffectTypes.cpp
@@ -17,6 +17,8 @@
 #include "Misc/DataValidation.h"
 #endif
 
+#include "Iris/Serialization/NetBitStreamUtil.h"
+
 #include UE_INLINE_GENERATED_CPP_BY_NAME(GameplayEffectTypes)
 
 
@@ -500,6 +502,204 @@ FString EGameplayCueEventToString(int32 Type)
 	return e->GetNameStringByValue(Type);
 }
 
+class FNetGameplayTagCountContainerState : public INetDeltaBaseState
+{
+public:
+	struct FGameplayTagCountStateItem
+	{
+		FGameplayTag Tag;
+		int32 Count = -1;
+		EGameplayTagReplicationState ReplicationState = EGameplayTagReplicationState::None;
+
+		bool operator== (const FGameplayTagCountStateItem& Other) const
+		{
+			return Tag == Other.Tag && Count == Other.Count && ReplicationState == Other.ReplicationState;
+		}
+	};
+
+	virtual bool IsStateEqual(INetDeltaBaseState* OtherState) override
+	{
+		FNetGameplayTagCountContainerState* Other = static_cast<FNetGameplayTagCountContainerState*>(OtherState);
+		return (TagStates == Other->TagStates);
+	}
+
+	TArray<FGameplayTagCountStateItem> TagStates;
+};
+
+bool FGameplayTagCountContainer::NetDeltaSerialize(FNetDeltaSerializeInfo & DeltaParams)
+{
+	bool bPackCount = false;
+	
+	if (DeltaParams.Writer)
+	{
+		bool bIsOwner = false;
+		if (AActor* OwnerActor = Cast<UActorComponent>(DeltaParams.Object)->GetOwner())
+		{
+			if (UNetConnection* OwnerNetConnection = OwnerActor->GetNetConnection())
+			{
+				bIsOwner = OwnerNetConnection == DeltaParams.Connection;
+			}
+		}
+		
+		FBitWriter& Writer = *DeltaParams.Writer;
+
+		FNetGameplayTagCountContainerState* OldState = static_cast<FNetGameplayTagCountContainerState*>(DeltaParams.OldState);
+		FNetGameplayTagCountContainerState* NewState = new FNetGameplayTagCountContainerState();
+
+		// Keep track of the indices for changed data in the Items list so we can iterate just those after calculating the new state.
+		TArray<int32> ChangeIdx;
+
+		check(DeltaParams.NewState);
+		*DeltaParams.NewState = TSharedPtr<INetDeltaBaseState>(NewState);
+		NewState->TagStates.SetNum(Items.Num());
+		ChangeIdx.Reserve(Items.Num());
+
+		// We set ChangelistHistory and LastAckedHistory to 0, as that's the default RepLayout expects.
+		// Setting them to anything else may cause issues with initial sends.
+		uint32 OldChangelistHistory = 0;
+		uint32 OldLastAckedHistory = 0;
+
+		TArray<FNetGameplayTagCountContainerState::FGameplayTagCountStateItem> Removals = OldState ? OldState->TagStates : TArray<FNetGameplayTagCountContainerState::FGameplayTagCountStateItem>{};
+
+		// Iterate over our current items and build our new state data while removing valid entries from the Removals list so we know what's missing
+		for (int32 Idx = 0; Idx < Items.Num(); Idx++)
+		{
+			FGameplayTagCountItem& CountItem = Items[Idx];
+			FNetGameplayTagCountContainerState::FGameplayTagCountStateItem& NewItem = NewState->TagStates[Idx];
+
+			if (CountItem.ReplicationState == EGameplayTagReplicationState::None)
+			{
+				continue;
+			}
+		
+			bPackCount = CountItem.ReplicationState >= (bIsOwner ? EGameplayTagReplicationState::CountToOwner : EGameplayTagReplicationState::TagAndCountToAll);
+
+			NewItem.Tag = CountItem.Tag;
+			NewItem.ReplicationState = CountItem.ReplicationState;
+			NewItem.Count = bPackCount ? CountItem.Count : 1;
+
+			FNetGameplayTagCountContainerState::FGameplayTagCountStateItem* OldItem = Removals.FindByPredicate(
+				[NewItem] (FNetGameplayTagCountContainerState::FGameplayTagCountStateItem& Item)
+			{
+				return Item.Tag == NewItem.Tag;
+			});
+			
+			if (!OldItem)
+			{
+				ChangeIdx.Add(Idx);
+				continue;
+			}
+
+			if (OldItem->ReplicationState > NewItem.ReplicationState)
+			{
+				ABILITY_LOG(Warning, TEXT("Existing tag: %s with \"%s\" replication state now serializing with state \"%s\""), 
+				*NewItem.Tag.ToString(), *UEnum::GetDisplayValueAsText(OldItem->ReplicationState).ToString(), *UEnum::GetDisplayValueAsText(NewItem.ReplicationState).ToString());
+			}
+			if (bPackCount && OldItem->Count != NewItem.Count)
+			{
+				ChangeIdx.Add(Idx);
+			}
+
+			Removals.RemoveSingleSwap(*OldItem);
+		}
+
+		uint32 NumRemovals = Removals.Num();
+
+		// If there's no changes between the states we can just leave
+		if (ChangeIdx.Num() == 0 && NumRemovals == 0)
+		{
+			return false;
+		}
+
+		if (OldState)
+		{
+			OldChangelistHistory = DeltaParams.OldState->GetChangelistHistory();
+			OldLastAckedHistory = DeltaParams.OldState->GetLastAckedHistory();
+		}
+
+		NewState->SetChangelistHistory(OldChangelistHistory);
+		NewState->SetLastAckedHistory(OldLastAckedHistory);
+
+		// Pack any tags left in the Removals array to be removed on the client side
+		Writer.SerializeIntPacked(NumRemovals);
+		bool bTagSuccess = true;
+		if (NumRemovals > 0)
+		{
+			for (FNetGameplayTagCountContainerState::FGameplayTagCountStateItem& Item : Removals)
+			{
+				Item.Tag.NetSerialize(Writer, DeltaParams.Map, bTagSuccess);
+			}
+		}
+
+		// Finish packing up any remaining changes to new tags or count changes.
+		uint32 NumChanges = ChangeIdx.Num();
+		Writer.SerializeIntPacked(NumChanges);
+
+		uint32 TagCount = 0;
+		for (int32& Idx : ChangeIdx)
+		{
+			FNetGameplayTagCountContainerState::FGameplayTagCountStateItem& NewItem = NewState->TagStates[Idx];
+			// only pack count if we're sending to the correct client for the replication state.
+			bPackCount = NewItem.ReplicationState >= (bIsOwner ? EGameplayTagReplicationState::CountToOwner : EGameplayTagReplicationState::TagAndCountToAll);
+			NewItem.Tag.NetSerialize(Writer, DeltaParams.Map, bTagSuccess);
+			Writer.SerializeBits(&bPackCount, 1);
+			if (bPackCount)
+			{
+				TagCount = (uint32)NewItem.Count;
+				Writer.SerializeIntPacked(TagCount);
+			}
+		}
+	}
+
+	if (DeltaParams.Reader)
+	{
+		FBitReader& Reader = *DeltaParams.Reader;
+		bool bTagSuccess = true;
+
+		uint32 NumRemovals = 0;
+		Reader.SerializeIntPacked(NumRemovals);
+		if (NumRemovals > 0)
+		{
+			FNetGameplayTagCountContainerState::FGameplayTagCountStateItem RemovedItem{};
+			for (; NumRemovals > 0; NumRemovals--)
+			{
+				RemovedItem.Tag.NetSerialize(Reader, DeltaParams.Map, bTagSuccess);
+				UpdateTagMap_Internal(RemovedItem.Tag, -GameplayTagCountMap[RemovedItem.Tag], EGameplayTagReplicationState::None);
+			}
+		}
+		
+		uint32 NumChanges = 0;
+		Reader.SerializeIntPacked(NumChanges);
+		if (NumChanges > 0)
+		{
+			FNetGameplayTagCountContainerState::FGameplayTagCountStateItem UpdatedItem{};
+			FGameplayTag UpdatedTag;
+			uint32 NewCount = 1;
+			for (; NumChanges > 0; NumChanges--)
+			{
+				UpdatedItem.Tag.NetSerialize(Reader, DeltaParams.Map, bTagSuccess);
+				Reader.SerializeBits(&bPackCount, 1);
+				int32& Count = GameplayTagCountMap.FindOrAdd(UpdatedItem.Tag);
+				if (bPackCount)
+				{
+					Reader.SerializeIntPacked(NewCount);
+				}
+
+				if (NewCount != Count)
+				{
+					UpdateTagMap_Internal(UpdatedItem.Tag, NewCount - Count, EGameplayTagReplicationState::None);
+				}
+			}
+		}
+	}
+	return true;
+}
+
+void FGameplayTagCountContainer::SetOwner(UAbilitySystemComponent* ASC)
+{
+	Owner = ASC;
+}
+
 void FGameplayTagCountContainer::Notify_StackCountChange(const FGameplayTag& Tag)
 {	
 	// The purpose of this function is to let anyone listening on the EGameplayTagEventType::AnyCountChange event know that the 
@@ -533,8 +733,6 @@ FOnGameplayEffectTagCountChanged& FGameplayTagCountContainer::RegisterGameplayTa
 void FGameplayTagCountContainer::Reset(bool bResetCallbacks)
 {
 	GameplayTagCountMap.Reset();
-	ExplicitTagCountMap.Reset();
-	ExplicitTags.Reset();
 
 	if (bResetCallbacks)
 	{
@@ -543,41 +741,46 @@ void FGameplayTagCountContainer::Reset(bool bResetCallbacks)
 	}
 }
 
-bool FGameplayTagCountContainer::UpdateExplicitTags(const FGameplayTag& Tag, const int32 CountDelta, const bool bDeferParentTagsOnRemove)
+bool FGameplayTagCountContainer::UpdateExplicitTags(const FGameplayTag& Tag, const int32 CountDelta, const bool bDeferParentTagsOnRemove, EGameplayTagReplicationState TagRepState)
 {
-	const bool bTagAlreadyExplicitlyExists = ExplicitTags.HasTagExact(Tag);
+	int32 TagCountIndex = Items.Find(Tag);
 
-	// Need special case handling to maintain the explicit tag list correctly, adding the tag to the list if it didn't previously exist and a
-	// positive delta comes in, and removing it from the list if it did exist and a negative delta comes in.
-	if (!bTagAlreadyExplicitlyExists)
+	if (TagCountIndex != INDEX_NONE)
 	{
-		// Brand new tag with a positive delta needs to be explicitly added
-		if (CountDelta > 0)
-		{
-			ExplicitTags.AddTag(Tag);
-		}
-		// Block attempted reduction of non-explicit tags, as they were never truly added to the container directly
-		else
+		FGameplayTagCountItem& TagCountItem = Items[TagCountIndex];
+		const bool RepTypeMatch = TagCountItem.ReplicationState == TagRepState;
+		if (!RepTypeMatch && CountDelta > 0)
 		{
-			// only warn about tags that are in the container but will not be removed because they aren't explicitly in the container
-			if (ExplicitTags.HasTag(Tag))
+			// We're attempting to predict a replicated tag on a client but they already have potentially non-replicated tags
+			if (Owner && Owner->GetOwnerRole() != ROLE_Authority && TagCountItem.ReplicationState == EGameplayTagReplicationState::None)
 			{
-				ABILITY_LOG(Warning, TEXT("Attempted to remove tag: %s from tag count container, but it is not explicitly in the container!"), *Tag.ToString());
+				ABILITY_LOG(Warning, TEXT("Attempting to predict %s tag addition, but potentially non-replicated tag already exists"), *Tag.ToString());
 			}
-			return false;
+			// We're changing the rep state of an existing tag on the server
+			else if (Owner && Owner->GetOwnerRole() == ROLE_Authority && TagCountItem.ReplicationState > EGameplayTagReplicationState::None)
+			{
+				ABILITY_LOG(Warning, TEXT("Trying to change existing tag \"%s\" replication state from \"%s\" to \"%s\""), *Tag.ToString(), *UEnum::GetDisplayValueAsText(TagCountItem.ReplicationState).ToString(), *UEnum::GetDisplayValueAsText(TagRepState).ToString());
+			}
+
+			// Some sort of authority check before updating this value?
+			TagCountItem.ReplicationState = TagRepState > TagCountItem.ReplicationState ? TagRepState : TagCountItem.ReplicationState;
+		}
+		TagCountItem.Count += CountDelta;
+		if (TagCountItem.Count <= 0)
+		{
+			Items.RemoveAtSwap(TagCountIndex);
+			ExplicitTags.RemoveTag(Tag, bDeferParentTagsOnRemove);
 		}
 	}
-
-	// Update the explicit tag count map. This has to be separate than the map below because otherwise the count of nested tags ends up wrong
-	int32& ExistingCount = ExplicitTagCountMap.FindOrAdd(Tag);
-
-	ExistingCount = FMath::Max(ExistingCount + CountDelta, 0);
-
-	// If our new count is 0, remove us from the explicit tag list
-	if (ExistingCount <= 0)
+	else if (CountDelta > 0)
+	{
+		Items.Add(FGameplayTagCountItem(Tag, CountDelta, TagRepState));
+		ExplicitTags.AddTag(Tag);
+	}
+	else
 	{
-		// Remove from the explicit list
-		ExplicitTags.RemoveTag(Tag, bDeferParentTagsOnRemove);
+		ABILITY_LOG(Warning, TEXT("Attempted to remove tag: %s from tag count container, but it is not explicitly in the container!"), *Tag.ToString());
+		return false;
 	}
 
 	return true;
@@ -636,14 +839,14 @@ bool FGameplayTagCountContainer::GatherTagChangeDelegates(const FGameplayTag& Ta
 	return CreatedSignificantChange;
 }
 
-bool FGameplayTagCountContainer::UpdateTagMap_Internal(const FGameplayTag& Tag, int32 CountDelta)
+bool FGameplayTagCountContainer::UpdateTagMap_Internal(const FGameplayTag& Tag, int32 CountDelta, EGameplayTagReplicationState TagRepState)
 {
 	if (!ensureMsgf(Tag.IsValid(), TEXT("FGameplayTagCountContainer attempted to update tag map with an invalid Tag (None)!")))
 	{
 		return false;
 	}
 
-	if (!UpdateExplicitTags(Tag, CountDelta, false))
+	if (!UpdateExplicitTags(Tag, CountDelta, false, TagRepState))
 	{
 		return false;
 	}
@@ -658,9 +861,9 @@ bool FGameplayTagCountContainer::UpdateTagMap_Internal(const FGameplayTag& Tag,
 	return bSignificantChange;
 }
 
-bool FGameplayTagCountContainer::UpdateTagMapDeferredParentRemoval_Internal(const FGameplayTag& Tag, int32 CountDelta, TArray<FDeferredTagChangeDelegate>& DeferredTagChangeDelegates)
+bool FGameplayTagCountContainer::UpdateTagMapDeferredParentRemoval_Internal(const FGameplayTag& Tag, int32 CountDelta, TArray<FDeferredTagChangeDelegate>& DeferredTagChangeDelegates, EGameplayTagReplicationState TagRepState)
 {
-	if (!UpdateExplicitTags(Tag, CountDelta, true))
+	if (!UpdateExplicitTags(Tag, CountDelta, true, TagRepState))
 	{
 		return false;
 	}
diff --git a/Engine/Plugins/Runtime/GameplayAbilities/Source/GameplayAbilities/Private/Tests/GameplayTagCountContainerTests.cpp b/Engine/Plugins/Runtime/GameplayAbilities/Source/GameplayAbilities/Private/Tests/GameplayTagCountContainerTests.cpp
--- a/Engine/Plugins/Runtime/GameplayAbilities/Source/GameplayAbilities/Private/Tests/GameplayTagCountContainerTests.cpp
+++ b/Engine/Plugins/Runtime/GameplayAbilities/Source/GameplayAbilities/Private/Tests/GameplayTagCountContainerTests.cpp
@@ -43,7 +43,7 @@ bool FGameplayTagCountContainerTests::RunTest(const FString& Parameters)
 	TestTrue(TEXT("GetTagCount of Parent == 3"), TagCountContainer.GetTagCount(TestsDotGenericTag) == 3);
 
 	// Remove one
-	TagCountContainer.UpdateTagCount(TestsDotGenericTagDotOne, -1);
+	TagCountContainer.UpdateTagCount(TestsDotGenericTagDotOne, -1, EGameplayTagReplicationState::None);
 	TestFalse(TEXT("HasAllMatchingGameplayTags after Removing"), TagCountContainer.HasAllMatchingGameplayTags(ContainerOne));
 	TestFalse(TEXT("HasAnyMatchingGameplayTags after Removing"), TagCountContainer.HasAnyMatchingGameplayTags(ContainerOne));
 	TestTrue(TEXT("GetTagCount of One == 0"), TagCountContainer.GetTagCount(TestsDotGenericTagDotOne) == 0);
diff --git a/Engine/Plugins/Runtime/GameplayAbilities/Source/GameplayAbilities/Public/AbilitySystemBlueprintLibrary.h b/Engine/Plugins/Runtime/GameplayAbilities/Source/GameplayAbilities/Public/AbilitySystemBlueprintLibrary.h
--- a/Engine/Plugins/Runtime/GameplayAbilities/Source/GameplayAbilities/Public/AbilitySystemBlueprintLibrary.h
+++ b/Engine/Plugins/Runtime/GameplayAbilities/Source/GameplayAbilities/Public/AbilitySystemBlueprintLibrary.h
@@ -482,13 +482,21 @@ public:
 	static UE_API TArray<FGameplayEffectSpecHandle> GetAllLinkedGameplayEffectSpecHandles(FGameplayEffectSpecHandle SpecHandle);
 
 	/** Manually adds a set of tags to a given actor, and optionally replicates them. */
-	UFUNCTION(BlueprintCallable, Category="Ability|GameplayEffect")
+	UFUNCTION(BlueprintCallable, Category="Ability|GameplayEffect", meta=(DeprecatedFunction="Deprecated 5.7, Use AddGameplayTags instead"))
 	static UE_API bool AddLooseGameplayTags(AActor* Actor, const FGameplayTagContainer& GameplayTags, bool bShouldReplicate=false);
 
 	/** Manually removes a set of tags from a given actor, with optional replication. */
-	UFUNCTION(BlueprintCallable, Category="Ability|GameplayEffect")
+	UFUNCTION(BlueprintCallable, Category="Ability|GameplayEffect", meta=(DeprecatedFunction="Deprecated 5.7, Use RemoveGameplayTags instead"))
 	static UE_API bool RemoveLooseGameplayTags(AActor* Actor, const FGameplayTagContainer& GameplayTags, bool bShouldReplicate=false);
 
+	/** Manually adds a set of tags to a given actor */
+	UFUNCTION(BlueprintCallable, Category="Ability|GameplayEffect")
+	static UE_API bool AddGameplayTags(AActor* Actor, const FGameplayTagContainer& GameplayTags, EGameplayTagReplicationState ReplicationRule = EGameplayTagReplicationState::None);
+
+	/** Manually removes a set of tags from a given actor */
+	UFUNCTION(BlueprintCallable, Category="Ability|GameplayEffect")
+	static UE_API bool RemoveGameplayTags(AActor* Actor, const FGameplayTagContainer& GameplayTags);
+
 	/** Get the GameplayEffect definition from a GameplayEffectSpecHandle. */
 	UFUNCTION(BlueprintCallable, Category = "Ability|GameplayEffect")
 	static UE_API const UGameplayEffect* GetGameplayEffectFromSpecHandle(FGameplayEffectSpecHandle Handle);
diff --git a/Engine/Plugins/Runtime/GameplayAbilities/Source/GameplayAbilities/Public/AbilitySystemComponent.h b/Engine/Plugins/Runtime/GameplayAbilities/Source/GameplayAbilities/Public/AbilitySystemComponent.h
--- a/Engine/Plugins/Runtime/GameplayAbilities/Source/GameplayAbilities/Public/AbilitySystemComponent.h
+++ b/Engine/Plugins/Runtime/GameplayAbilities/Source/GameplayAbilities/Public/AbilitySystemComponent.h
@@ -612,29 +612,29 @@ class UAbilitySystemComponent : public UGameplayTasksComponent, public IGameplay
 	}
 
 	/** Forcibly sets the number of instances of a given tag */
-	inline void SetTagMapCount(const FGameplayTag& Tag, int32 NewCount)
+	inline void SetTagMapCount(const FGameplayTag& Tag, int32 NewCount, EGameplayTagReplicationState TagRepState = EGameplayTagReplicationState::None)
 	{
-		if (GameplayTagCountContainer.SetTagCount(Tag, NewCount))
+		if (GameplayTagCountContainer.SetTagCount(Tag, NewCount, TagRepState))
 		{
 			OnTagUpdated(Tag, NewCount > 0);
 		}
 	}
 
 	/** Update the number of instances of a given tag and calls callback */
-	inline void UpdateTagMap(const FGameplayTag& BaseTag, int32 CountDelta)
+	inline void UpdateTagMap(const FGameplayTag& BaseTag, int32 CountDelta, EGameplayTagReplicationState TagRepState)
 	{
-		if (GameplayTagCountContainer.UpdateTagCount(BaseTag, CountDelta))
+		if (GameplayTagCountContainer.UpdateTagCount(BaseTag, CountDelta, TagRepState))
 		{
 			OnTagUpdated(BaseTag, CountDelta > 0);
 		}
 	}
 
 	/** Update the number of instances of a given tag and calls callback */
-	inline void UpdateTagMap(const FGameplayTagContainer& Container, int32 CountDelta)
+	inline void UpdateTagMap(const FGameplayTagContainer& Container, int32 CountDelta, EGameplayTagReplicationState TagRepState)
 	{
 		if (!Container.IsEmpty())
 		{
-			UpdateTagMap_Internal(Container, CountDelta);
+			UpdateTagMap_Internal(Container, CountDelta, TagRepState);
 		}
 	}
 
@@ -644,39 +644,38 @@ class UAbilitySystemComponent : public UGameplayTasksComponent, public IGameplay
 		TagContainer.AppendTags(GetBlockedAbilityTags());
 	}
 
-	[[nodiscard]] inline const FGameplayTagContainer& GetBlockedAbilityTags() const
+	[[nodiscard]] inline const FGameplayTagContainer GetBlockedAbilityTags() const
 	{
 		return BlockedAbilityTags.GetExplicitGameplayTags();
 	}
 
 	/** 	 
-	 *  Allows GameCode to add loose gameplaytags which are not backed by a GameplayEffect. 
-	 *	Tags added this way are not replicated! Use the 'Replicated' versions of these functions if replication is needed.
+	 *  Allows GameCode to add loose gameplaytags which are not backed by a GameplayEffect.
 	 *	It is up to the calling GameCode to make sure these tags are added on clients/server where necessary
 	 */
-	inline void AddLooseGameplayTag(const FGameplayTag& GameplayTag, int32 Count=1)
+	inline void AddLooseGameplayTag(const FGameplayTag& GameplayTag, int32 Count=1, EGameplayTagReplicationState TagRepState = EGameplayTagReplicationState::None)
 	{
-		UpdateTagMap(GameplayTag, Count);
+		UpdateTagMap(GameplayTag, Count, TagRepState);
 	}
 
-	inline void AddLooseGameplayTags(const FGameplayTagContainer& GameplayTags, int32 Count = 1)
+	inline void AddLooseGameplayTags(const FGameplayTagContainer& GameplayTags, int32 Count = 1, EGameplayTagReplicationState TagRepState = EGameplayTagReplicationState::None)
 	{
-		UpdateTagMap(GameplayTags, Count);
+		UpdateTagMap(GameplayTags, Count, TagRepState);
 	}
 
 	inline void RemoveLooseGameplayTag(const FGameplayTag& GameplayTag, int32 Count = 1)
 	{
-		UpdateTagMap(GameplayTag, -Count);
+		UpdateTagMap(GameplayTag, -Count, EGameplayTagReplicationState::None);
 	}
 
 	inline void RemoveLooseGameplayTags(const FGameplayTagContainer& GameplayTags, int32 Count = 1)
 	{
-		UpdateTagMap(GameplayTags, -Count);
+		UpdateTagMap(GameplayTags, -Count, EGameplayTagReplicationState::None);
 	}
 
-	inline void SetLooseGameplayTagCount(const FGameplayTag& GameplayTag, int32 NewCount)
+	inline void SetLooseGameplayTagCount(const FGameplayTag& GameplayTag, int32 NewCount, EGameplayTagReplicationState TagRepState = EGameplayTagReplicationState::None)
 	{
-		SetTagMapCount(GameplayTag, NewCount);
+		SetTagMapCount(GameplayTag, NewCount, TagRepState);
 	}
 
 	/**
@@ -689,57 +688,31 @@ class UAbilitySystemComponent : public UGameplayTasksComponent, public IGameplay
 	UFUNCTION(BlueprintPure, Category = "Gameplay Tags")
 	UE_API int32 GetGameplayTagCount(FGameplayTag GameplayTag) const;
 
-	/**
-	 *  Allows GameCode to add loose gameplaytags which are not backed by a GameplayEffect. Tags added using 
-	 *  these functions will be replicated. Note that replicated loose tags will override any locally-set tag counts
-	 *  on simulated proxies.
-	 */
-	inline void AddReplicatedLooseGameplayTag(const FGameplayTag& GameplayTag)
-	{
-		GetReplicatedLooseTags_Mutable().AddTag(GameplayTag);
-	}
-
-	inline void AddReplicatedLooseGameplayTags(const FGameplayTagContainer& GameplayTags)
-	{
-		GetReplicatedLooseTags_Mutable().AddTags(GameplayTags);
-	}
-
-	inline void RemoveReplicatedLooseGameplayTag(const FGameplayTag& GameplayTag)
-	{
-		GetReplicatedLooseTags_Mutable().RemoveTag(GameplayTag);
-	}
-
-	inline void RemoveReplicatedLooseGameplayTags(const FGameplayTagContainer& GameplayTags)
-	{
-		GetReplicatedLooseTags_Mutable().RemoveTags(GameplayTags);
-	}
-
-	inline void SetReplicatedLooseGameplayTagCount(const FGameplayTag& GameplayTag, int32 NewCount)
-	{
-		GetReplicatedLooseTags_Mutable().SetTagCount(GameplayTag, NewCount);
-	}
-
 	/** 	 
 	 * Minimally replicated tags are replicated tags that come from GEs when in bMinimalReplication mode. 
 	 * (The GEs do not replicate, but the tags they grant do replicate via these functions)
 	 */
 	inline void AddMinimalReplicationGameplayTag(const FGameplayTag& GameplayTag)
 	{
+		UpdateTagMap(GameplayTag, 1, EGameplayTagReplicationState::TagOnly);
 		GetMinimalReplicationTags_Mutable().AddTag(GameplayTag);
 	}
 
 	inline void AddMinimalReplicationGameplayTags(const FGameplayTagContainer& GameplayTags)
 	{
+		UpdateTagMap(GameplayTags, 1, EGameplayTagReplicationState::TagOnly);
 		GetMinimalReplicationTags_Mutable().AddTags(GameplayTags);
 	}
 
 	inline void RemoveMinimalReplicationGameplayTag(const FGameplayTag& GameplayTag)
 	{
+		UpdateTagMap(GameplayTag, -1, EGameplayTagReplicationState::TagOnly);
 		GetMinimalReplicationTags_Mutable().RemoveTag(GameplayTag);
 	}
 
 	inline void RemoveMinimalReplicationGameplayTags(const FGameplayTagContainer& GameplayTags)
 	{
+		UpdateTagMap(GameplayTags, -1, EGameplayTagReplicationState::TagOnly);
 		GetMinimalReplicationTags_Mutable().RemoveTags(GameplayTags);
 	}
 
@@ -1895,11 +1868,6 @@ protected:
 	/** Abilities with these tags are not able to be activated */
 	FGameplayTagCountContainer BlockedAbilityTags;
 
-	UE_DEPRECATED(4.26, "This will be made private in future engine versions. Use SetBlockedAbilityBindings, GetBlockedAbilityBindings, or GetBlockedAbilityBindings_Mutable instead.")
-	/** Tracks abilities that are blocked based on input binding. An ability is blocked if BlockedAbilityBindings[InputID] > 0 */
-	UPROPERTY(Transient, Replicated)
-	TArray<uint8> BlockedAbilityBindings;
-
 	UE_API void SetBlockedAbilityBindings(const TArray<uint8>& NewBlockedAbilityBindings);
 	UE_API TArray<uint8>& GetBlockedAbilityBindings_Mutable();
 	UE_API const TArray<uint8>& GetBlockedAbilityBindings() const;
@@ -1907,11 +1875,8 @@ protected:
 	UE_API void DebugCyclicAggregatorBroadcasts(struct FAggregator* Aggregator);
 	
 	/** Acceleration map for all gameplay tags (OwnedGameplayTags from GEs and explicit GameplayCueTags) */
-	FGameplayTagCountContainer GameplayTagCountContainer;
-
-	UE_DEPRECATED(4.26, "This will be made private in future engine versions. Use SetMinimalReplicationTags, GetMinimalReplicationTags, or GetMinimalReplicationTags_Mutable instead.")
 	UPROPERTY(Replicated)
-	FMinimalReplicationTagCountMap MinimalReplicationTags;
+	FGameplayTagCountContainer GameplayTagCountContainer;
 
 	UE_API void SetMinimalReplicationTags(const FMinimalReplicationTagCountMap& NewMinimalReplicationTags);
 	UE_API FMinimalReplicationTagCountMap& GetMinimalReplicationTags_Mutable();
@@ -1930,7 +1895,7 @@ protected:
 	UE_API const UAttributeSet*	GetAttributeSubobjectChecked(const TSubclassOf<UAttributeSet> AttributeClass) const;
 	UE_API const UAttributeSet*	GetOrCreateAttributeSubobject(TSubclassOf<UAttributeSet> AttributeClass);
 
-	UE_API void UpdateTagMap_Internal(const FGameplayTagContainer& Container, int32 CountDelta);
+	UE_API void UpdateTagMap_Internal(const FGameplayTagContainer& Container, int32 CountDelta, EGameplayTagReplicationState TagRepState);
 
 	friend struct FActiveGameplayEffect;
 	friend struct FActiveGameplayEffectAction;
@@ -1955,6 +1920,14 @@ private:
 	TArray<TObjectPtr<UGameplayAbility>>& GetReplicatedInstancedAbilities_Mutable() { return AllReplicatedInstancedAbilities; }
 	PRAGMA_ENABLE_DEPRECATION_WARNINGS
 
+	/** Tracks abilities that are blocked based on input binding. An ability is blocked if BlockedAbilityBindings[InputID] > 0 */
+	UPROPERTY(Transient, Replicated)
+	TArray<uint8> BlockedAbilityBindings;
+
+	UE_DEPRECATED(5.7, "This will be removed in favor of using TagOnly replication state for on the replicated Tag Count Container")
+	UPROPERTY(Replicated)
+	FMinimalReplicationTagCountMap MinimalReplicationTags;
+
 private:
 
 	/** List of attribute sets */
@@ -1968,6 +1941,7 @@ private:
 	FTimerHandle    OnRep_ActivateAbilitiesTimerHandle;
 
 	/** Container used for replicating loose gameplay tags */
+	UE_DEPRECATED(5.7, "This will be removed in favor of using CountToOwner replication state for on the replicated Tag Count Container")
 	UPROPERTY(Replicated)
 	FMinimalReplicationTagCountMap ReplicatedLooseTags;
 
diff --git a/Engine/Plugins/Runtime/GameplayAbilities/Source/GameplayAbilities/Public/GameplayEffectTypes.h b/Engine/Plugins/Runtime/GameplayAbilities/Source/GameplayAbilities/Public/GameplayEffectTypes.h
--- a/Engine/Plugins/Runtime/GameplayAbilities/Source/GameplayAbilities/Public/GameplayEffectTypes.h
+++ b/Engine/Plugins/Runtime/GameplayAbilities/Source/GameplayAbilities/Public/GameplayEffectTypes.h
@@ -13,6 +13,7 @@
 #include "GameFramework/Actor.h"
 #include "GameplayTagContainer.h"
 #include "AttributeSet.h"
+#include "Net/Serialization/FastArraySerializer.h"
 #include "PhysicalMaterials/PhysicalMaterial.h"
 #include "GameplayEffectTypes.generated.h"
 
@@ -1037,16 +1038,65 @@ namespace EGameplayTagEventType
 	};
 }
 
+/**
+ * Enum to describe replication rules for each tag
+ * Any state besides "None" must be added on the server and can be predicted on clients.
+ */
+UENUM()
+enum class EGameplayTagReplicationState : uint8
+{
+	None,				// Tag will not replicate
+	TagOnly,			// Tag will replicate to all without count
+	CountToOwner,		// Tag will replicate to all and Count will replicate to owner
+	TagAndCountToAll	// Tag and Count will replicate to all
+};
+
+USTRUCT()
+struct FGameplayTagCountItem
+{
+	GENERATED_USTRUCT_BODY()
+	
+	FGameplayTagCountItem()
+	: Tag(FGameplayTag())
+	, Count(0)
+	, ReplicationState(EGameplayTagReplicationState::None)
+	{}
+
+	FGameplayTagCountItem(const FGameplayTag& InTag, const int32 InCount = 1, EGameplayTagReplicationState TagRepState = EGameplayTagReplicationState::None)
+	: Tag(InTag)
+	, Count(InCount)
+	, ReplicationState(TagRepState)
+	{}
+
+	UPROPERTY()
+	FGameplayTag Tag;
+
+	UPROPERTY()
+	int32 Count;
+
+	EGameplayTagReplicationState ReplicationState;
+
+	bool operator == (const FGameplayTagCountItem& Other) const
+	{
+		return Tag.MatchesTagExact(Other.Tag);
+	}
+
+	bool operator == (const FGameplayTag& OtherTag) const
+	{
+		return Tag.MatchesTagExact(OtherTag);
+	}
+};
+
 /**
  * Struct that tracks the number/count of tag applications within it. Explicitly tracks the tags added or removed,
  * while simultaneously tracking the count of parent tags as well. Events/delegates are fired whenever the tag counts
  * of any tag (explicit or parent) are modified.
  */
+USTRUCT()
 struct FGameplayTagCountContainer
-{	
-	FGameplayTagCountContainer()
-	{}
-
+{
+	GENERATED_BODY()
+	
 	/**
 	 * Check if the count container has a gameplay tag that matches against the specified tag (expands to include parents of asset tags)
 	 * 
@@ -1126,7 +1176,7 @@ struct FGameplayTagCountContainer
 			TArray<FDeferredTagChangeDelegate> DeferredTagChangeDelegates;
 			for (auto TagIt = Container.CreateConstIterator(); TagIt; ++TagIt)
 			{
-				bUpdatedAny |= UpdateTagMapDeferredParentRemoval_Internal(*TagIt, CountDelta, DeferredTagChangeDelegates);
+				bUpdatedAny |= UpdateTagMapDeferredParentRemoval_Internal(*TagIt, CountDelta, DeferredTagChangeDelegates, EGameplayTagReplicationState::None);
 			}
 
 			if (bUpdatedAny && CountDelta < 0)
@@ -1149,11 +1199,11 @@ struct FGameplayTagCountContainer
 	 * 
 	 * @return True if tag was *either* added or removed. (E.g., we had the tag and now dont. or didnt have the tag and now we do. We didn't just change the count (1 count -> 2 count would return false).
 	 */
-	inline bool UpdateTagCount(const FGameplayTag& Tag, int32 CountDelta)
+	inline bool UpdateTagCount(const FGameplayTag& Tag, int32 CountDelta, EGameplayTagReplicationState TagRepState)
 	{
 		if (Tag.IsValid() && CountDelta != 0)
 		{
-			return UpdateTagMap_Internal(Tag, CountDelta);
+			return UpdateTagMap_Internal(Tag, CountDelta, TagRepState);
 		}
 
 		return false;
@@ -1173,7 +1223,7 @@ struct FGameplayTagCountContainer
 	{
 		if (CountDelta != 0)
 		{
-			return UpdateTagMapDeferredParentRemoval_Internal(Tag, CountDelta, DeferredTagChangeDelegates);
+			return UpdateTagMapDeferredParentRemoval_Internal(Tag, CountDelta, DeferredTagChangeDelegates, EGameplayTagReplicationState::None);
 		}
 
 		return false;
@@ -1184,10 +1234,11 @@ struct FGameplayTagCountContainer
 	 * 
 	 * @param Tag			Tag to update
 	 * @param Count			New count of the tag
+	 * @param TagRepState	The replication state for the tag, if it is being added it will be added with this state
 	 * 
 	 * @return True if tag was *either* added or removed. (E.g., we had the tag and now dont. or didnt have the tag and now we do. We didn't just change the count (1 count -> 2 count would return false).
 	 */
-	inline bool SetTagCount(const FGameplayTag& Tag, int32 NewCount)
+	inline bool SetTagCount(const FGameplayTag& Tag, int32 NewCount, EGameplayTagReplicationState TagRepState = EGameplayTagReplicationState::None)
 	{
 		if (!Tag.IsValid())
 		{
@@ -1195,15 +1246,16 @@ struct FGameplayTagCountContainer
 		}
 
 		int32 ExistingCount = 0;
-		if (int32* Ptr  = ExplicitTagCountMap.Find(Tag))
+		const int32 Idx  = Items.Find(Tag);
+		if (Idx != INDEX_NONE)
 		{
-			ExistingCount = *Ptr;
+			ExistingCount = Items[Idx].Count;
 		}
 
 		int32 CountDelta = NewCount - ExistingCount;
 		if (CountDelta != 0)
 		{
-			return UpdateTagMap_Internal(Tag, CountDelta);
+			return UpdateTagMap_Internal(Tag, CountDelta, TagRepState);
 		}
 
 		return false;
@@ -1237,14 +1289,21 @@ struct FGameplayTagCountContainer
 	*/
 	inline int32 GetExplicitTagCount(const FGameplayTag& Tag) const
 	{
-		if (const int32* Ptr = ExplicitTagCountMap.Find(Tag))
+		const int32 Idx  = Items.Find(Tag);
+		if (Idx != INDEX_NONE)
 		{
-			return *Ptr;
+			return Items[Idx].Count;
 		}
 
 		return 0;
 	}
 
+	/** Fills in ParentTags from GameplayTags */
+	void FillParentTags()
+	{
+		ExplicitTags.FillParentTags();
+	}
+
 	/**
 	 *	Broadcasts the AnyChange event for this tag. This is called when the stack count of the backing gameplay effect change.
 	 *	It is up to the receiver of the broadcasted delegate to decide what to do with this.
@@ -1283,11 +1342,9 @@ struct FGameplayTagCountContainer
 	 */
 	UE_API void Reset(bool bResetCallbacks = true);
 
-	/** Fills in ParentTags from GameplayTags */
-	void FillParentTags()
-	{
-		ExplicitTags.FillParentTags();
-	}
+	bool NetDeltaSerialize(FNetDeltaSerializeInfo& DeltaParams);
+
+	UE_API void SetOwner(UAbilitySystemComponent* ASC);
 
 private:
 
@@ -1297,34 +1354,46 @@ private:
 		FOnGameplayEffectTagCountChanged	OnAnyChange;
 	};
 
+	/** Array of explicit tags and their count */
+	TArray<FGameplayTagCountItem> Items;
+
 	/** Map of tag to delegate that will be fired when the count for the key tag changes to or away from zero */
 	TMap<FGameplayTag, FDelegateInfo> GameplayTagEventMap;
 
 	/** Map of tag to active count of that tag */
 	TMap<FGameplayTag, int32> GameplayTagCountMap;
 
-	/** Map of tag to explicit count of that tag. Cannot share with above map because it's not safe to merge explicit and generic counts */	
-	TMap<FGameplayTag, int32> ExplicitTagCountMap;
-
 	/** Delegate fired whenever any tag's count changes to or away from zero */
 	FOnGameplayEffectTagCountChanged OnAnyTagChangeDelegate;
 
 	/** Container of tags that were explicitly added */
 	FGameplayTagContainer ExplicitTags;
 
+	/** We want to track the owner for the container so that we can give valuable logging about tag replication mismatches. */
+	UPROPERTY()
+	TObjectPtr<class UAbilitySystemComponent> Owner;
+
 	/** Internal helper function to adjust the explicit tag list & corresponding maps/delegates/etc. as necessary */
-	UE_API bool UpdateTagMap_Internal(const FGameplayTag& Tag, int32 CountDelta);
+	UE_API bool UpdateTagMap_Internal(const FGameplayTag& Tag, int32 CountDelta, EGameplayTagReplicationState TagRepState = EGameplayTagReplicationState::None);
 
 	/** Internal helper function to adjust the explicit tag list & corresponding maps/delegates/etc. as necessary. This does not call FillParentTags or any of the tag change delegates. These delegates are returned and must be executed by the caller. */
-	UE_API bool UpdateTagMapDeferredParentRemoval_Internal(const FGameplayTag& Tag, int32 CountDelta, TArray<FDeferredTagChangeDelegate>& DeferredTagChangeDelegates);
+	UE_API bool UpdateTagMapDeferredParentRemoval_Internal(const FGameplayTag& Tag, int32 CountDelta, TArray<FDeferredTagChangeDelegate>& DeferredTagChangeDelegates, EGameplayTagReplicationState TagRepState = EGameplayTagReplicationState::None);
 
 	/** Internal helper function to adjust the explicit tag list & corresponding map. */
-	UE_API bool UpdateExplicitTags(const FGameplayTag& Tag, int32 CountDelta, bool bDeferParentTagsOnRemove);
+	UE_API bool UpdateExplicitTags(const FGameplayTag& Tag, int32 CountDelta, bool bDeferParentTagsOnRemove, EGameplayTagReplicationState TagRepState = EGameplayTagReplicationState::None);
 
 	/** Internal helper function to collect the delegates that need to be called when Tag has its count changed by CountDelta. */
 	UE_API bool GatherTagChangeDelegates(const FGameplayTag& Tag, int32 CountDelta, TArray<FDeferredTagChangeDelegate>& TagChangeDelegates);
 };
 
+template<>
+struct TStructOpsTypeTraits< FGameplayTagCountContainer > : public TStructOpsTypeTraitsBase2< FGameplayTagCountContainer >
+{
+	enum 
+	{
+		WithNetDeltaSerializer = true,
+   };
+};
 
 /**
  * Struct used to update a blueprint property with a gameplay tag count.
