UE: Remove deprecation of AddLooseGameplayTags/RemoveLooseGameplayTags for now

diff --git a/Engine/Plugins/Runtime/GameplayAbilities/Source/GameplayAbilities/Public/AbilitySystemBlueprintLibrary.h b/Engine/Plugins/Runtime/GameplayAbilities/Source/GameplayAbilities/Public/AbilitySystemBlueprintLibrary.h
--- a/Engine/Plugins/Runtime/GameplayAbilities/Source/GameplayAbilities/Public/AbilitySystemBlueprintLibrary.h
+++ b/Engine/Plugins/Runtime/GameplayAbilities/Source/GameplayAbilities/Public/AbilitySystemBlueprintLibrary.h
@@ -482,11 +482,11 @@ public:
 	static UE_API TArray<FGameplayEffectSpecHandle> GetAllLinkedGameplayEffectSpecHandles(FGameplayEffectSpecHandle SpecHandle);
 
 	/** Manually adds a set of tags to a given actor, and optionally replicates them. */
-	UFUNCTION(BlueprintCallable, Category="Ability|GameplayEffect", meta=(DeprecatedFunction="Deprecated 5.7, Use AddGameplayTags instead"))
+	UFUNCTION(BlueprintCallable, Category="Ability|GameplayEffect")
 	static UE_API bool AddLooseGameplayTags(AActor* Actor, const FGameplayTagContainer& GameplayTags, bool bShouldReplicate=false);
 
 	/** Manually removes a set of tags from a given actor, with optional replication. */
-	UFUNCTION(BlueprintCallable, Category="Ability|GameplayEffect", meta=(DeprecatedFunction="Deprecated 5.7, Use RemoveGameplayTags instead"))
+	UFUNCTION(BlueprintCallable, Category="Ability|GameplayEffect")
 	static UE_API bool RemoveLooseGameplayTags(AActor* Actor, const FGameplayTagContainer& GameplayTags, bool bShouldReplicate=false);
 
 	/** Manually adds a set of tags to a given actor */
