Cancel Ability Tags Gameplay Effect Component

- I ran across a use case that led me to believe adding a new Gameplay Effect Component to cancel active abilities when a GE is applied would be worthwhile.
- This Gameplay Effect Component can either cancel Abilities with matching tags or cancel all Abilities without matching tags.
- I've added a 'mode' setting for either cancelling found abilities on GE execution or application. Execution would be used for periodic GE's which may potentially want to repeatedly cancel abilities on periodic execution. Application would be for all cases one time ability cancelation when the GE is applied.

[RN] Added new Gameplay Effect Component to cancel all Gameplay Abilities with or without specified matching tags

diff --git a/Engine/Plugins/Runtime/GameplayAbilities/Source/GameplayAbilities/Private/GameplayEffectComponents/CancelAbilityTagsGameplayEffectComponent.cpp b/Engine/Plugins/Runtime/GameplayAbilities/Source/GameplayAbilities/Private/GameplayEffectComponents/CancelAbilityTagsGameplayEffectComponent.cpp
new file mode 100644
--- /dev/null
+++ b/Engine/Plugins/Runtime/GameplayAbilities/Source/GameplayAbilities/Private/GameplayEffectComponents/CancelAbilityTagsGameplayEffectComponent.cpp
@@ -0,0 +1,117 @@
+﻿// Copyright Epic Games, Inc. All Rights Reserved.
+
+#include "GameplayEffectComponents/CancelAbilityTagsGameplayEffectComponent.h"
+
+#include "AbilitySystemComponent.h"
+#include "Misc/DataValidation.h"
+
+#include UE_INLINE_GENERATED_CPP_BY_NAME(CancelAbilityTagsGameplayEffectComponent)
+
+#define LOCTEXT_NAMESPACE "CancelAbilityTagsGameplayEffectComponent"
+
+void UCancelAbilityTagsGameplayEffectComponent::PostInitProperties()
+{
+	Super::PostInitProperties();
+
+	// Try to find the parent and update the inherited tags.  We do this 'early' because this is the only function
+	// we can use for creation of the object (and this runs post-constructor).
+	const UCancelAbilityTagsGameplayEffectComponent* Parent = FindParentComponent(*this);
+	InheritableCancelAbilitiesWithTagsContainer.UpdateInheritedTagProperties(Parent ? &Parent->InheritableCancelAbilitiesWithTagsContainer : nullptr);
+
+#if WITH_EDITORONLY_DATA
+	EditorFriendlyName = TEXT("Cancel Abilities With Tags");
+#endif
+}
+
+void UCancelAbilityTagsGameplayEffectComponent::OnGameplayEffectApplied(FActiveGameplayEffectsContainer& ActiveGEContainer, FGameplayEffectSpec& GESpec, FPredictionKey& PredictionKey) const
+{
+	Super::OnGameplayEffectApplied(ActiveGEContainer, GESpec, PredictionKey);
+
+	if (ComponentMode != ECancelAbilityTagsGameplayEffectComponentMode::OnApplication)
+	{
+		return;
+	}
+
+	CancelOwnerAbilities(ActiveGEContainer);
+}
+
+void UCancelAbilityTagsGameplayEffectComponent::OnGameplayEffectExecuted(FActiveGameplayEffectsContainer& ActiveGEContainer, FGameplayEffectSpec& GESpec, FPredictionKey& PredictionKey) const
+{
+	Super::OnGameplayEffectExecuted(ActiveGEContainer, GESpec, PredictionKey);
+
+	if (ComponentMode != ECancelAbilityTagsGameplayEffectComponentMode::OnExecution)
+	{
+		return;
+	}
+	
+	CancelOwnerAbilities(ActiveGEContainer);
+}
+
+void UCancelAbilityTagsGameplayEffectComponent::OnGameplayEffectChanged()
+{
+	Super::OnGameplayEffectChanged();
+	SetAndApplyCanceledAbilityTagChanges(InheritableCancelAbilitiesWithTagsContainer, InheritableCancelAbilitiesWithoutTagsContainer);
+}
+
+void UCancelAbilityTagsGameplayEffectComponent::SetAndApplyCanceledAbilityTagChanges(const FInheritedTagContainer& CanceledAbilityWithTagsMods, const FInheritedTagContainer& CanceledAbilityWithoutTagsMods)
+{
+	InheritableCancelAbilitiesWithTagsContainer = CanceledAbilityWithTagsMods;
+	InheritableCancelAbilitiesWithoutTagsContainer = CanceledAbilityWithoutTagsMods;
+	
+	// Try to find the parent and update the inherited tags
+	const UCancelAbilityTagsGameplayEffectComponent* Parent = FindParentComponent(*this);
+	InheritableCancelAbilitiesWithTagsContainer.UpdateInheritedTagProperties(Parent ? &Parent->InheritableCancelAbilitiesWithTagsContainer : nullptr);
+	InheritableCancelAbilitiesWithoutTagsContainer.UpdateInheritedTagProperties(Parent ? &Parent->InheritableCancelAbilitiesWithoutTagsContainer : nullptr);
+}
+
+void UCancelAbilityTagsGameplayEffectComponent::CancelOwnerAbilities(const FActiveGameplayEffectsContainer& ActiveGEContainer) const
+{
+	if (!ActiveGEContainer.OwnerIsNetAuthority)
+	{
+		return;
+	}
+	
+	UAbilitySystemComponent* ASC = ActiveGEContainer.Owner;
+	if (!ensure(ASC))
+	{
+		return;
+	}
+
+	ASC->CancelAbilities(&InheritableCancelAbilitiesWithTagsContainer.CombinedTags, &InheritableCancelAbilitiesWithoutTagsContainer.CombinedTags);
+}
+
+#if WITH_EDITOR
+void UCancelAbilityTagsGameplayEffectComponent::PostEditChangeProperty(FPropertyChangedEvent& PropertyChangedEvent)
+{
+	Super::PostEditChangeProperty(PropertyChangedEvent);
+
+	if (PropertyChangedEvent.GetMemberPropertyName() == GetInheritableCancelAbilitiesWithTagsContainerPropertyName() ||
+		PropertyChangedEvent.GetMemberPropertyName() == GetInheritableCancelAbilitiesWithoutTagsContainerPropertyName())
+	{
+		// Tell the GE it needs to reconfigure itself based on these updated properties (this will reaggregate the tags)
+		UGameplayEffect* Owner = GetOwner();
+		Owner->OnGameplayEffectChanged();
+	}
+}
+
+EDataValidationResult UCancelAbilityTagsGameplayEffectComponent::IsDataValid(FDataValidationContext& Context) const
+{
+	EDataValidationResult Result = Super::IsDataValid(Context);
+
+	const bool bInstantEffect = (GetOwner()->DurationPolicy == EGameplayEffectDurationType::Instant);
+	const bool bHasPeriod = GetOwner()->Period.Value > 0.0f;
+	if (!bInstantEffect && !bHasPeriod && ComponentMode == ECancelAbilityTagsGameplayEffectComponentMode::OnExecution && !InheritableCancelAbilitiesWithTagsContainer.CombinedTags.IsEmpty())
+	{
+		Context.AddError(FText::FormatOrdered(
+			LOCTEXT("GEInstantAndBlockAbilityTags", "GE {0} has a duration and no period while {1} has mode: OnExecution. {1} will not be able to function as expected, change mode to 'OnApplication'."),
+			FText::FromString(GetNameSafe(GetOwner())),
+			FText::FromString(EditorFriendlyName)));
+		
+		Result = EDataValidationResult::Invalid;
+	}
+
+	return Result;
+}
+#endif // WITH_EDITOR
+
+#undef LOCTEXT_NAMESPACE
\ No newline at end of file
diff --git a/Engine/Plugins/Runtime/GameplayAbilities/Source/GameplayAbilities/Public/GameplayEffectComponents/CancelAbilityTagsGameplayEffectComponent.h b/Engine/Plugins/Runtime/GameplayAbilities/Source/GameplayAbilities/Public/GameplayEffectComponents/CancelAbilityTagsGameplayEffectComponent.h
new file mode 100644
--- /dev/null
+++ b/Engine/Plugins/Runtime/GameplayAbilities/Source/GameplayAbilities/Public/GameplayEffectComponents/CancelAbilityTagsGameplayEffectComponent.h
@@ -0,0 +1,82 @@
+﻿// Copyright Epic Games, Inc. All Rights Reserved.
+
+#pragma once
+
+#include "GameplayEffect.h"
+#include "GameplayEffectComponent.h"
+#include "CancelAbilityTagsGameplayEffectComponent.generated.h"
+
+#define UE_API GAMEPLAYABILITIES_API
+
+
+/** Enumeration representing when the ability cancellation logic for the 'CancelAbilityTagsGameplayEffectComponent' happens */
+UENUM()
+enum class ECancelAbilityTagsGameplayEffectComponentMode : uint8
+{
+	/** Component applies cancellation logic when the owning gameplay effect is Applied. Use this if we want to cancel abilities a single time on GE application, regardless of GE duration type. */
+	OnApplication,
+
+	/** Component applies cancellation logic when the owning gameplay effect is Executed. Use this if we want to cancel abilities on each periodic GE execution. */
+	OnExecution
+};
+
+/** Handles canceling active Gameplay Abilities based on Gameplay Tags for the Target actor of the owner Gameplay Effect */
+UCLASS(DisplayName="Cancel Abilities with Tags", MinimalAPI)
+class UCancelAbilityTagsGameplayEffectComponent : public UGameplayEffectComponent
+{
+	GENERATED_BODY()
+
+public:
+	/** Set up an EditorFriendlyName and do some initialization */
+	UE_API virtual void PostInitProperties() override;
+	
+	/** Needed to properly apply FInheritedTagContainer properties */
+	UE_API virtual void OnGameplayEffectChanged() override;
+	
+	/** Needed to cancel abilities if we're canceling abilities on execution */
+	UE_API virtual void OnGameplayEffectExecuted(FActiveGameplayEffectsContainer& ActiveGEContainer, FGameplayEffectSpec& GESpec, FPredictionKey& PredictionKey) const override;
+
+	/** Needed to cancel abilities if we're canceling abilities on effect application */
+	UE_API virtual void OnGameplayEffectApplied(FActiveGameplayEffectsContainer& ActiveGEContainer, FGameplayEffectSpec& GESpec, FPredictionKey& PredictionKey) const override;
+	
+	/** Applies the Canceled Ability Tags to the GameplayEffect (and saves those changes) so that when the Gameplay Effect is applied, abilities with these tags are cancelled */
+	UE_API void SetAndApplyCanceledAbilityTagChanges(const FInheritedTagContainer& CanceledAbilityWithTagsMods, const FInheritedTagContainer& CanceledAbilityWithoutTagsMods);
+
+private:
+	void CancelOwnerAbilities(const FActiveGameplayEffectsContainer& ActiveGEContainer) const;
+
+#if WITH_EDITOR
+	/** Needed to properly update FInheritedTagContainer properties */
+	UE_API virtual void PostEditChangeProperty(FPropertyChangedEvent& PropertyChangedEvent) override;
+
+	/** Validate incompatible configurations */
+	UE_API virtual EDataValidationResult IsDataValid(class FDataValidationContext& Context) const override;
+
+	/** Get a cached version of the FProperty Name for PostEditChangeProperty */
+	static const FName& GetInheritableCancelAbilitiesWithTagsContainerPropertyName()
+	{
+		static FName NAME_InheritableCancelAbilitiesWithTagsContainer = GET_MEMBER_NAME_CHECKED(UCancelAbilityTagsGameplayEffectComponent, InheritableCancelAbilitiesWithTagsContainer);
+		return NAME_InheritableCancelAbilitiesWithTagsContainer;					
+	}
+
+	static const FName& GetInheritableCancelAbilitiesWithoutTagsContainerPropertyName()
+	{
+		static FName NAME_InheritableCancelAbilitiesWithoutTagsContainer = GET_MEMBER_NAME_CHECKED(UCancelAbilityTagsGameplayEffectComponent, InheritableCancelAbilitiesWithoutTagsContainer);
+		return NAME_InheritableCancelAbilitiesWithoutTagsContainer;
+	}
+#endif // WITH_EDITOR
+
+	/* When do we want the ability cancellation logic for this component to happen? */ 
+	UPROPERTY(EditDefaultsOnly, Category = None, meta = (DisplayName = "Mode"))
+	ECancelAbilityTagsGameplayEffectComponentMode ComponentMode = ECancelAbilityTagsGameplayEffectComponentMode::OnApplication;
+	
+	/** These tags cancel active Gameplay Abilities with these tags. */
+	UPROPERTY(EditDefaultsOnly, Category = None, meta = (DisplayName = "Cancel Abilities With Tags"))
+	FInheritedTagContainer InheritableCancelAbilitiesWithTagsContainer;
+
+	/** These tags cancel active Gameplay Abilities without these tags. */
+	UPROPERTY(EditDefaultsOnly, Category = None, meta = (DisplayName = "Cancel Abilities Without Tags"))
+	FInheritedTagContainer InheritableCancelAbilitiesWithoutTagsContainer;
+};
+
+#undef UE_API
