Improved logging for missing attributes in FActiveGameplayEffectHandle when using GetGameplayEffectMagnitude.
[RN] Improved logging for missing attributes in FActiveGameplayEffectHandle when using GetGameplayEffectMagnitude.
#13446

diff --git a/Engine/Plugins/Runtime/GameplayAbilities/Source/GameplayAbilities/Private/GameplayEffect.cpp b/Engine/Plugins/Runtime/GameplayAbilities/Source/GameplayAbilities/Private/GameplayEffect.cpp
--- a/Engine/Plugins/Runtime/GameplayAbilities/Source/GameplayAbilities/Private/GameplayEffect.cpp
+++ b/Engine/Plugins/Runtime/GameplayAbilities/Source/GameplayAbilities/Private/GameplayEffect.cpp
@@ -3755,6 +3755,9 @@ float FActiveGameplayEffectsContainer::GetGameplayEffectMagnitude(FActiveGamepla
 					return ModSpec.GetEvaluatedMagnitude();
 				}
 			}
+			
+			UE_LOG(LogGameplayEffects, Warning, TEXT("GetGameplayEffectMagnitude with Handle: %s did not find Attribute %s in the modifiers."), *Handle.ToString(), *Attribute.GetName());
+			return -1.f;
 		}
 	}
 
