Consolidates code that checks whether an AttributeSet's FProperty represents a gameplay attribute, and improved coverage and stability.

Previously, AttributeSetInitter allowed any FNumericProperty (including doubles) but caused a crash when encountering an FIntProperty on an attribute set. GameplayDebugger only listed FGameplayAttributeData and FFloatProperty, so doubles were ignored. Runtime spawned attributes only cleaned up float property aggregators, but not for doubles.

Now, all systems support FNumericProperties that are floating points (avoids crash on int) and all property checking is moved to one place.

diff --git a/Engine/Plugins/Runtime/GameplayAbilities/Source/GameplayAbilities/Private/AttributeSet.cpp b/Engine/Plugins/Runtime/GameplayAbilities/Source/GameplayAbilities/Private/AttributeSet.cpp
--- a/Engine/Plugins/Runtime/GameplayAbilities/Source/GameplayAbilities/Private/AttributeSet.cpp
+++ b/Engine/Plugins/Runtime/GameplayAbilities/Source/GameplayAbilities/Private/AttributeSet.cpp
@@ -54,23 +54,20 @@ void FGameplayAttributeData::SetBaseValue(float NewValue)
 
 FGameplayAttribute::FGameplayAttribute(FProperty *NewProperty)
 {
-	// we allow numeric properties and gameplay attribute data properties for now
-	// @todo deprecate numeric properties
-	Attribute = CastField<FNumericProperty>(NewProperty);
-	AttributeOwner = nullptr;
-
-	if (!Attribute.Get())
+	if (FGameplayAttribute::IsSupportedProperty(NewProperty))
 	{
-		if (IsGameplayAttributeDataProperty(NewProperty))
-		{
-			Attribute = NewProperty;
-		}
+		Attribute = NewProperty;
+		AttributeOwner = Attribute->GetOwnerStruct();
+		Attribute->GetName(AttributeName);
 	}
-
-	if (Attribute.Get())
+	else
 	{
- 		AttributeOwner = Attribute->GetOwnerStruct();
- 		Attribute->GetName(AttributeName);
+		// Non-null property passed in that's incompatible.
+		ensureMsgf(!NewProperty, TEXT("Tried to construct FGameplayAttribute with invalid property '%s' of type '%s'. Treating as invalid attribute."), *NewProperty->GetName(), *NewProperty->GetClass()->GetName());
+
+		Attribute = nullptr;
+		AttributeOwner = nullptr;
+		AttributeName = "";
 	}
 }
 
@@ -225,6 +222,13 @@ bool FGameplayAttribute::IsSystemAttribute() const
 	return GetAttributeSetClass()->IsChildOf(UAbilitySystemComponent::StaticClass());
 }
 
+bool FGameplayAttribute::IsSupportedProperty(const FProperty* Property)
+{
+	// FGameplayAttributeData is preferred. Floating point properties are supported for legacy reasons.
+	// @todo: Deprecate numeric properties and only allow FGameplayAttributeData.
+	return Property && (FGameplayAttribute::IsGameplayAttributeDataProperty(Property) || (Property->IsA<FNumericProperty>() && CastField<FNumericProperty>(Property)->IsFloatingPoint()));
+}
+
 bool FGameplayAttribute::IsGameplayAttributeDataProperty(const FProperty* Property)
 {
 	const FStructProperty* StructProp = CastField<FStructProperty>(Property);
@@ -411,11 +415,7 @@ void UAttributeSet::GetAttributesFromSetClass(const TSubclassOf<UAttributeSet>&
 {
 	for (TFieldIterator<FProperty> It(AttributeSetClass); It; ++It)
 	{
-		if (FFloatProperty* FloatProperty = CastField<FFloatProperty>(*It))
-		{
-			Attributes.Add(FGameplayAttribute(FloatProperty));
-		}
-		else if (FGameplayAttribute::IsGameplayAttributeDataProperty(*It))
+		if (FGameplayAttribute::IsSupportedProperty(*It))
 		{
 			Attributes.Add(FGameplayAttribute(*It));
 		}
@@ -620,9 +620,14 @@ void FAttributeSetInitterDiscreteLevels::PreloadAttributeSetData(const TArray<UC
 
 			// Find the FProperty
 			FProperty* Property = FindFProperty<FProperty>(*Set, *AttributeName);
-			if (!IsSupportedProperty(Property))
+			if (!Property)
 			{
-				ABILITY_LOG(Verbose, TEXT("FAttributeSetInitterDiscreteLevels::PreloadAttributeSetData Unable to match Attribute from %s (row: %s)"), *AttributeName, *RowName);
+				ABILITY_LOG(Verbose, TEXT("FAttributeSetInitterDiscreteLevels::PreloadAttributeSetData Unable to match Attribute from %s (row: %s): Property not found"), *AttributeName, *RowName);
+				continue;
+			}
+			else if (!FGameplayAttribute::IsSupportedProperty(Property))
+			{
+				ABILITY_LOG(Verbose, TEXT("FAttributeSetInitterDiscreteLevels::PreloadAttributeSetData Unable to match Attribute from %s (row: %s): Property type '%s' incompatible. Must be FGameplayAttributeData or a floating point type."), *AttributeName, *RowName, *Property->GetClass()->GetName());
 				continue;
 			}
 
@@ -800,10 +805,3 @@ TArray<float> FAttributeSetInitterDiscreteLevels::GetAttributeSetValues(UClass*
 	}
 	return AttributeSetValues;
 }
-
-
-bool FAttributeSetInitterDiscreteLevels::IsSupportedProperty(FProperty* Property) const
-{
-	return (Property && (CastField<FNumericProperty>(Property) || FGameplayAttribute::IsGameplayAttributeDataProperty(Property)));
-}
-
diff --git a/Engine/Plugins/Runtime/GameplayAbilities/Source/GameplayAbilities/Public/AttributeSet.h b/Engine/Plugins/Runtime/GameplayAbilities/Source/GameplayAbilities/Public/AttributeSet.h
--- a/Engine/Plugins/Runtime/GameplayAbilities/Source/GameplayAbilities/Public/AttributeSet.h
+++ b/Engine/Plugins/Runtime/GameplayAbilities/Source/GameplayAbilities/Public/AttributeSet.h
@@ -106,6 +106,9 @@ struct FGameplayAttribute
 	/** Returns true if this is one of the special attributes defined on the AbilitySystemComponent itself */
 	UE_API bool IsSystemAttribute() const;
 
+	/** Returns whether a property is able to represent a Gameplay Attribute. Does not check whether the property is owned by an attribute set. **/
+	static UE_API bool IsSupportedProperty(const FProperty* Prop);
+
 	/** Returns true if the variable associated with Property is of type FGameplayAttributeData or one of its subclasses */
 	static UE_API bool IsGameplayAttributeDataProperty(const FProperty* Property);
 
@@ -356,10 +359,8 @@ struct FAttributeSetInitterDiscreteLevels : public FAttributeSetInitter
 	UE_API virtual void ApplyAttributeDefault(UAbilitySystemComponent* AbilitySystemComponent, FGameplayAttribute& InAttribute, FName GroupName, int32 Level) const override;
 
 	UE_API virtual TArray<float> GetAttributeSetValues(UClass* AttributeSetClass, FProperty* AttributeProperty, FName GroupName) const override;
-private:
-
-	bool IsSupportedProperty(FProperty* Property) const;
 
+private:
 	struct FAttributeDefaultValueList
 	{
 		void AddPair(FProperty* InProperty, float InValue)
diff --git a/Engine/Plugins/Runtime/GameplayAbilities/Source/GameplayAbilitiesEditor/Private/SGameplayAttributeWidget.cpp b/Engine/Plugins/Runtime/GameplayAbilities/Source/GameplayAbilitiesEditor/Private/SGameplayAttributeWidget.cpp
--- a/Engine/Plugins/Runtime/GameplayAbilities/Source/GameplayAbilitiesEditor/Private/SGameplayAttributeWidget.cpp
+++ b/Engine/Plugins/Runtime/GameplayAbilities/Source/GameplayAbilitiesEditor/Private/SGameplayAttributeWidget.cpp
@@ -280,6 +280,12 @@ TSharedPtr<FAttributeViewerNode> SAttributeListWidget::UpdatePropertyOptions()
 			{
 				FProperty *Property = *PropertyIt;
 
+				// Ignore property types that cannot represent gameplay attributes
+				if (!FGameplayAttribute::IsSupportedProperty(Property))
+				{
+					continue;
+				}
+
 				// if we have a search string and this doesn't match, don't show it
 				if (AttributeTextFilter.IsValid() && !AttributeTextFilter->PassesFilter(*Property))
 				{
