GAS: UAttributeSet::InitFromMetaDataTable now only processes property types that can back gameplay attributes.
Checking this explicitly prevents unsafe calls to SetFloatingPointPropertyValue while allowing users to have other properties in their set classes.

diff --git a/Engine/Plugins/Runtime/GameplayAbilities/Source/GameplayAbilities/Private/AttributeSet.cpp b/Engine/Plugins/Runtime/GameplayAbilities/Source/GameplayAbilities/Private/AttributeSet.cpp
--- a/Engine/Plugins/Runtime/GameplayAbilities/Source/GameplayAbilities/Private/AttributeSet.cpp
+++ b/Engine/Plugins/Runtime/GameplayAbilities/Source/GameplayAbilities/Private/AttributeSet.cpp
@@ -444,32 +444,33 @@ void UAttributeSet::InitFromMetaDataTable(const UDataTable* DataTable)
 	for( TFieldIterator<FProperty> It(GetClass(), EFieldIteratorFlags::IncludeSuper) ; It ; ++It )
 	{
 		FProperty* Property = *It;
-		FNumericProperty *NumericProperty = CastField<FNumericProperty>(Property);
-		if (NumericProperty)
+
+		// Only process properties that can back gameplay attributes. They will be either FGameplayAttributeData 
+		// or a floating-point numeric property (floats, doubles, potentially user custom).
+		if (FGameplayAttribute::IsSupportedProperty(Property))
 		{
 			FString RowNameStr = FString::Printf(TEXT("%s.%s"), *Property->GetOwnerVariant().GetName(), *Property->GetName());
-		
-			FAttributeMetaData * MetaData = DataTable->FindRow<FAttributeMetaData>(FName(*RowNameStr), Context, false);
-			if (MetaData)
+			if (FAttributeMetaData* MetaData = DataTable->FindRow<FAttributeMetaData>(FName(*RowNameStr), Context, false))
 			{
-				void *Data = NumericProperty->ContainerPtrToValuePtr<void>(this);
-				NumericProperty->SetFloatingPointPropertyValue(Data, MetaData->BaseValue);
-			}
-		}
-		else if (FGameplayAttribute::IsGameplayAttributeDataProperty(Property))
-		{
-			FString RowNameStr = FString::Printf(TEXT("%s.%s"), *Property->GetOwnerVariant().GetName(), *Property->GetName());
-
-			FAttributeMetaData * MetaData = DataTable->FindRow<FAttributeMetaData>(FName(*RowNameStr), Context, false);
-			if (MetaData)
-			{
-				FStructProperty* StructProperty = CastField<FStructProperty>(Property);
-				check(StructProperty);
-				FGameplayAttributeData* DataPtr = StructProperty->ContainerPtrToValuePtr<FGameplayAttributeData>(this);
-				check(DataPtr);
-				DataPtr->SetBaseValue(MetaData->BaseValue);
-				DataPtr->SetCurrentValue(MetaData->BaseValue);
+				FNumericProperty* NumericProperty = CastField<FNumericProperty>(Property);
+				if (NumericProperty)
+				{
+					// Passing FGameplayAttribute::IsSupportedProperty() as numeric property already implies it's floating point
+					check(NumericProperty->IsFloatingPoint());
+					void* Data = NumericProperty->ContainerPtrToValuePtr<void>(this);
+					NumericProperty->SetFloatingPointPropertyValue(Data, MetaData->BaseValue);
+				}
+				else if (FGameplayAttribute::IsGameplayAttributeDataProperty(Property))
+				{
+					FStructProperty* StructProperty = CastField<FStructProperty>(Property);
+					check(StructProperty);
+					FGameplayAttributeData* DataPtr = StructProperty->ContainerPtrToValuePtr<FGameplayAttributeData>(this);
+					check(DataPtr);
+					DataPtr->SetBaseValue(MetaData->BaseValue);
+					DataPtr->SetCurrentValue(MetaData->BaseValue);
+				}
 			}
+			
 		}
 	}
 
