Temp tag count fix for working with iris.
The tag count replication was not working correctly in Iris so when it's enabled we're falling back to using the minimal and replicated loose tags while I get the new NetSerializer up and running.

diff --git a/Engine/Plugins/Runtime/GameplayAbilities/Source/GameplayAbilities/Private/AbilitySystemComponent.cpp b/Engine/Plugins/Runtime/GameplayAbilities/Source/GameplayAbilities/Private/AbilitySystemComponent.cpp
--- a/Engine/Plugins/Runtime/GameplayAbilities/Source/GameplayAbilities/Private/AbilitySystemComponent.cpp
+++ b/Engine/Plugins/Runtime/GameplayAbilities/Source/GameplayAbilities/Private/AbilitySystemComponent.cpp
@@ -19,6 +19,7 @@
 #include "TimerManager.h"
 #include "Net/Core/PushModel/PushModel.h"
 #include "HAL/IConsoleManager.h"
+#include "Net/Iris/ReplicationSystem/EngineReplicationBridge.h"
 
 #include UE_INLINE_GENERATED_CPP_BY_NAME(AbilitySystemComponent)
 
@@ -632,6 +633,27 @@ FOnActiveGameplayEffectInhibitionChanged* UAbilitySystemComponent::OnGameplayEff
 	return ActiveEffect ? &ActiveEffect->EventSet.OnInhibitionChanged : nullptr;
 }
 
+void UAbilitySystemComponent::SetTagMapCount(const FGameplayTag& Tag, int32 NewCount, EGameplayTagReplicationState TagRepState)
+{
+	if (UE::Net::ShouldUseIrisReplication())
+	{
+		// Iris replication doesn't work the same so we're falling back to using the Minimal and replicated loose tags until it's fixed.
+		ENetMode NetMode = GetWorld()->GetNetMode();
+		if (NetMode >= NM_DedicatedServer)
+		{
+			if (TagRepState > EGameplayTagReplicationState::TagOnly)
+			{
+				GetReplicatedLooseTags_Mutable().SetTagCount(Tag, NewCount);
+			}
+		}
+	}
+		
+	if (GameplayTagCountContainer.SetTagCount(Tag, NewCount, TagRepState))
+	{
+		OnTagUpdated(Tag, NewCount > 0);
+	}
+}
+
 FOnGivenActiveGameplayEffectRemoved& UAbilitySystemComponent::OnAnyGameplayEffectRemovedDelegate()
 {
 	return ActiveGameplayEffects.OnActiveGameplayEffectRemovedDelegate;
@@ -639,6 +661,31 @@ FOnGivenActiveGameplayEffectRemoved& UAbilitySystemComponent::OnAnyGameplayEffec
 
 void UAbilitySystemComponent::UpdateTagMap_Internal(const FGameplayTagContainer& Container, const int32 CountDelta, EGameplayTagReplicationState TagRepState)
 {
+	if (UE::Net::ShouldUseIrisReplication())
+	{
+		// Iris replication doesn't work the same so we're falling back to using the Minimal and replicated loose tags until it's fixed.
+		ENetMode NetMode = GetWorld()->GetNetMode();
+		if (NetMode >= NM_DedicatedServer)
+		{
+			if (TagRepState == EGameplayTagReplicationState::TagOnly)
+			{
+				if (CountDelta > 0)
+				{
+					GetMinimalReplicationTags_Mutable().AddTags(Container);
+				}
+			}
+			else if (CountDelta < 0 || TagRepState > EGameplayTagReplicationState::TagOnly)
+			{
+				for (auto TagIt = Container.CreateConstIterator(); TagIt; ++TagIt)
+				{
+					const FGameplayTag& Tag = *TagIt;
+					int32& TagCount = GetReplicatedLooseTags_Mutable().TagMap.FindOrAdd(Tag);
+					TagCount += CountDelta;
+				}
+			}
+		}
+	}
+	
 	// For removal, reorder calls so that FillParentTags is only called once
 	if (CountDelta > 0)
 	{
@@ -664,6 +711,16 @@ void UAbilitySystemComponent::UpdateTagMap_Internal(const FGameplayTagContainer&
 			if (GameplayTagCountContainer.UpdateTagCount_DeferredParentRemoval(Tag, CountDelta, DeferredTagChangeDelegates))
 			{
 				RemovedTags.Add(Tag);
+				if (UE::Net::ShouldUseIrisReplication())
+				{
+					// Iris replication doesn't work the same so we're falling back to using the Minimal and replicated loose tags until it's fixed.
+					ENetMode NetMode = GetWorld()->GetNetMode();
+					if (NetMode >= NM_DedicatedServer)
+					{
+						GetMinimalReplicationTags_Mutable().RemoveTag(Tag);
+						GetReplicatedLooseTags_Mutable().RemoveTag(Tag);
+					}
+				}
 			}
 		}
 
diff --git a/Engine/Plugins/Runtime/GameplayAbilities/Source/GameplayAbilities/Public/AbilitySystemComponent.h b/Engine/Plugins/Runtime/GameplayAbilities/Source/GameplayAbilities/Public/AbilitySystemComponent.h
--- a/Engine/Plugins/Runtime/GameplayAbilities/Source/GameplayAbilities/Public/AbilitySystemComponent.h
+++ b/Engine/Plugins/Runtime/GameplayAbilities/Source/GameplayAbilities/Public/AbilitySystemComponent.h
@@ -612,13 +612,7 @@ class UAbilitySystemComponent : public UGameplayTasksComponent, public IGameplay
 	}
 
 	/** Forcibly sets the number of instances of a given tag */
-	inline void SetTagMapCount(const FGameplayTag& Tag, int32 NewCount, EGameplayTagReplicationState TagRepState = EGameplayTagReplicationState::None)
-	{
-		if (GameplayTagCountContainer.SetTagCount(Tag, NewCount, TagRepState))
-		{
-			OnTagUpdated(Tag, NewCount > 0);
-		}
-	}
+	void SetTagMapCount(const FGameplayTag& Tag, int32 NewCount, EGameplayTagReplicationState TagRepState = EGameplayTagReplicationState::None);
 
 	/** Update the number of instances of a given tag and calls callback */
 	inline void UpdateTagMap(const FGameplayTag& BaseTag, int32 CountDelta, EGameplayTagReplicationState TagRepState)
