GameplayEffect applications can now extend the duration of an applied GE. Introduced configurable MaxDuration.
- Added the 'ExtendDuration' option for StackDurationRefreshPolicy. This lets stacking GEs add their duration, or whatever duration is provided by the new GE spec, onto the existing duration. This is compatible with GEs that don't actually stack (stack up to 1), and convenient for consumables where multiple consumptions should increase duration.
- Added the 'MaxDurationMagnitude' setting. This allows designers to configure a MaxDuration: fixed, scaled by effect level, or SetByCaller.
- Moved the OnDurationChange call to later during GE spec application, namely the moment the duration is finalized. Previously, OnDurationChange didn't have access to the final duration which might be affected by modifiers (and now by carry-over value).
- Added a GE validation rule that warns users when they combine 'ExtendDuration' with 'RefreshDuration' options on stack expiration. Since both access the applied Spec's duration (write or read), it's likely that the combination of settings are not what a designer has in mind, so warning them against that.

diff --git a/Engine/Plugins/Runtime/GameplayAbilities/Source/GameplayAbilities/Private/GameplayEffect.cpp b/Engine/Plugins/Runtime/GameplayAbilities/Source/GameplayAbilities/Private/GameplayEffect.cpp
--- a/Engine/Plugins/Runtime/GameplayAbilities/Source/GameplayAbilities/Private/GameplayEffect.cpp
+++ b/Engine/Plugins/Runtime/GameplayAbilities/Source/GameplayAbilities/Private/GameplayEffect.cpp
@@ -289,6 +289,7 @@ EDataValidationResult UGameplayEffect::IsDataValid(FDataValidationContext& Conte
 	if (DurationPolicy == EGameplayEffectDurationType::HasDuration)
 	{
 		ValidationResult = CombineDataValidationResults(ValidationResult, DurationMagnitude.IsDataValid(Context, AssetPathName + TEXT(".DurationMagnitude")));
+		ValidationResult = CombineDataValidationResults(ValidationResult, MaxDurationMagnitude.IsDataValid(Context, AssetPathName + TEXT(".MaxDurationMagnitude")));
 	}
 
 	ValidationResult = CombineDataValidationResults(ValidationResult, Period.IsDataValid(Context, AssetPathName + TEXT(".Period")));
@@ -310,6 +311,12 @@ EDataValidationResult UGameplayEffect::IsDataValid(FDataValidationContext& Conte
 		}
 	}
 
+	if (DurationPolicy == EGameplayEffectDurationType::HasDuration && GetStackingType() != EGameplayEffectStackingType::None &&
+		StackDurationRefreshPolicy == EGameplayEffectStackingDurationPolicy::ExtendDuration && StackExpirationPolicy != EGameplayEffectStackingExpirationPolicy::ClearEntireStack)
+	{
+		Context.AddWarning(LOCTEXT("ExtendAndRefreshDurationBehaviour", "Stack Duration and Expiration Policy: When duration is extended on new stack but GE expiration does not clear the whole stack, duration will be refreshed using the total extended duration."));
+	}
+
 	PRAGMA_DISABLE_DEPRECATION_WARNINGS
 	if (DurationPolicy == EGameplayEffectDurationType::Instant && StackingType != EGameplayEffectStackingType::None)
 	{
@@ -1762,6 +1769,13 @@ void FGameplayEffectSpec::SetupAttributeCaptureDefinitions()
 		{
 			CapturedRelevantAttributes.AddCaptureDefinition(CurDurationCaptureDef);
 		}
+
+		CaptureDefs.Reset();
+		Def->MaxDurationMagnitude.GetAttributeCaptureDefinitions(CaptureDefs);
+		for (const FGameplayEffectAttributeCaptureDefinition& CurDurationCaptureDef : CaptureDefs)
+		{
+			CapturedRelevantAttributes.AddCaptureDefinition(CurDurationCaptureDef);
+		}
 	}
 
 	// Gather all capture definitions from modifiers
@@ -1825,32 +1839,42 @@ void FGameplayEffectSpec::RecaptureSourceActorTags()
 	EffectContext.GetOwnedGameplayTags(CapturedSourceTags.GetActorTags(), CapturedSourceTags.GetSpecTags());
 }
 
-bool FGameplayEffectSpec::AttemptCalculateDurationFromDef(OUT float& OutDefDuration) const
+bool FGameplayEffectSpec::AttemptCalculateDurationRelatedMagnitude(const FGameplayEffectModifierMagnitude& MagnitudeDef, OUT float& OutDurationValue) const
 {
-	check(Def);
-
 	bool bCalculatedDuration = true;
 
 	const EGameplayEffectDurationType DurType = Def->DurationPolicy;
 	if (DurType == EGameplayEffectDurationType::Infinite)
 	{
-		OutDefDuration = UGameplayEffect::INFINITE_DURATION;
+		OutDurationValue = UGameplayEffect::INFINITE_DURATION;
 	}
 	else if (DurType == EGameplayEffectDurationType::Instant)
 	{
-		OutDefDuration = UGameplayEffect::INSTANT_APPLICATION;
+		OutDurationValue = UGameplayEffect::INSTANT_APPLICATION;
 	}
 	else
 	{
 		// The last parameters (false, 1.f) are so that if SetByCaller hasn't been set yet, we don't warn and default
 		// to 1.f. This is so that the rest of the system doesn't treat the effect as an instant effect. 1.f is arbitrary
 		// and this makes it illegal to SetByCaller something into an instant effect.
-		bCalculatedDuration = Def->DurationMagnitude.AttemptCalculateMagnitude(*this, OutDefDuration, false, 1.f);
+		bCalculatedDuration = MagnitudeDef.AttemptCalculateMagnitude(*this, OutDurationValue, false, 1.f);
 	}
 
 	return bCalculatedDuration;
 }
 
+bool FGameplayEffectSpec::AttemptCalculateDurationFromDef(OUT float& OutDefDuration) const
+{
+	check(Def);
+	return AttemptCalculateDurationRelatedMagnitude(Def->DurationMagnitude, OutDefDuration);
+}
+
+bool FGameplayEffectSpec::AttemptCalculateMaxDurationFromDef(OUT float& OutDefMaxDuration) const
+{
+	check(Def);
+	return AttemptCalculateDurationRelatedMagnitude(Def->MaxDurationMagnitude, OutDefMaxDuration);
+}
+
 void FGameplayEffectSpec::SetLevel(float InLevel)
 {
 	Level = InLevel;
@@ -1951,7 +1975,6 @@ float FGameplayEffectSpec::GetPeriod() const
 PRAGMA_DISABLE_DEPRECATION_WARNINGS
 void FGameplayEffectSpec::SetStackCount(int32 NewStackCount)
 {
-
 	StackCount = NewStackCount;
 }
 
@@ -4100,11 +4123,18 @@ FActiveGameplayEffect* FActiveGameplayEffectsContainer::ApplyGameplayEffectSpec(
 	FActiveGameplayEffect* AppliedActiveGE = nullptr;
 	FActiveGameplayEffect* ExistingStackableGE = FindStackableActiveGameplayEffect(Spec);
 
-	bool bSetDuration = true;
-	bool bSetPeriod = true;
+	// Whether to (re)schedule the duration check with world timer manager
+	bool bSetDurationTimer = true;
+	// Whether to (re)schedule the periodic tick with world timer manager
+	bool bSetPeriodTimer = true;
+
 	int32 StartingStackCount = 0;
 	int32 NewStackCount = 0;
 
+	// If we're applying a stack and the GE is configured to extend duration, 
+	// we will remember how much duration was left of the existing application.
+	float CarryOverDuration = 0.0f;
+
 	// For HasActiveGameplayEffectOverflowBehavior() 
 	// and EActiveGameplayEffectOverflowBehavior
 	// and bAllowPredictiveApplicationOfStackingGameplayEffects
@@ -4171,6 +4201,13 @@ FActiveGameplayEffect* FActiveGameplayEffectsContainer::ApplyGameplayEffectSpec(
 		ensureMsgf(ExistingSpec.DynamicGrantedTags == Spec.DynamicGrantedTags, TEXT("While adding a stack of the gameplay effect: %s, the old stack and the new application had different dynamically granted tags, which is currently not resolved properly!"), *Spec.Def->GetName());
 		ensureMsgf(ExistingSpec.GetDynamicAssetTags() == Spec.GetDynamicAssetTags(), TEXT("While adding a stack of the gameplay effect: %s, the old stack and the new application had different dynamic asset tags, which is currently not resolved properly! Existing: %s. New: %s"), *Spec.Def->GetName(), *ExistingSpec.GetDynamicAssetTags().ToStringSimple(), *Spec.GetDynamicAssetTags().ToStringSimple());
 
+		// If GE's stack duration policy is configured to extend duration, the remaining time and the new spec's time will be summed.
+		// Store the remaining time for later use. Cache it now before overwriting existing spec's duration.
+		if (ExistingSpec.Def->StackDurationRefreshPolicy == EGameplayEffectStackingDurationPolicy::ExtendDuration)
+		{
+			CarryOverDuration = ExistingStackableGE->GetTimeRemaining(GetWorldTime());
+		}
+
 PRAGMA_DISABLE_DEPRECATION_WARNINGS
 		// We only grant abilities on the first apply. So we *dont* want the new spec's GrantedAbilitySpecs list
 		TArray<FGameplayAbilitySpecDef>	GrantedSpecTempArray(MoveTemp(ExistingStackableGE->Spec.GrantedAbilitySpecs));
@@ -4191,17 +4228,18 @@ PRAGMA_ENABLE_DEPRECATION_WARNINGS
 		// Client's predicted GE will be combined with the existing GE and we don't want to set any duration timer in this case
 		if ((InPredictionKey.IsLocalClientKey() && InPredictionKey.WasLocallyGenerated()) || GEDef->StackDurationRefreshPolicy == EGameplayEffectStackingDurationPolicy::NeverRefresh)
 		{
-			bSetDuration = false;
+			bSetDurationTimer = false;
 		}
 		else
 		{
+			// Duration is going to be refreshed. Rebase the active GE's start time.
 			RestartActiveGameplayEffectDuration(*ExistingStackableGE);
 		}
 
 		// Make sure the GE actually wants to reset its period
 		if (GEDef->StackPeriodResetPolicy == EGameplayEffectStackingPeriodPolicy::NeverReset)
 		{
-			bSetPeriod = false;
+			bSetPeriodTimer = false;
 		}
 	}
 	else
@@ -4321,17 +4359,43 @@ PRAGMA_ENABLE_DEPRECATION_WARNINGS
 	{
 		float FinalDuration = AppliedEffectSpec.CalculateModifiedDuration();
 
+		// If GE is configured to extend duration when stacking, add the remaining time of existing effect to the new final duration
+		if (CarryOverDuration > 0.0f)
+		{
+			FinalDuration += CarryOverDuration;
+		}
+
+		// Calculate max duration from the spec. If the new spec provides a different value for max duration, for example 
+		// through different SetByCaller magnitudes, then that takes precedence over the previously applied spec's.
+		float SpecMaxDuration = 0.0f;
+		if (!AppliedEffectSpec.AttemptCalculateMaxDurationFromDef(SpecMaxDuration))
+		{
+			if (AppliedEffectSpec.Def->MaxDurationMagnitude.GetMagnitudeCalculationType() == EGameplayEffectMagnitudeCalculation::SetByCaller)
+			{
+				AppliedEffectSpec.Def->MaxDurationMagnitude.AttemptCalculateMagnitude(AppliedEffectSpec, SpecMaxDuration);
+			}
+		}
+
 		// We cannot mod ourselves into an instant or infinite duration effect
 		if (FinalDuration <= 0.f)
 		{
 			UE_LOG(LogGameplayEffects, Error, TEXT("ActiveGE %s Duration was modified to %.2f. Clamping to 0.1s duration."), *AppliedActiveGE->GetDebugString(), FinalDuration);
 			FinalDuration = 0.1f;
 		}
+		else if (SpecMaxDuration > 0.0f && FinalDuration > SpecMaxDuration)
+		{
+			UE_LOG(LogGameplayEffects, Verbose, TEXT("ActiveGE %s Duration (%.2f) exceeds MaxDuration (%.2f). Clamping to MaxDuration."), *AppliedActiveGE->GetDebugString(), FinalDuration, SpecMaxDuration);
+			FinalDuration = SpecMaxDuration;
+		}
+
+		// Store the final calculated duration on the applied spec
+		AppliedEffectSpec.SetDuration(FinalDuration, /*bLockDuration=*/false);
 
-		AppliedEffectSpec.SetDuration(FinalDuration, true);
+		// Broadcast latest duration
+		OnDurationChange(*AppliedActiveGE);
 
 		// Register duration callbacks with the timer manager
-		if (Owner && bSetDuration)
+		if (Owner && bSetDurationTimer)
 		{
 			FTimerManager& TimerManager = Owner->GetWorld()->GetTimerManager();
 			FTimerDelegate Delegate = FTimerDelegate::CreateUObject(Owner, &UAbilitySystemComponent::CheckDurationExpired, AppliedActiveGE->Handle);
@@ -4346,7 +4410,7 @@ PRAGMA_ENABLE_DEPRECATION_WARNINGS
 	}
 	
 	// Register period callbacks with the timer manager
-	if (bSetPeriod && Owner && (AppliedEffectSpec.GetPeriod() > UGameplayEffect::NO_PERIOD))
+	if (bSetPeriodTimer && Owner && (AppliedEffectSpec.GetPeriod() > UGameplayEffect::NO_PERIOD))
 	{
 		FTimerManager& TimerManager = Owner->GetWorld()->GetTimerManager();
 		FTimerDelegate Delegate = FTimerDelegate::CreateUObject(Owner, &UAbilitySystemComponent::ExecutePeriodicEffect, AppliedActiveGE->Handle);
@@ -4979,8 +5043,6 @@ void FActiveGameplayEffectsContainer::RestartActiveGameplayEffectDuration(FActiv
 	ActiveGameplayEffect.CachedStartServerWorldTime = ActiveGameplayEffect.StartServerWorldTime;
 	ActiveGameplayEffect.StartWorldTime = GetWorldTime();
 	MarkItemDirty(ActiveGameplayEffect);
-
-	OnDurationChange(ActiveGameplayEffect);
 }
 
 ELifetimeCondition FActiveGameplayEffectsContainer::GetReplicationCondition() const
@@ -5268,6 +5330,10 @@ void FActiveGameplayEffectsContainer::CheckDuration(FActiveGameplayEffectHandle
 		if (RefreshStartTime)
 		{
 			RestartActiveGameplayEffectDuration(Effect);
+
+			// Since the GE will go through another iteration, broadcast OnDurationChange
+			// allowing game code to respond to the new start and end time.
+			OnDurationChange(Effect);
 		}
 
 		if (RefreshDurationTimer)
diff --git a/Engine/Plugins/Runtime/GameplayAbilities/Source/GameplayAbilities/Public/GameplayEffect.h b/Engine/Plugins/Runtime/GameplayAbilities/Source/GameplayAbilities/Public/GameplayEffect.h
--- a/Engine/Plugins/Runtime/GameplayAbilities/Source/GameplayAbilities/Public/GameplayEffect.h
+++ b/Engine/Plugins/Runtime/GameplayAbilities/Source/GameplayAbilities/Public/GameplayEffect.h
@@ -679,6 +679,9 @@ enum class EGameplayEffectStackingDurationPolicy : uint8
 
 	/** The duration of the effect will never be refreshed */
 	NeverRefresh,
+
+	/** New stacks will add their GE spec's duration onto current remaining time */
+	ExtendDuration,
 };
 
 /** Enumeration of policies for dealing with the period of a gameplay effect while stacking */
@@ -1044,6 +1047,15 @@ struct FGameplayEffectSpec
 	 */
 	UE_API bool AttemptCalculateDurationFromDef(OUT float& OutDefDuration) const;
 
+	/**
+	 * Helper function to attempt to calculate the maximum duration of the spec from its GE definition
+	 *
+	 * @param OutDefMaxDuration	Computed max duration of the spec from its GE definition
+	 *
+	 * @return True if the calculation was successful, false if it was not
+	 */
+	UE_API bool AttemptCalculateMaxDurationFromDef(OUT float& OutDefMaxDuration) const;
+
 	/** Sets duration. This should only be called as the GameplayEffect is being created and applied; Ignores calls after attribute capture */
 	UE_API void SetDuration(float NewDuration, bool bLockDuration);
 
@@ -1162,6 +1174,9 @@ private:
 
 	UE_API void CaptureDataFromSource(bool bSkipRecaptureSourceActorTags = false);
 
+	/** Attempts to calculate a duration-related magnitude. Helper function to allow reuse between Duration and MaxDuration. */
+	UE_API bool AttemptCalculateDurationRelatedMagnitude(const FGameplayEffectModifierMagnitude& MagnitudeDef, OUT float& OutDurationValue) const;
+
 public:
 
 	/** GameplayEfect definition. The static data that this spec points to. */
@@ -1335,7 +1350,7 @@ struct FActiveGameplayEffect : public FFastArraySerializerItem
 
 	float GetTimeRemaining(float WorldTime) const
 	{
-		float Duration = GetDuration();		
+		float Duration = GetDuration();
 		return (Duration == FGameplayEffectConstants::INFINITE_DURATION ? -1.f : Duration - (WorldTime - StartWorldTime));
 	}
 	
@@ -1930,6 +1945,11 @@ private:
 	/** Internal callback fired as a result of a custom modifier magnitude external dependency delegate firing; Updates affected active gameplay effects as necessary */
 	UE_API void OnCustomMagnitudeExternalDependencyFired(TSubclassOf<UGameplayModMagnitudeCalculation> MagnitudeCalculationClass);
 
+	/** 
+	 * Updates the active GE's start time to current world time. Despite what the function name suggests, duration is not refreshed or re-calculated here.
+	 * Only queries to the active effect's remaining time are affected. Call-sites should update duration on the GE spec if relevant, and should ensure that 
+	 * a new expiration timer is set in the world timer manager that fires at T = (new) start time + (new) duration.
+	 */
 	UE_API void RestartActiveGameplayEffectDuration(FActiveGameplayEffect& ActiveGameplayEffect);
 
 	// -------------------------------------------------------------------------------------------
@@ -2228,10 +2248,14 @@ public:
 	UPROPERTY(EditDefaultsOnly, Category=Duration)
 	EGameplayEffectDurationType DurationPolicy;
 
-	/** Duration in seconds. 0.0 for instantaneous effects; -1.0 for infinite duration. */
+	/** Duration in seconds. 0.0 for instantaneous effects; -1.0 for infinite duration. When applying stacks onto an existing active effect, the new spec's Duration is considered. */
 	UPROPERTY(EditDefaultsOnly, Category=Duration, meta=(EditCondition="DurationPolicy == EGameplayEffectDurationType::HasDuration", EditConditionHides))
 	FGameplayEffectModifierMagnitude DurationMagnitude;
 
+	/** MaxDuration in seconds. <= 0.0 for unlimited. When applying stacks onto an existing active effect, the new spec's MaxDuration is considered. */
+	UPROPERTY(EditDefaultsOnly, Category = Duration, meta = (EditCondition = "DurationPolicy == EGameplayEffectDurationType::HasDuration", EditConditionHides))
+	FGameplayEffectModifierMagnitude MaxDurationMagnitude;
+
 	/** Period in seconds. 0.0 for non-periodic effects */
 	UPROPERTY(EditDefaultsOnly, BlueprintReadOnly, Category="Duration|Period", meta=(EditCondition="DurationPolicy != EGameplayEffectDurationType::Instant", EditConditionHides))
 	FScalableFloat	Period;
