Unshelved from pending changelist '43141229':

PR #13341: Fix animation parameter when calling ServerCurrentMontageSetPlayRate

diff --git a/Engine/Plugins/Runtime/GameplayAbilities/Source/GameplayAbilities/Private/AbilitySystemComponent_Abilities.cpp b/Engine/Plugins/Runtime/GameplayAbilities/Source/GameplayAbilities/Private/AbilitySystemComponent_Abilities.cpp
--- a/Engine/Plugins/Runtime/GameplayAbilities/Source/GameplayAbilities/Private/AbilitySystemComponent_Abilities.cpp
+++ b/Engine/Plugins/Runtime/GameplayAbilities/Source/GameplayAbilities/Private/AbilitySystemComponent_Abilities.cpp
@@ -3559,7 +3559,7 @@ void UAbilitySystemComponent::CurrentMontageSetPlayRate(float InPlayRate)
 		else
 		{
 			UAnimSequenceBase* Animation = LocalAnimMontageInfo.AnimMontage->IsDynamicMontage() ? LocalAnimMontageInfo.AnimMontage->GetFirstAnimReference() : LocalAnimMontageInfo.AnimMontage.Get();
-			ServerCurrentMontageSetPlayRate(LocalAnimMontageInfo.AnimMontage, InPlayRate);
+			ServerCurrentMontageSetPlayRate(Animation, InPlayRate);
 		}
 	}
 }
