Fix ability system holding onto released structs

#rnx

diff --git a/Engine/Plugins/Runtime/GameplayAbilities/Source/GameplayAbilities/Private/AbilitySystemGlobals.cpp b/Engine/Plugins/Runtime/GameplayAbilities/Source/GameplayAbilities/Private/AbilitySystemGlobals.cpp
--- a/Engine/Plugins/Runtime/GameplayAbilities/Source/GameplayAbilities/Private/AbilitySystemGlobals.cpp
+++ b/Engine/Plugins/Runtime/GameplayAbilities/Source/GameplayAbilities/Private/AbilitySystemGlobals.cpp
@@ -16,6 +16,7 @@
 #include "GameplayTagResponseTable.h"
 #include "GameplayTagsManager.h"
 #include "Engine/Engine.h"
+#include "Modules/ModuleManager.h"
 #include "UObject/UObjectIterator.h"
 
 #if UE_WITH_IRIS
@@ -371,6 +372,29 @@ void UAbilitySystemGlobals::InitTargetDataScriptStructCache()
 {
 	TargetDataStructCache.InitForType(FGameplayAbilityTargetData::StaticStruct());
 	EffectContextStructCache.InitForType(FGameplayEffectContext::StaticStruct());
+
+	// Handle loading for new modules that might have new structs to add
+	if (!ModulesChangedHandle.IsValid())
+	{
+		ModulesChangedHandle = FModuleManager::Get().OnModulesChanged().AddWeakLambda(this, [this](FName ModuleName, EModuleChangeReason Reason)
+		{
+			if (Reason == EModuleChangeReason::ModuleLoaded)
+			{
+				TargetDataStructCache.SetInitDirty();
+				EffectContextStructCache.SetInitDirty();
+			}
+		});
+	}
+
+	// Handle module unloading to avoid holding onto released structs
+	if (!ModulesUnloadedHandle.IsValid())
+	{
+		ModulesUnloadedHandle = FCoreUObjectDelegates::CompiledInUObjectsRemovedDelegate.AddWeakLambda(this, [this](TConstArrayView<UPackage*> Packages)
+		{
+			TargetDataStructCache.RemoveInPackages(Packages);
+			EffectContextStructCache.RemoveInPackages(Packages);
+		});
+	}
 }
 
 // --------------------------------------------------------------------
@@ -716,23 +740,49 @@ void UAbilitySystemGlobals::NonShipping_ApplyGlobalAbilityScaler_Duration(float&
 
 void FNetSerializeScriptStructCache::InitForType(UScriptStruct* InScriptStruct)
 {
+	TRACE_CPUPROFILER_EVENT_SCOPE(FNetSerializeScriptStructCache::InitForType);
+
 	ScriptStructs.Reset();
+	BaseStructType = InScriptStruct;
 
 	// Find all script structs of this type and add them to the list
 	// (not sure of a better way to do this but it should only happen once at startup)
 	for (TObjectIterator<UScriptStruct> It; It; ++It)
 	{
-		if (It->IsChildOf(InScriptStruct))
+		if (It->IsChildOf(InScriptStruct) && IsValid(*It))
 		{
+			UE_LOG(LogAbilitySystem, Verbose, TEXT("FNetSerializeScriptStructCache::InitForType: found struct type %s"),
+				*It->GetName());
 			ScriptStructs.Add(*It);
 		}
 	}
 	
 	ScriptStructs.Sort([](const UScriptStruct& A, const UScriptStruct& B) { return A.GetName().ToLower() > B.GetName().ToLower(); });
+
+	bIsInitDirty = false;
+}
+
+void FNetSerializeScriptStructCache::RemoveInPackages(TConstArrayView<UPackage*> Packages)
+{
+	TRACE_CPUPROFILER_EVENT_SCOPE(FNetSerializeScriptStructCache::RemoveInPackages);
+
+	for (auto It = ScriptStructs.CreateIterator(); It; ++It)
+	{
+		if (Packages.Contains(It->GetPackage()))
+		{
+			ABILITY_LOG(Verbose, TEXT("FNetSerializeScriptStructCache::RemoveInPackages: removed %s"), *It->GetName());
+			It.RemoveCurrent();
+		}
+	}
 }
 
 bool FNetSerializeScriptStructCache::NetSerialize(FArchive& Ar, UScriptStruct*& Struct)
 {
+	if (bIsInitDirty)
+	{
+		InitForType(BaseStructType);
+	}
+
 	if (Ar.IsSaving())
 	{
 		int32 idx;
diff --git a/Engine/Plugins/Runtime/GameplayAbilities/Source/GameplayAbilities/Public/AbilitySystemGlobals.h b/Engine/Plugins/Runtime/GameplayAbilities/Source/GameplayAbilities/Public/AbilitySystemGlobals.h
--- a/Engine/Plugins/Runtime/GameplayAbilities/Source/GameplayAbilities/Public/AbilitySystemGlobals.h
+++ b/Engine/Plugins/Runtime/GameplayAbilities/Source/GameplayAbilities/Public/AbilitySystemGlobals.h
@@ -32,11 +32,23 @@ struct FNetSerializeScriptStructCache
 
 	void InitForType(UScriptStruct* InScriptStruct);
 
+	void SetInitDirty()
+	{
+		bIsInitDirty = true;
+	}
+
+	void RemoveInPackages(TConstArrayView<UPackage*>);
+
 	// Serializes reference to given script struct (must be in the cache)
 	bool NetSerialize(FArchive& Ar, UScriptStruct*& Struct);
 
 	UPROPERTY()
 	TArray<TObjectPtr<UScriptStruct>> ScriptStructs;
+
+	UPROPERTY()
+	TObjectPtr<UScriptStruct> BaseStructType;
+
+	bool bIsInitDirty = false;
 };
 
 
@@ -436,6 +448,9 @@ protected:
 private:
 	UE_API void PerformDeveloperSettingsUpgrade();
 
+	FDelegateHandle ModulesChangedHandle;
+	FDelegateHandle ModulesUnloadedHandle;
+
 public:
 	//To add functionality for opening assets directly from the game.
 	UE_API void Notify_OpenAssetInEditor(FString AssetName, int AssetType);
