Iris fixed bug where FGameplayAbilityTargetingLocationInfoNetSerializer would not properly call FreeDynamicState

#rnx

diff --git a/Engine/Plugins/Runtime/GameplayAbilities/Source/GameplayAbilities/Private/Serialization/GameplayAbilityTargetingLocationInfoNetSerializer.cpp b/Engine/Plugins/Runtime/GameplayAbilities/Source/GameplayAbilities/Private/Serialization/GameplayAbilityTargetingLocationInfoNetSerializer.cpp
--- a/Engine/Plugins/Runtime/GameplayAbilities/Source/GameplayAbilities/Private/Serialization/GameplayAbilityTargetingLocationInfoNetSerializer.cpp
+++ b/Engine/Plugins/Runtime/GameplayAbilities/Source/GameplayAbilities/Private/Serialization/GameplayAbilityTargetingLocationInfoNetSerializer.cpp
@@ -191,6 +191,7 @@ void FGameplayAbilityTargetingLocationInfoNetSerializer::Deserialize(FNetSeriali
 		FNetFreeDynamicStateArgs FreeArgs;
 		FreeArgs.NetSerializerConfig = &StructNetSerializerConfig;
 		FreeArgs.Source = NetSerializerValuePointer(&Target);
+		StructNetSerializer->FreeDynamicState(Context, FreeArgs);
 	}
 
 	Target = TempValue;
