Unshelved from pending changelist '45833999':

Small fix for Adding a single loose tag following a different logic path to tag containers and missing the Iris fix code.

diff --git a/Engine/Plugins/Runtime/GameplayAbilities/Source/GameplayAbilities/Private/AbilitySystemComponent.cpp b/Engine/Plugins/Runtime/GameplayAbilities/Source/GameplayAbilities/Private/AbilitySystemComponent.cpp
--- a/Engine/Plugins/Runtime/GameplayAbilities/Source/GameplayAbilities/Private/AbilitySystemComponent.cpp
+++ b/Engine/Plugins/Runtime/GameplayAbilities/Source/GameplayAbilities/Private/AbilitySystemComponent.cpp
@@ -637,6 +637,67 @@ FOnGivenActiveGameplayEffectRemoved& UAbilitySystemComponent::OnAnyGameplayEffec
 	return ActiveGameplayEffects.OnActiveGameplayEffectRemovedDelegate;
 }
 
+void UAbilitySystemComponent::UpdateTagMapSingle_Internal(const FGameplayTag& Tag, const int32 CountDelta, EGameplayTagReplicationState TagRepState)
+{
+	if (UE::Net::ShouldUseIrisReplication())
+	{
+		// Iris replication doesn't work the same so we're falling back to using the Minimal and replicated loose tags until it's fixed.
+		ENetMode NetMode = GetWorld()->GetNetMode();
+		if (NetMode == NM_DedicatedServer || NetMode == NM_ListenServer)
+		{
+			if (TagRepState == EGameplayTagReplicationState::TagOnly)
+			{
+				if (CountDelta > 0)
+				{
+					GetMinimalReplicationTags_Mutable().AddTag(Tag);
+				}
+			}
+			else if (CountDelta < 0 || TagRepState > EGameplayTagReplicationState::TagOnly)
+			{
+				int32& TagCount = GetReplicatedLooseTags_Mutable().TagMap.FindOrAdd(Tag);
+				TagCount += CountDelta;
+			}
+		}
+	}
+	
+	// For removal, reorder calls so that FillParentTags is only called once
+	if (CountDelta > 0)
+	{
+		if (GameplayTagCountContainer.UpdateTagCount(Tag, CountDelta, TagRepState))
+		{
+			OnTagUpdated(Tag, true);
+		}
+	}
+	else if (CountDelta < 0)
+	{
+		TArray<FDeferredTagChangeDelegate> DeferredTagChangeDelegates;
+		
+		if (GameplayTagCountContainer.UpdateTagCount(Tag, CountDelta, TagRepState))
+		{
+			OnTagUpdated(Tag, false);
+			if (UE::Net::ShouldUseIrisReplication())
+			{
+				// Iris replication doesn't work the same so we're falling back to using the Minimal and replicated loose tags until it's fixed.
+				ENetMode NetMode = GetWorld()->GetNetMode();
+				if (NetMode == NM_DedicatedServer || NetMode == NM_ListenServer)
+				{
+					FMinimalReplicationTagCountMap& TagMap = GetMinimalReplicationTags_Mutable();
+					if (TagMap.TagMap.Contains(Tag))
+					{
+						TagMap.RemoveTag(Tag);
+					}
+
+					TagMap = GetReplicatedLooseTags_Mutable();
+					if (TagMap.TagMap.Contains(Tag))
+					{
+						TagMap.RemoveTag(Tag);
+					}
+				}
+			}
+		}
+	}
+}
+
 void UAbilitySystemComponent::UpdateTagMap_Internal(const FGameplayTagContainer& Container, const int32 CountDelta, EGameplayTagReplicationState TagRepState)
 {
 	if (UE::Net::ShouldUseIrisReplication())
@@ -695,8 +756,17 @@ void UAbilitySystemComponent::UpdateTagMap_Internal(const FGameplayTagContainer&
 					ENetMode NetMode = GetWorld()->GetNetMode();
 					if (NetMode == NM_DedicatedServer || NetMode == NM_ListenServer)
 					{
-						GetMinimalReplicationTags_Mutable().RemoveTag(Tag);
-						GetReplicatedLooseTags_Mutable().RemoveTag(Tag);
+						FMinimalReplicationTagCountMap& TagMap = GetMinimalReplicationTags_Mutable();
+						if (TagMap.TagMap.Contains(Tag))
+						{
+							TagMap.RemoveTag(Tag);
+						}
+
+						TagMap = GetReplicatedLooseTags_Mutable();
+						if (TagMap.TagMap.Contains(Tag))
+						{
+							TagMap.RemoveTag(Tag);
+						}
 					}
 				}
 			}
diff --git a/Engine/Plugins/Runtime/GameplayAbilities/Source/GameplayAbilities/Public/AbilitySystemComponent.h b/Engine/Plugins/Runtime/GameplayAbilities/Source/GameplayAbilities/Public/AbilitySystemComponent.h
--- a/Engine/Plugins/Runtime/GameplayAbilities/Source/GameplayAbilities/Public/AbilitySystemComponent.h
+++ b/Engine/Plugins/Runtime/GameplayAbilities/Source/GameplayAbilities/Public/AbilitySystemComponent.h
@@ -615,28 +615,14 @@ class UAbilitySystemComponent : public UGameplayTasksComponent, public IGameplay
 	/** Forcibly sets the number of instances of a given tag */
 	inline void SetTagMapCount(const FGameplayTag& Tag, int32 NewCount, EGameplayTagReplicationState TagRepState = EGameplayTagReplicationState::None)
 	{
-		// Iris replication doesn't work the same so we're falling back to using the Minimal and replicated loose tags until it's fixed.
-		if (OwnerActor && OwnerActor->HasAuthority())
-		{
-			if (TagRepState > EGameplayTagReplicationState::TagOnly)
-			{
-				GetReplicatedLooseTags_Mutable().SetTagCount(Tag, NewCount);
-			}
-		}
-		
-		if (GameplayTagCountContainer.SetTagCount(Tag, NewCount, TagRepState))
-		{
-			OnTagUpdated(Tag, NewCount > 0);
-		}
+		const int32 CurrentCount = GameplayTagCountContainer.GetTagCount(Tag);
+		UpdateTagMapSingle_Internal(Tag, NewCount - CurrentCount, TagRepState);
 	}
 
 	/** Update the number of instances of a given tag and calls callback */
 	inline void UpdateTagMap(const FGameplayTag& BaseTag, int32 CountDelta, EGameplayTagReplicationState TagRepState)
 	{
-		if (GameplayTagCountContainer.UpdateTagCount(BaseTag, CountDelta, TagRepState))
-		{
-			OnTagUpdated(BaseTag, CountDelta > 0);
-		}
+		UpdateTagMapSingle_Internal(BaseTag, CountDelta, TagRepState);
 	}
 
 	/** Update the number of instances of a given tag and calls callback */
@@ -1905,6 +1891,7 @@ protected:
 	UE_API const UAttributeSet*	GetAttributeSubobjectChecked(const TSubclassOf<UAttributeSet> AttributeClass) const;
 	UE_API const UAttributeSet*	GetOrCreateAttributeSubobject(TSubclassOf<UAttributeSet> AttributeClass);
 
+	UE_API void UpdateTagMapSingle_Internal(const FGameplayTag& Tag, int32 CountDelta, EGameplayTagReplicationState TagRepState);
 	UE_API void UpdateTagMap_Internal(const FGameplayTagContainer& Container, int32 CountDelta, EGameplayTagReplicationState TagRepState);
 
 	friend struct FActiveGameplayEffect;
