TagCountContainerNetSerializer added for Iris.
Updates to GE prediction caught up and rejected logic so that predicted tags work correctly.

#FailedValidations Preflight

diff --git a/Engine/Plugins/Runtime/GameplayAbilities/Source/GameplayAbilities/Private/AbilitySystemComponent.cpp b/Engine/Plugins/Runtime/GameplayAbilities/Source/GameplayAbilities/Private/AbilitySystemComponent.cpp
--- a/Engine/Plugins/Runtime/GameplayAbilities/Source/GameplayAbilities/Private/AbilitySystemComponent.cpp
+++ b/Engine/Plugins/Runtime/GameplayAbilities/Source/GameplayAbilities/Private/AbilitySystemComponent.cpp
@@ -639,27 +639,6 @@ FOnGivenActiveGameplayEffectRemoved& UAbilitySystemComponent::OnAnyGameplayEffec
 
 void UAbilitySystemComponent::UpdateTagMapSingle_Internal(const FGameplayTag& Tag, const int32 CountDelta, EGameplayTagReplicationState TagRepState)
 {
-	if (UE::Net::ShouldUseIrisReplication())
-	{
-		// Iris replication doesn't work the same so we're falling back to using the Minimal and replicated loose tags until it's fixed.
-		ENetMode NetMode = GetWorld()->GetNetMode();
-		if (NetMode == NM_DedicatedServer || NetMode == NM_ListenServer)
-		{
-			if (TagRepState == EGameplayTagReplicationState::TagOnly)
-			{
-				if (CountDelta > 0)
-				{
-					GetMinimalReplicationTags_Mutable().AddTag(Tag);
-				}
-			}
-			else if (CountDelta < 0 || TagRepState > EGameplayTagReplicationState::TagOnly)
-			{
-				int32& TagCount = GetReplicatedLooseTags_Mutable().TagMap.FindOrAdd(Tag);
-				TagCount += CountDelta;
-			}
-		}
-	}
-	
 	// For removal, reorder calls so that FillParentTags is only called once
 	if (CountDelta > 0)
 	{
@@ -675,56 +654,12 @@ void UAbilitySystemComponent::UpdateTagMapSingle_Internal(const FGameplayTag& Ta
 		if (GameplayTagCountContainer.UpdateTagCount(Tag, CountDelta, TagRepState))
 		{
 			OnTagUpdated(Tag, false);
-			if (UE::Net::ShouldUseIrisReplication())
-			{
-				// Iris replication doesn't work the same so we're falling back to using the Minimal and replicated loose tags until it's fixed.
-				ENetMode NetMode = GetWorld()->GetNetMode();
-				if (NetMode == NM_DedicatedServer || NetMode == NM_ListenServer)
-				{
-					FMinimalReplicationTagCountMap& TagMap = GetMinimalReplicationTags_Mutable();
-					if (TagMap.TagMap.Contains(Tag))
-					{
-						TagMap.RemoveTag(Tag);
-					}
-
-					TagMap = GetReplicatedLooseTags_Mutable();
-					if (TagMap.TagMap.Contains(Tag))
-					{
-						TagMap.RemoveTag(Tag);
-					}
-				}
-			}
 		}
 	}
 }
 
 void UAbilitySystemComponent::UpdateTagMap_Internal(const FGameplayTagContainer& Container, const int32 CountDelta, EGameplayTagReplicationState TagRepState)
 {
-	if (UE::Net::ShouldUseIrisReplication())
-	{
-		// Iris replication doesn't work the same so we're falling back to using the Minimal and replicated loose tags until it's fixed.
-		ENetMode NetMode = GetWorld()->GetNetMode();
-		if (NetMode == NM_DedicatedServer || NetMode == NM_ListenServer)
-		{
-			if (TagRepState == EGameplayTagReplicationState::TagOnly)
-			{
-				if (CountDelta > 0)
-				{
-					GetMinimalReplicationTags_Mutable().AddTags(Container);
-				}
-			}
-			else if (CountDelta < 0 || TagRepState > EGameplayTagReplicationState::TagOnly)
-			{
-				for (auto TagIt = Container.CreateConstIterator(); TagIt; ++TagIt)
-				{
-					const FGameplayTag& Tag = *TagIt;
-					int32& TagCount = GetReplicatedLooseTags_Mutable().TagMap.FindOrAdd(Tag);
-					TagCount += CountDelta;
-				}
-			}
-		}
-	}
-	
 	// For removal, reorder calls so that FillParentTags is only called once
 	if (CountDelta > 0)
 	{
@@ -750,25 +685,6 @@ void UAbilitySystemComponent::UpdateTagMap_Internal(const FGameplayTagContainer&
 			if (GameplayTagCountContainer.UpdateTagCount_DeferredParentRemoval(Tag, CountDelta, DeferredTagChangeDelegates))
 			{
 				RemovedTags.Add(Tag);
-				if (UE::Net::ShouldUseIrisReplication())
-				{
-					// Iris replication doesn't work the same so we're falling back to using the Minimal and replicated loose tags until it's fixed.
-					ENetMode NetMode = GetWorld()->GetNetMode();
-					if (NetMode == NM_DedicatedServer || NetMode == NM_ListenServer)
-					{
-						FMinimalReplicationTagCountMap& TagMap = GetMinimalReplicationTags_Mutable();
-						if (TagMap.TagMap.Contains(Tag))
-						{
-							TagMap.RemoveTag(Tag);
-						}
-
-						TagMap = GetReplicatedLooseTags_Mutable();
-						if (TagMap.TagMap.Contains(Tag))
-						{
-							TagMap.RemoveTag(Tag);
-						}
-					}
-				}
 			}
 		}
 
@@ -1175,6 +1091,17 @@ void UAbilitySystemComponent::RemoveActiveGameplayEffect_AllowClientRemoval(FAct
 	ActiveGameplayEffects.RemoveActiveGameplayEffect(Handle, StacksToRemove);
 }
 
+void UAbilitySystemComponent::OnRejectedActiveGameplayEffect(FActiveGameplayEffectHandle Handle, int32 StacksToRemove)
+{
+	constexpr bool bPredictionRejected = true;
+	ActiveGameplayEffects.RemoveActiveGameplayEffect(Handle, StacksToRemove, bPredictionRejected);
+}
+
+void UAbilitySystemComponent::OnCaughtUpActiveGameplayEffect(FActiveGameplayEffectHandle Handle, int32 StacksToRemove)
+{
+	ActiveGameplayEffects.RemoveActiveGameplayEffect(Handle, StacksToRemove);
+}
+
 void UAbilitySystemComponent::OnPredictiveGameplayEffectStackCaughtUp(FActiveGameplayEffectHandle Handle)
 {
 	ActiveGameplayEffects.OnPredictiveGameplayEffectStackCaughtUp(Handle);
diff --git a/Engine/Plugins/Runtime/GameplayAbilities/Source/GameplayAbilities/Private/GameplayEffect.cpp b/Engine/Plugins/Runtime/GameplayAbilities/Source/GameplayAbilities/Private/GameplayEffect.cpp
--- a/Engine/Plugins/Runtime/GameplayAbilities/Source/GameplayAbilities/Private/GameplayEffect.cpp
+++ b/Engine/Plugins/Runtime/GameplayAbilities/Source/GameplayAbilities/Private/GameplayEffect.cpp
@@ -2703,6 +2703,8 @@ FActiveGameplayEffect& FActiveGameplayEffect::operator=(FActiveGameplayEffect&&
 	{
 		MostRecentArrayReplicationKey = Other.MostRecentArrayReplicationKey;
 	}
+	
+	bPostPredictObject = Other.bPostPredictObject;
 
 	return *this;
 }
@@ -2740,6 +2742,8 @@ FActiveGameplayEffect& FActiveGameplayEffect::operator=(const FActiveGameplayEff
 	{
 		MostRecentArrayReplicationKey = Other.MostRecentArrayReplicationKey;
 	}
+	
+	bPostPredictObject = Other.bPostPredictObject;
 
 	return *this;
 }
@@ -2812,11 +2816,11 @@ void FActiveGameplayEffect::PostReplicatedAdd(const struct FActiveGameplayEffect
 	}
 
 	bool ShouldInvokeGameplayCueEvents = true;
-	if (PredictionKey.IsLocalClientKey())
+	bPostPredictObject = PredictionKey.IsLocalClientKey();
+	if (bPostPredictObject)
 	{
 		// PredictionKey will only be valid on the client that predicted it. So if this has a valid PredictionKey, we can assume we already predicted it and shouldn't invoke GameplayCues.
 		// We may need to do more bookkeeping here in the future. Possibly give the predicted gameplayeffect a chance to pass something off to the new replicated gameplay effect.
-		
 		if (InArray.HasPredictedEffectWithPredictedKey(PredictionKey))
 		{
 			ShouldInvokeGameplayCueEvents = false;
@@ -4439,7 +4443,8 @@ PRAGMA_ENABLE_DEPRECATION_WARNINGS
 		{
 			// Once replicated state has caught up to this prediction key, we must remove this gameplay effect.
 			constexpr int32 RemoveAllStacks = -1;
-			InPredictionKey.NewRejectOrCaughtUpDelegate(FPredictionKeyEvent::CreateUObject(Owner, &UAbilitySystemComponent::RemoveActiveGameplayEffect_AllowClientRemoval, AppliedActiveGE->Handle, RemoveAllStacks));
+			InPredictionKey.NewCaughtUpDelegate().BindUObject(Owner, &UAbilitySystemComponent::OnCaughtUpActiveGameplayEffect, AppliedActiveGE->Handle, RemoveAllStacks);
+			InPredictionKey.NewRejectedDelegate().BindUObject(Owner, &UAbilitySystemComponent::OnRejectedActiveGameplayEffect, AppliedActiveGE->Handle, RemoveAllStacks);
 		}
 		else // We have an existing Stackable GE
 		{
@@ -4465,7 +4470,7 @@ PRAGMA_ENABLE_DEPRECATION_WARNINGS
 		const bool bInvokeGameplayCueEvents = (Spec.Def->DurationPolicy != EGameplayEffectDurationType::Instant);
 		InternalOnActiveGameplayEffectAdded(*AppliedActiveGE, bInvokeGameplayCueEvents);
 	}
-
+	
 	return AppliedActiveGE;
 }
 
@@ -4507,6 +4512,12 @@ void FActiveGameplayEffectsContainer::AddActiveGameplayEffectGrantedTagsAndModif
 		return;
 	}
 
+	if (!Owner)
+	{
+		UE_LOG(LogGameplayEffects, Error, TEXT("AddActiveGameplayEffectGrantedTagsAndModifiers called with null Owner!"));
+		return;
+	}
+
 	GAMEPLAYEFFECT_SCOPE_LOCK();
 
 	// Register this ActiveGameplayEffects modifiers with our Attribute Aggregators
@@ -4524,7 +4535,7 @@ void FActiveGameplayEffectsContainer::AddActiveGameplayEffectGrantedTagsAndModif
 			const FGameplayModifierInfo &ModInfo = Effect.Spec.Def->Modifiers[ModIdx];
 
 			// skip over any modifiers for attributes that we don't have
-			if (!Owner || Owner->HasAttributeSetForAttribute(ModInfo.Attribute) == false)
+			if (Owner->HasAttributeSetForAttribute(ModInfo.Attribute) == false)
 			{
 				continue;
 			}
@@ -4543,7 +4554,7 @@ void FActiveGameplayEffectsContainer::AddActiveGameplayEffectGrantedTagsAndModif
 	}
 	else
 	{
-		if (Effect.Spec.Def->PeriodicInhibitionPolicy != EGameplayEffectPeriodInhibitionRemovedPolicy::NeverReset && Owner && Owner->IsOwnerActorAuthoritative())
+		if (Effect.Spec.Def->PeriodicInhibitionPolicy != EGameplayEffectPeriodInhibitionRemovedPolicy::NeverReset && Owner->IsOwnerActorAuthoritative())
 		{
 			FTimerManager& TimerManager = Owner->GetWorld()->GetTimerManager();
 			FTimerDelegate Delegate = FTimerDelegate::CreateUObject(Owner, &UAbilitySystemComponent::ExecutePeriodicEffect, Effect.Handle);
@@ -4558,11 +4569,14 @@ void FActiveGameplayEffectsContainer::AddActiveGameplayEffectGrantedTagsAndModif
 		}
 	}
 
-
-	// Update our owner with the tags this GameplayEffect grants them
-	Owner->UpdateTagMap(Effect.Spec.Def->GetGrantedTags(), 1, ShouldUseMinimalReplication() ? EGameplayTagReplicationState::TagOnly : EGameplayTagReplicationState::None);
-	Owner->UpdateTagMap(Effect.Spec.DynamicGrantedTags, 1, ShouldUseMinimalReplication() ? EGameplayTagReplicationState::TagOnly : EGameplayTagReplicationState::None);
-
+	// Local application should be responsible for applying tags
+	if (!Effect.bPostPredictObject)
+	{
+		// Update our owner with the tags this GameplayEffect grants them
+		Owner->UpdateTagMap(Effect.Spec.Def->GetGrantedTags(), 1, EGameplayTagReplicationState::CountToOwner);
+		Owner->UpdateTagMap(Effect.Spec.DynamicGrantedTags, 1, EGameplayTagReplicationState::CountToOwner);
+	}
+	
 	// Update our owner with the blocked ability tags this GameplayEffect adds to them
 	Owner->BlockAbilitiesWithTags(Effect.Spec.Def->GetBlockedAbilityTags());
 
@@ -4624,8 +4638,8 @@ PRAGMA_ENABLE_DEPRECATION_WARNINGS
 	Owner->OnActiveGameplayEffectAddedDelegateToSelf.Broadcast(Owner, Effect.Spec, Effect.Handle);
 }
 
-/** Called on server to remove a GameplayEffect */
-bool FActiveGameplayEffectsContainer::RemoveActiveGameplayEffect(FActiveGameplayEffectHandle Handle, int32 StacksToRemove)
+/** Called on server to remove a GameplayEffect and on owning client for Post-Prediction cleanup */
+bool FActiveGameplayEffectsContainer::RemoveActiveGameplayEffect(FActiveGameplayEffectHandle Handle, int32 StacksToRemove, bool bPredictionRejected)
 {
 	// Iterating through manually since this is a removal operation and we need to pass the index into InternalRemoveActiveGameplayEffect
 	int32 NumGameplayEffects = GetNumGameplayEffects();
@@ -4634,7 +4648,7 @@ bool FActiveGameplayEffectsContainer::RemoveActiveGameplayEffect(FActiveGameplay
 		FActiveGameplayEffect& Effect = *GetActiveGameplayEffect(ActiveGEIdx);
 		if (Effect.Handle == Handle && Effect.IsPendingRemove == false)
 		{
-			InternalRemoveActiveGameplayEffect(ActiveGEIdx, StacksToRemove, true);
+			InternalRemoveActiveGameplayEffect(ActiveGEIdx, StacksToRemove, true, bPredictionRejected);
 			return true;
 		}
 	}
@@ -4678,7 +4692,7 @@ void FActiveGameplayEffectsContainer::InternalExecutePeriodicGameplayEffect(FAct
 }
 
 /** Called by server to actually remove a GameplayEffect */
-bool FActiveGameplayEffectsContainer::InternalRemoveActiveGameplayEffect(int32 Idx, int32 StacksToRemove, bool bPrematureRemoval)
+bool FActiveGameplayEffectsContainer::InternalRemoveActiveGameplayEffect(int32 Idx, int32 StacksToRemove, bool bPrematureRemoval, bool bPredictionRejected)
 {
 	SCOPE_CYCLE_COUNTER(STAT_RemoveActiveGameplayEffect);
 
@@ -4711,6 +4725,7 @@ bool FActiveGameplayEffectsContainer::InternalRemoveActiveGameplayEffect(int32 I
 		GameplayEffectRemovalInfo.ActiveEffect = &Effect;
 		GameplayEffectRemovalInfo.StackCount = Effect.Spec.GetStackCount();
 		GameplayEffectRemovalInfo.bPrematureRemoval = bPrematureRemoval;
+		GameplayEffectRemovalInfo.bPredictionRejected = bPredictionRejected;
 		GameplayEffectRemovalInfo.EffectContext = Effect.Spec.GetEffectContext();
 
 		if (StacksToRemove > 0 && Effect.Spec.GetStackCount() > StacksToRemove)
@@ -4816,7 +4831,7 @@ void FActiveGameplayEffectsContainer::InternalOnActiveGameplayEffectRemoved(FAct
 		// Only Need to update tags and modifiers if the gameplay effect is active.
 		if (!Effect.bIsInhibited)
 		{
-			RemoveActiveGameplayEffectGrantedTagsAndModifiers(Effect, bInvokeGameplayCueEvents);
+			RemoveActiveGameplayEffectGrantedTagsAndModifiers(Effect, bInvokeGameplayCueEvents, GameplayEffectRemovalInfo.bPredictionRejected);
 		}
 
 		RemoveCustomMagnitudeExternalDependencies(Effect);
@@ -4831,7 +4846,7 @@ void FActiveGameplayEffectsContainer::InternalOnActiveGameplayEffectRemoved(FAct
 	OnActiveGameplayEffectRemovedDelegate.Broadcast(Effect);
 }
 
-void FActiveGameplayEffectsContainer::RemoveActiveGameplayEffectGrantedTagsAndModifiers(const FActiveGameplayEffect& Effect, bool bInvokeGameplayCueEvents)
+void FActiveGameplayEffectsContainer::RemoveActiveGameplayEffectGrantedTagsAndModifiers(const FActiveGameplayEffect& Effect, bool bInvokeGameplayCueEvents, bool bPredictionRejected)
 {
 	// Update AttributeAggregators: remove mods from this ActiveGE Handle
 	if (Effect.Spec.GetPeriod() <= UGameplayEffect::NO_PERIOD)
@@ -4848,9 +4863,13 @@ void FActiveGameplayEffectsContainer::RemoveActiveGameplayEffectGrantedTagsAndMo
 		}
 	}
 
-	// Update gameplaytag count and broadcast delegate if we are at 0
-	Owner->UpdateTagMap(Effect.Spec.Def->GetGrantedTags(), -1, ShouldUseMinimalReplication() ? EGameplayTagReplicationState::TagOnly : EGameplayTagReplicationState::None);
-	Owner->UpdateTagMap(Effect.Spec.DynamicGrantedTags, -1, ShouldUseMinimalReplication() ? EGameplayTagReplicationState::TagOnly : EGameplayTagReplicationState::None);
+	// We're no longer predicting for remove unless it was rejected, so let the server replication handle the tags.
+	if (IsNetAuthority() || bPredictionRejected)
+	{
+		// Update gameplaytag count and broadcast delegate if we are at 0
+		Owner->UpdateTagMap(Effect.Spec.Def->GetGrantedTags(), -1, EGameplayTagReplicationState::None);
+		Owner->UpdateTagMap(Effect.Spec.DynamicGrantedTags, -1, EGameplayTagReplicationState::None);
+	}
 
 	// Update our owner with the blocked ability tags this GameplayEffect adds to them
 	Owner->UnBlockAbilitiesWithTags(Effect.Spec.Def->GetBlockedAbilityTags());
diff --git a/Engine/Plugins/Runtime/GameplayAbilities/Source/GameplayAbilities/Private/GameplayEffectTypes.cpp b/Engine/Plugins/Runtime/GameplayAbilities/Source/GameplayAbilities/Private/GameplayEffectTypes.cpp
--- a/Engine/Plugins/Runtime/GameplayAbilities/Source/GameplayAbilities/Private/GameplayEffectTypes.cpp
+++ b/Engine/Plugins/Runtime/GameplayAbilities/Source/GameplayAbilities/Private/GameplayEffectTypes.cpp
@@ -18,6 +18,10 @@
 #endif
 
 #include "Iris/Serialization/NetBitStreamUtil.h"
+#include "Iris/Serialization/NetSerializationContext.h"
+#include "Iris/Serialization/NetSerializerDelegates.h"
+#include "Iris/Serialization/PolymorphicNetSerializerImpl.h"
+#include "Iris/ReplicationState/PropertyNetSerializerInfoRegistry.h"
 
 #include UE_INLINE_GENERATED_CPP_BY_NAME(GameplayEffectTypes)
 
@@ -650,7 +654,7 @@ bool FGameplayTagCountContainer::NetDeltaSerialize(FNetDeltaSerializeInfo & Delt
 			}
 		}
 	}
-
+	
 	if (DeltaParams.Reader)
 	{
 		FBitReader& Reader = *DeltaParams.Reader;
@@ -664,9 +668,10 @@ bool FGameplayTagCountContainer::NetDeltaSerialize(FNetDeltaSerializeInfo & Delt
 			for (; NumRemovals > 0; NumRemovals--)
 			{
 				RemovedItem.Tag.NetSerialize(Reader, DeltaParams.Map, bTagSuccess);
-				if (GameplayTagCountMap.Contains(RemovedItem.Tag))
+				int32* Count = GameplayTagCountMap.Find(RemovedItem.Tag);
+				if (Count && *Count > 0)
 				{
-					UpdateTagMap_Internal(RemovedItem.Tag, -GameplayTagCountMap[RemovedItem.Tag], EGameplayTagReplicationState::None);
+					UpdateTagMap_Internal(RemovedItem.Tag, -(*Count), EGameplayTagReplicationState::None);
 				}
 			}
 		}
@@ -683,16 +688,12 @@ bool FGameplayTagCountContainer::NetDeltaSerialize(FNetDeltaSerializeInfo & Delt
 				NewCount = 1;
 				UpdatedItem.Tag.NetSerialize(Reader, DeltaParams.Map, bTagSuccess);
 				Reader.SerializeBits(&bPackCount, 1);
-				int32& Count = GameplayTagCountMap.FindOrAdd(UpdatedItem.Tag);
 				if (bPackCount)
 				{
 					Reader.SerializeIntPacked(NewCount);
 				}
 
-				if (NewCount != Count)
-				{
-					UpdateTagMap_Internal(UpdatedItem.Tag, NewCount - Count, EGameplayTagReplicationState::None);
-				}
+				SetTagCount(UpdatedItem.Tag, NewCount, EGameplayTagReplicationState::None);
 			}
 		}
 	}
@@ -766,10 +767,21 @@ bool FGameplayTagCountContainer::UpdateExplicitTags(const FGameplayTag& Tag, con
 				ABILITY_LOG(Warning, TEXT("Trying to change existing tag \"%s\" replication state from \"%s\" to \"%s\""), *Tag.ToString(), *UEnum::GetDisplayValueAsText(TagCountItem.ReplicationState).ToString(), *UEnum::GetDisplayValueAsText(TagRepState).ToString());
 			}
 
-			// Some sort of authority check before updating this value?
+			// Promote our rep state to the higher level
 			TagCountItem.ReplicationState = TagRepState > TagCountItem.ReplicationState ? TagRepState : TagCountItem.ReplicationState;
+			TagRepState = TagCountItem.ReplicationState;
+		}
+
+		// make sure clients only have the 1 tag if it's tag Only.
+		if (Owner && Owner->GetOwnerRole() != ROLE_Authority && CountDelta > 0  && TagCountItem.ReplicationState == EGameplayTagReplicationState::TagOnly)
+		{
+			TagCountItem.Count = 1;
+		}
+		else
+		{
+			TagCountItem.Count += CountDelta;
 		}
-		TagCountItem.Count += CountDelta;
+		
 		if (TagCountItem.Count <= 0)
 		{
 			Items.RemoveAtSwap(TagCountIndex);
@@ -790,7 +802,7 @@ bool FGameplayTagCountContainer::UpdateExplicitTags(const FGameplayTag& Tag, con
 	return true;
 }
 
-bool FGameplayTagCountContainer::GatherTagChangeDelegates(const FGameplayTag& Tag, const int32 CountDelta, TArray<FDeferredTagChangeDelegate>& TagChangeDelegates)
+bool FGameplayTagCountContainer::GatherTagChangeDelegates(const FGameplayTag& Tag, const int32 CountDelta, TArray<FDeferredTagChangeDelegate>& TagChangeDelegates, EGameplayTagReplicationState TagRepState)
 {
 	// Check if change delegates are required to fire for the tag or any of its parents based on the count change
 	FGameplayTagContainer TagAndParentsContainer = Tag.GetGameplayTagParents();
@@ -803,10 +815,13 @@ bool FGameplayTagCountContainer::GatherTagChangeDelegates(const FGameplayTag& Ta
 		int32& TagCountRef = GameplayTagCountMap.FindOrAdd(CurTag);
 
 		const int32 OldCount = TagCountRef;
-
+		
+		const bool ClampTagOnlyCount = Owner && Owner->GetOwnerRole() != ROLE_Authority && CountDelta > 0  && TagRepState == EGameplayTagReplicationState::TagOnly;
+		
 		// Apply the delta to the count in the map
 		int32 NewTagCount = FMath::Max(OldCount + CountDelta, 0);
-		TagCountRef = NewTagCount;
+		
+		TagCountRef = ClampTagOnlyCount ? FMath::Min(NewTagCount, 1) : NewTagCount;
 
 		// If a significant change (new addition or total removal) occurred, trigger related delegates
 		const bool SignificantChange = (OldCount == 0 || NewTagCount == 0);
@@ -814,9 +829,9 @@ bool FGameplayTagCountContainer::GatherTagChangeDelegates(const FGameplayTag& Ta
 		if (SignificantChange)
 		{
 			TagChangeDelegates.AddDefaulted();
-			TagChangeDelegates.Last().BindLambda([Delegate = OnAnyTagChangeDelegate, CurTag, NewTagCount]()
+			TagChangeDelegates.Last().BindLambda([Delegate = OnAnyTagChangeDelegate, CurTag, TagCountRef]()
 			{
-				Delegate.Broadcast(CurTag, NewTagCount);
+				Delegate.Broadcast(CurTag, TagCountRef);
 			});
 		}
 
@@ -824,17 +839,17 @@ bool FGameplayTagCountContainer::GatherTagChangeDelegates(const FGameplayTag& Ta
 		if (DelegateInfo)
 		{
 			TagChangeDelegates.AddDefaulted();
-			TagChangeDelegates.Last().BindLambda([Delegate = DelegateInfo->OnAnyChange, CurTag, NewTagCount]()
+			TagChangeDelegates.Last().BindLambda([Delegate = DelegateInfo->OnAnyChange, CurTag, TagCountRef]()
 			{
-				Delegate.Broadcast(CurTag, NewTagCount);
+				Delegate.Broadcast(CurTag, TagCountRef);
 			});
 
 			if (SignificantChange)
 			{
 				TagChangeDelegates.AddDefaulted();
-				TagChangeDelegates.Last().BindLambda([Delegate = DelegateInfo->OnNewOrRemove, CurTag, NewTagCount]()
+				TagChangeDelegates.Last().BindLambda([Delegate = DelegateInfo->OnNewOrRemove, CurTag, TagCountRef]()
 				{
-					Delegate.Broadcast(CurTag, NewTagCount);
+					Delegate.Broadcast(CurTag, TagCountRef);
 				});
 			}
 		}
@@ -849,6 +864,11 @@ bool FGameplayTagCountContainer::UpdateTagMap_Internal(const FGameplayTag& Tag,
 	{
 		return false;
 	}
+	
+	if (!ensureMsgf(CountDelta != 0, TEXT("FGameplayTagCountContainer Modifying tag count by 0 for tag %s"), *Tag.ToString()))
+	{
+		return false;
+	}
 
 	if (!UpdateExplicitTags(Tag, CountDelta, false, TagRepState))
 	{
@@ -856,7 +876,7 @@ bool FGameplayTagCountContainer::UpdateTagMap_Internal(const FGameplayTag& Tag,
 	}
 
 	TArray<FDeferredTagChangeDelegate> DeferredTagChangeDelegates;
-	bool bSignificantChange = GatherTagChangeDelegates(Tag, CountDelta, DeferredTagChangeDelegates);
+	bool bSignificantChange = GatherTagChangeDelegates(Tag, CountDelta, DeferredTagChangeDelegates, TagRepState);
 	for (FDeferredTagChangeDelegate& Delegate : DeferredTagChangeDelegates)
 	{
 		Delegate.Execute();
@@ -872,7 +892,7 @@ bool FGameplayTagCountContainer::UpdateTagMapDeferredParentRemoval_Internal(cons
 		return false;
 	}
 
-	return GatherTagChangeDelegates(Tag, CountDelta, DeferredTagChangeDelegates);
+	return GatherTagChangeDelegates(Tag, CountDelta, DeferredTagChangeDelegates, TagRepState);
 }
 
 FGameplayTagBlueprintPropertyMap::FGameplayTagBlueprintPropertyMap()
diff --git a/Engine/Plugins/Runtime/GameplayAbilities/Source/GameplayAbilities/Public/AbilitySystemComponent.h b/Engine/Plugins/Runtime/GameplayAbilities/Source/GameplayAbilities/Public/AbilitySystemComponent.h
--- a/Engine/Plugins/Runtime/GameplayAbilities/Source/GameplayAbilities/Public/AbilitySystemComponent.h
+++ b/Engine/Plugins/Runtime/GameplayAbilities/Source/GameplayAbilities/Public/AbilitySystemComponent.h
@@ -1847,6 +1847,10 @@ protected:
 	/** A version of RemoveActiveGameplayEffect that allows us to remove it even when we don't have authority */
 	UE_API void RemoveActiveGameplayEffect_AllowClientRemoval(FActiveGameplayEffectHandle Handle, int32 StacksToRemove = -1);
 
+	UE_API void OnCaughtUpActiveGameplayEffect(FActiveGameplayEffectHandle Handle, int32 StacksToRemove = -1);
+
+	UE_API void OnRejectedActiveGameplayEffect(FActiveGameplayEffectHandle Handle, int32 StacksToRemove = -1);
+
 	UE_API void OnPredictiveGameplayEffectStackCaughtUp(FActiveGameplayEffectHandle Handle);
 
 	/** Contains all of the gameplay effects that are currently active on this component */
diff --git a/Engine/Plugins/Runtime/GameplayAbilities/Source/GameplayAbilities/Public/GameplayEffect.h b/Engine/Plugins/Runtime/GameplayAbilities/Source/GameplayAbilities/Public/GameplayEffect.h
--- a/Engine/Plugins/Runtime/GameplayAbilities/Source/GameplayAbilities/Public/GameplayEffect.h
+++ b/Engine/Plugins/Runtime/GameplayAbilities/Source/GameplayAbilities/Public/GameplayEffect.h
@@ -1445,6 +1445,9 @@ struct FActiveGameplayEffect : public FFastArraySerializerItem
 	
 	/** All the bindable events for this active effect (bundled to allow easier non-const access to these events via the ASC) */
 	FActiveGameplayEffectEvents EventSet;
+
+	/** Signifies an active effect that was predicted and has now been replicated back to the client */
+	bool bPostPredictObject = false;
 };
 
 DECLARE_DELEGATE_RetVal_OneParam(bool, FActiveGameplayEffectQueryCustomMatch, const FActiveGameplayEffect&);
@@ -1694,7 +1697,7 @@ struct FActiveGameplayEffectsContainer : public FFastArraySerializer
 	
 	UE_API void ExecutePeriodicGameplayEffect(FActiveGameplayEffectHandle Handle);	// This should not be outward facing to the skill system API, should only be called by the owning AbilitySystemComponent
 
-	UE_API bool RemoveActiveGameplayEffect(FActiveGameplayEffectHandle Handle, int32 StacksToRemove);
+	UE_API bool RemoveActiveGameplayEffect(FActiveGameplayEffectHandle Handle, int32 StacksToRemove, bool bPredictionRejected = false);
 
 	UE_API void GetGameplayEffectStartTimeAndDuration(FActiveGameplayEffectHandle Handle, float& EffectStartTime, float& EffectDuration) const;
 
@@ -1924,13 +1927,13 @@ private:
 	UE_API void InternalExecutePeriodicGameplayEffect(FActiveGameplayEffect& ActiveEffect);
 
 	/** Called internally to actually remove a GameplayEffect or to reduce its StackCount. Returns true if we resized our internal GameplayEffect array. */
-	UE_API bool InternalRemoveActiveGameplayEffect(int32 Idx, int32 StacksToRemove, bool bPrematureRemoval);
+	UE_API bool InternalRemoveActiveGameplayEffect(int32 Idx, int32 StacksToRemove, bool bPrematureRemoval, bool bPredictionRejected = false);
 	
 	/** Called both in server side creation and replication creation/deletion */
 	UE_API void InternalOnActiveGameplayEffectAdded(FActiveGameplayEffect& Effect, bool bInvokeGameplayCueEvents);
 	UE_API void InternalOnActiveGameplayEffectRemoved(FActiveGameplayEffect& Effect, bool bInvokeGameplayCueEvents, const FGameplayEffectRemovalInfo& GameplayEffectRemovalInfo);
 
-	UE_API void RemoveActiveGameplayEffectGrantedTagsAndModifiers(const FActiveGameplayEffect& Effect, bool bInvokeGameplayCueEvents);
+	UE_API void RemoveActiveGameplayEffectGrantedTagsAndModifiers(const FActiveGameplayEffect& Effect, bool bInvokeGameplayCueEvents, bool bPredictionRejected = false);
 	UE_API void AddActiveGameplayEffectGrantedTagsAndModifiers(FActiveGameplayEffect& Effect, bool bInvokeGameplayCueEvents);
 
 	/** Updates tag dependency map when a GameplayEffect is removed */
diff --git a/Engine/Plugins/Runtime/GameplayAbilities/Source/GameplayAbilities/Public/GameplayEffectTypes.h b/Engine/Plugins/Runtime/GameplayAbilities/Source/GameplayAbilities/Public/GameplayEffectTypes.h
--- a/Engine/Plugins/Runtime/GameplayAbilities/Source/GameplayAbilities/Public/GameplayEffectTypes.h
+++ b/Engine/Plugins/Runtime/GameplayAbilities/Source/GameplayAbilities/Public/GameplayEffectTypes.h
@@ -36,8 +36,10 @@ struct FGameplayEffectModCallbackData;
 struct FGameplayEffectSpec;
 struct FHitResult;
 struct FMinimalReplicationTagCountMapForNetSerializer;
+struct FNetGameplayTagCountContainerStateForNetSerialize;
 namespace UE::Net
 {
+	struct FGameplayTagCountContainerNetSerializer;
 	struct FGameplayEffectContextHandleAccessorForNetSerializer;
 	class FMinimalReplicationTagCountMapReplicationFragment;
 }
@@ -819,6 +821,10 @@ struct FGameplayEffectRemovalInfo
 	UPROPERTY(VisibleAnywhere, BlueprintReadOnly, Category="Removal")
 	bool bPrematureRemoval = false;
 
+	/** True when the effect is being removed due to prediction rejection */
+	UPROPERTY(VisibleAnywhere, BlueprintReadOnly, Category="Removal")
+	bool bPredictionRejected = false;
+
 	/** Number of Stacks this gameplay effect had before it was removed. */
 	UPROPERTY(VisibleAnywhere, BlueprintReadOnly, Category="Removal")
 	int32 StackCount = 0;
@@ -826,7 +832,7 @@ struct FGameplayEffectRemovalInfo
 	/** Actor this gameplay effect was targeting. */
 	UPROPERTY(VisibleAnywhere, BlueprintReadOnly, Category="Removal")
 	FGameplayEffectContextHandle EffectContext;
-
+	
 	/** The Effect being Removed */
 	const FActiveGameplayEffect* ActiveEffect = nullptr;
 };
@@ -1047,7 +1053,7 @@ enum class EGameplayTagReplicationState : uint8
 {
 	None,				// Tag will not replicate
 	TagOnly,			// Tag will replicate to all without count
-	CountToOwner,		// Tag will replicate to all and Count will replicate to owner
+	CountToOwner,		// Tag will replicate to all and Count will replicate to owner, Functions same as TagAndCountToAll with Iris
 	TagAndCountToAll	// Tag and Count will replicate to all
 };
 
@@ -1346,8 +1352,31 @@ struct FGameplayTagCountContainer
 
 	UE_API void SetOwner(UAbilitySystemComponent* ASC);
 
+	bool operator==(const FGameplayTagCountContainer& Other) const
+	{
+		if (Items.Num() != Other.Items.Num())
+		{
+			return false;
+		}
+
+		int32 Idx = INDEX_NONE;
+		for (const FGameplayTagCountItem& Item : Items)
+		{
+			Idx = Other.Items.Find(Item);
+			if (Idx == INDEX_NONE || Item.Count != Other.Items[Idx].Count)
+			{
+				return false;
+			}
+		}
+
+		return true;
+	}
+
 private:
 
+	friend UE::Net::FGameplayTagCountContainerNetSerializer;
+	friend FNetGameplayTagCountContainerStateForNetSerialize;
+
 	struct FDelegateInfo
 	{
 		FOnGameplayEffectTagCountChanged	OnNewOrRemove;
@@ -1383,7 +1412,7 @@ private:
 	UE_API bool UpdateExplicitTags(const FGameplayTag& Tag, int32 CountDelta, bool bDeferParentTagsOnRemove, EGameplayTagReplicationState TagRepState = EGameplayTagReplicationState::None);
 
 	/** Internal helper function to collect the delegates that need to be called when Tag has its count changed by CountDelta. */
-	UE_API bool GatherTagChangeDelegates(const FGameplayTag& Tag, int32 CountDelta, TArray<FDeferredTagChangeDelegate>& TagChangeDelegates);
+	UE_API bool GatherTagChangeDelegates(const FGameplayTag& Tag, int32 CountDelta, TArray<FDeferredTagChangeDelegate>& TagChangeDelegates, EGameplayTagReplicationState TagRepState);
 };
 
 template<>
diff --git a/Engine/Plugins/Runtime/GameplayAbilities/Source/GameplayAbilities/Public/Serialization/GameplayTagCountContainerNetSerializer.cpp b/Engine/Plugins/Runtime/GameplayAbilities/Source/GameplayAbilities/Public/Serialization/GameplayTagCountContainerNetSerializer.cpp
new file mode 100644
--- /dev/null
+++ b/Engine/Plugins/Runtime/GameplayAbilities/Source/GameplayAbilities/Public/Serialization/GameplayTagCountContainerNetSerializer.cpp
@@ -0,0 +1,230 @@
+// Copyright Epic Games, Inc. All Rights Reserved.
+
+#include "GameplayTagCountContainerNetSerializer.h"
+
+#include UE_INLINE_GENERATED_CPP_BY_NAME(GameplayTagCountContainerNetSerializer)
+
+#include "GameplayEffectTypes.h"
+#include "Engine/NetConnection.h"
+#include "Iris/ReplicationState/PropertyNetSerializerInfoRegistry.h"
+#include "Iris/ReplicationState/ReplicationStateDescriptorBuilder.h"
+#include "Iris/Serialization/NetSerializerDelegates.h"
+#include "Iris/Serialization/NetSerializers.h"
+#include "Iris/Serialization/NetSerializationContext.h"
+#include "Iris/Serialization/NetSerializer.h"
+#include "Iris/Serialization/PolymorphicNetSerializerImpl.h"
+
+void FNetGameplayTagCountContainerStateForNetSerialize::CopyFromTagCountContainer(const FGameplayTagCountContainer& Container)
+{
+	for (auto Item : Container.Items)
+	{
+		if (Item.ReplicationState == EGameplayTagReplicationState::None)
+		{
+			continue;
+		}
+		
+		TagStates.Push(FGameplayTagCountItem(Item.Tag, Item.ReplicationState == EGameplayTagReplicationState::TagOnly ? 1 : Item.Count, Item
+		.ReplicationState));
+	}
+}
+
+namespace UE::Net
+{
+
+struct FGameplayTagCountContainerNetSerializer
+{
+	// Version
+	static const uint32 Version = 0;
+
+	// Traits
+	static constexpr bool bIsForwardingSerializer = true;
+	static constexpr bool bHasDynamicState = true;
+	static constexpr bool bUseSerializerIsEqual = true;
+	
+public:
+	typedef FGameplayTagCountContainer SourceType;
+	typedef FNetGameplayTagCountContainerStateForNetSerialize SourceArrayType;
+	
+	struct FQuantizedType
+	{
+		alignas(16) uint8 Buffer[16];
+	};
+
+	typedef FQuantizedType QuantizedType;
+	
+	typedef FGameplayTagCountContainerNetSerializerConfig ConfigType;
+	static const ConfigType DefaultConfig;
+
+	// 
+	static void Serialize(FNetSerializationContext&, const FNetSerializeArgs& Args);
+	static void Deserialize(FNetSerializationContext&, const FNetDeserializeArgs& Args);
+
+	static void SerializeDelta(FNetSerializationContext&, const FNetSerializeDeltaArgs& Args);
+	static void DeserializeDelta(FNetSerializationContext&, const FNetDeserializeDeltaArgs& Args);
+
+	static void Quantize(FNetSerializationContext&, const FNetQuantizeArgs& Args);
+	static void Dequantize(FNetSerializationContext&, const FNetDequantizeArgs& Args);
+
+	static bool IsEqual(FNetSerializationContext&, const FNetIsEqualArgs& Args);
+	static bool Validate(FNetSerializationContext&, const FNetValidateArgs& Args);
+
+	static void CloneDynamicState(FNetSerializationContext&, const FNetCloneDynamicStateArgs&);
+	static void FreeDynamicState(FNetSerializationContext&, const FNetFreeDynamicStateArgs&);
+
+	static void CollectNetReferences(FNetSerializationContext&, const FNetCollectReferencesArgs&);
+
+	class FNetSerializerRegistryDelegates final : private UE::Net::FNetSerializerRegistryDelegates
+	{
+	public:
+		FNetSerializerRegistryDelegates();
+		virtual ~FNetSerializerRegistryDelegates();
+
+	private:
+		virtual void OnPreFreezeNetSerializerRegistry() override;
+		virtual void OnPostFreezeNetSerializerRegistry() override;
+	};
+
+	static FNetSerializerRegistryDelegates NetSerializerRegistryDelegates;
+	inline static FStructNetSerializerConfig StructNetSerializerConfig;
+	static const FNetSerializer* StructNetSerializer;
+};
+
+UE_NET_IMPLEMENT_SERIALIZER(FGameplayTagCountContainerNetSerializer);
+const FNetSerializer* FGameplayTagCountContainerNetSerializer::StructNetSerializer = &UE_NET_GET_SERIALIZER(FStructNetSerializer);
+	
+const FGameplayTagCountContainerNetSerializer::ConfigType FGameplayTagCountContainerNetSerializer::DefaultConfig;
+FGameplayTagCountContainerNetSerializer::FNetSerializerRegistryDelegates FGameplayTagCountContainerNetSerializer::NetSerializerRegistryDelegates;
+bool bIsPostFreezeCalled = false;
+
+void FGameplayTagCountContainerNetSerializer::Serialize(FNetSerializationContext& Context, const FNetSerializeArgs& Args)
+{
+	FNetSerializeArgs TempArgs = Args;
+	TempArgs.NetSerializerConfig = &StructNetSerializerConfig;
+	StructNetSerializer->Serialize(Context, TempArgs);
+}
+
+void FGameplayTagCountContainerNetSerializer::Deserialize(FNetSerializationContext& Context, const FNetDeserializeArgs& Args)
+{
+	FNetDeserializeArgs TempArgs = Args;
+	TempArgs.NetSerializerConfig = &StructNetSerializerConfig;
+	StructNetSerializer->Deserialize(Context, TempArgs);
+}
+	
+void FGameplayTagCountContainerNetSerializer::SerializeDelta(FNetSerializationContext& Context, const FNetSerializeDeltaArgs& Args)
+{
+	FNetSerializeDeltaArgs TempArgs = Args;
+	TempArgs.NetSerializerConfig = &StructNetSerializerConfig;
+	StructNetSerializer->SerializeDelta(Context, TempArgs);
+}
+
+void FGameplayTagCountContainerNetSerializer::DeserializeDelta(FNetSerializationContext& Context, const FNetDeserializeDeltaArgs& Args)
+{
+	FNetDeserializeDeltaArgs TempArgs = Args;
+	TempArgs.NetSerializerConfig = &StructNetSerializerConfig;
+	StructNetSerializer->DeserializeDelta(Context, TempArgs);
+}
+
+void FGameplayTagCountContainerNetSerializer::Quantize(FNetSerializationContext& Context, const FNetQuantizeArgs& Args)
+{
+	const SourceType& Source = *reinterpret_cast<const SourceType*>(Args.Source);
+
+	FNetGameplayTagCountContainerStateForNetSerialize SourceCopy;
+	SourceCopy.CopyFromTagCountContainer(Source);
+
+	FNetQuantizeArgs TempArgs = Args;
+	TempArgs.Source = NetSerializerValuePointer(&SourceCopy);
+	TempArgs.NetSerializerConfig = &StructNetSerializerConfig;
+	StructNetSerializer->Quantize(Context, TempArgs);
+}
+
+void FGameplayTagCountContainerNetSerializer::Dequantize(FNetSerializationContext& Context, const FNetDequantizeArgs& Args)
+{
+	SourceType& TargetValue = *reinterpret_cast<SourceType*>(Args.Target);
+	
+	FNetGameplayTagCountContainerStateForNetSerialize SourceCopy;
+
+	FNetDequantizeArgs TempArgs = Args;
+	TempArgs.Target = NetSerializerValuePointer(&SourceCopy);
+	TempArgs.NetSerializerConfig = &StructNetSerializerConfig;
+	StructNetSerializer->Dequantize(Context, TempArgs);
+
+	for (int32 Idx = 0; Idx < SourceCopy.TagStates.Num(); Idx++)
+	{
+		TargetValue.SetTagCount(SourceCopy.TagStates[Idx].Tag, SourceCopy.TagStates[Idx].Count);
+	}
+}
+
+bool FGameplayTagCountContainerNetSerializer::IsEqual(FNetSerializationContext& Context, const FNetIsEqualArgs& Args)
+{
+	if (Args.bStateIsQuantized)
+	{
+		FNetIsEqualArgs TempArgs = Args;
+		TempArgs.NetSerializerConfig = &StructNetSerializerConfig;
+		return StructNetSerializer->IsEqual(Context, TempArgs);
+	}
+	else
+	{
+		const SourceType& Source0 = *reinterpret_cast<const SourceType*>(Args.Source0);
+		const SourceType& Source1 = *reinterpret_cast<const SourceType*>(Args.Source1);
+
+		return Source0 == Source1;
+	}
+}
+	
+void FGameplayTagCountContainerNetSerializer::CloneDynamicState(FNetSerializationContext& Context, const FNetCloneDynamicStateArgs& Args)
+{
+	FNetCloneDynamicStateArgs TempArgs = Args;
+	TempArgs.NetSerializerConfig = &StructNetSerializerConfig;
+	StructNetSerializer->CloneDynamicState(Context, TempArgs);
+}
+	
+void FGameplayTagCountContainerNetSerializer::FreeDynamicState(FNetSerializationContext& Context, const FNetFreeDynamicStateArgs& Args)
+{
+	FNetFreeDynamicStateArgs TempArgs = Args;
+	TempArgs.NetSerializerConfig = &StructNetSerializerConfig;
+	StructNetSerializer->FreeDynamicState(Context, TempArgs);
+}
+
+void FGameplayTagCountContainerNetSerializer::CollectNetReferences(FNetSerializationContext& Context, const FNetCollectReferencesArgs& Args)
+{
+}
+
+bool FGameplayTagCountContainerNetSerializer::Validate(FNetSerializationContext&, const FNetValidateArgs& Args)
+{
+	return true;
+}
+
+static const FName PropertyNetSerializerRegistry_NAME_GameplayTagCountContainer("GameplayTagCountContainer");
+UE_NET_IMPLEMENT_NAMED_STRUCT_NETSERIALIZER_INFO(PropertyNetSerializerRegistry_NAME_GameplayTagCountContainer, FGameplayTagCountContainerNetSerializer);
+
+FGameplayTagCountContainerNetSerializer::FNetSerializerRegistryDelegates::FNetSerializerRegistryDelegates()
+: UE::Net::FNetSerializerRegistryDelegates(EFlags::ShouldBindLoadedModulesUpdatedDelegate)
+{
+}
+	
+FGameplayTagCountContainerNetSerializer::FNetSerializerRegistryDelegates::~FNetSerializerRegistryDelegates()
+{
+	UE_NET_UNREGISTER_NETSERIALIZER_INFO(PropertyNetSerializerRegistry_NAME_GameplayTagCountContainer);
+}
+
+void FGameplayTagCountContainerNetSerializer::FNetSerializerRegistryDelegates::OnPreFreezeNetSerializerRegistry()
+{
+	UE_NET_REGISTER_NETSERIALIZER_INFO(PropertyNetSerializerRegistry_NAME_GameplayTagCountContainer);
+}
+
+void FGameplayTagCountContainerNetSerializer::FNetSerializerRegistryDelegates::OnPostFreezeNetSerializerRegistry()
+{
+	const UStruct* TagCountMapStruct = FNetGameplayTagCountContainerStateForNetSerialize::StaticStruct();
+	FReplicationStateDescriptorBuilder::FParameters Params;
+	StructNetSerializerConfig.StateDescriptor = FReplicationStateDescriptorBuilder::CreateDescriptorForStruct(TagCountMapStruct, Params);
+
+	const FReplicationStateDescriptor* Descriptor = StructNetSerializerConfig.StateDescriptor.GetReference();
+
+	// Validate our assumptions regarding quantized state size and alignment.
+	if ((sizeof(QuantizedType) < Descriptor->InternalSize) || (alignof(QuantizedType) < Descriptor->InternalAlignment))
+	{
+		LowLevelFatalError(TEXT("FGameplayTagCountContainerNetSerializer::FQuantizedType has size %u and alignment %u but requires size %u and alignment %u."), uint32(sizeof(FQuantizedType)), uint32(alignof(FQuantizedType)), Descriptor->InternalSize, Descriptor->InternalAlignment);
+	}
+}
+
+}
diff --git a/Engine/Plugins/Runtime/GameplayAbilities/Source/GameplayAbilities/Public/Serialization/GameplayTagCountContainerNetSerializer.h b/Engine/Plugins/Runtime/GameplayAbilities/Source/GameplayAbilities/Public/Serialization/GameplayTagCountContainerNetSerializer.h
new file mode 100644
--- /dev/null
+++ b/Engine/Plugins/Runtime/GameplayAbilities/Source/GameplayAbilities/Public/Serialization/GameplayTagCountContainerNetSerializer.h
@@ -0,0 +1,47 @@
+// Copyright Epic Games, Inc. All Rights Reserved.
+
+#pragma once
+
+#include "Iris/Serialization/NetSerializer.h"
+#include "GameplayEffectTypes.h"
+#include "Iris/ReplicationState/PropertyReplicationState.h"
+#include "Iris/ReplicationSystem/ReplicationFragment.h"
+
+#include "CoreMinimal.h"
+#include "Engine/HitResult.h"
+#include "UObject/ObjectMacros.h"
+#include "UObject/Class.h"
+#include "Engine/NetSerialization.h"
+#include "Engine/EngineTypes.h"
+#include "GameFramework/Actor.h"
+
+#include "GameplayTagCountContainerNetSerializer.generated.h"
+
+USTRUCT()
+struct FGameplayTagCountContainerNetSerializerConfig : public FNetSerializerConfig
+{
+GENERATED_BODY()
+};
+
+USTRUCT()
+struct FNetGameplayTagCountContainerStateForNetSerialize
+{
+	GENERATED_BODY()
+
+	UPROPERTY()
+	TArray<FGameplayTagCountItem> TagStates;
+
+	void CopyFromTagCountContainer(const FGameplayTagCountContainer& Container);
+
+	bool operator== (const FNetGameplayTagCountContainerStateForNetSerialize& Other) const
+	{
+		return TagStates == Other.TagStates;
+	}
+};
+
+namespace UE::Net
+{
+
+UE_NET_DECLARE_SERIALIZER(FGameplayTagCountContainerNetSerializer, GAMEPLAYABILITIES_API);
+
+}
