[Backout] - CL46127921
#rnx
Original CL Desc
-----------------------------------------------------------------
TagCountContainerNetSerializer added for Iris.
Updates to GE prediction caught up and rejected logic so that predicted tags work correctly.

#FailedValidations Preflight

diff --git a/Engine/Plugins/Runtime/GameplayAbilities/Source/GameplayAbilities/Private/AbilitySystemComponent.cpp b/Engine/Plugins/Runtime/GameplayAbilities/Source/GameplayAbilities/Private/AbilitySystemComponent.cpp
--- a/Engine/Plugins/Runtime/GameplayAbilities/Source/GameplayAbilities/Private/AbilitySystemComponent.cpp
+++ b/Engine/Plugins/Runtime/GameplayAbilities/Source/GameplayAbilities/Private/AbilitySystemComponent.cpp
@@ -639,6 +639,27 @@ FOnGivenActiveGameplayEffectRemoved& UAbilitySystemComponent::OnAnyGameplayEffec
 
 void UAbilitySystemComponent::UpdateTagMapSingle_Internal(const FGameplayTag& Tag, const int32 CountDelta, EGameplayTagReplicationState TagRepState)
 {
+	if (UE::Net::ShouldUseIrisReplication())
+	{
+		// Iris replication doesn't work the same so we're falling back to using the Minimal and replicated loose tags until it's fixed.
+		ENetMode NetMode = GetWorld()->GetNetMode();
+		if (NetMode == NM_DedicatedServer || NetMode == NM_ListenServer)
+		{
+			if (TagRepState == EGameplayTagReplicationState::TagOnly)
+			{
+				if (CountDelta > 0)
+				{
+					GetMinimalReplicationTags_Mutable().AddTag(Tag);
+				}
+			}
+			else if (CountDelta < 0 || TagRepState > EGameplayTagReplicationState::TagOnly)
+			{
+				int32& TagCount = GetReplicatedLooseTags_Mutable().TagMap.FindOrAdd(Tag);
+				TagCount += CountDelta;
+			}
+		}
+	}
+	
 	// For removal, reorder calls so that FillParentTags is only called once
 	if (CountDelta > 0)
 	{
@@ -654,12 +675,56 @@ void UAbilitySystemComponent::UpdateTagMapSingle_Internal(const FGameplayTag& Ta
 		if (GameplayTagCountContainer.UpdateTagCount(Tag, CountDelta, TagRepState))
 		{
 			OnTagUpdated(Tag, false);
+			if (UE::Net::ShouldUseIrisReplication())
+			{
+				// Iris replication doesn't work the same so we're falling back to using the Minimal and replicated loose tags until it's fixed.
+				ENetMode NetMode = GetWorld()->GetNetMode();
+				if (NetMode == NM_DedicatedServer || NetMode == NM_ListenServer)
+				{
+					FMinimalReplicationTagCountMap& TagMap = GetMinimalReplicationTags_Mutable();
+					if (TagMap.TagMap.Contains(Tag))
+					{
+						TagMap.RemoveTag(Tag);
+					}
+
+					TagMap = GetReplicatedLooseTags_Mutable();
+					if (TagMap.TagMap.Contains(Tag))
+					{
+						TagMap.RemoveTag(Tag);
+					}
+				}
+			}
 		}
 	}
 }
 
 void UAbilitySystemComponent::UpdateTagMap_Internal(const FGameplayTagContainer& Container, const int32 CountDelta, EGameplayTagReplicationState TagRepState)
 {
+	if (UE::Net::ShouldUseIrisReplication())
+	{
+		// Iris replication doesn't work the same so we're falling back to using the Minimal and replicated loose tags until it's fixed.
+		ENetMode NetMode = GetWorld()->GetNetMode();
+		if (NetMode == NM_DedicatedServer || NetMode == NM_ListenServer)
+		{
+			if (TagRepState == EGameplayTagReplicationState::TagOnly)
+			{
+				if (CountDelta > 0)
+				{
+					GetMinimalReplicationTags_Mutable().AddTags(Container);
+				}
+			}
+			else if (CountDelta < 0 || TagRepState > EGameplayTagReplicationState::TagOnly)
+			{
+				for (auto TagIt = Container.CreateConstIterator(); TagIt; ++TagIt)
+				{
+					const FGameplayTag& Tag = *TagIt;
+					int32& TagCount = GetReplicatedLooseTags_Mutable().TagMap.FindOrAdd(Tag);
+					TagCount += CountDelta;
+				}
+			}
+		}
+	}
+	
 	// For removal, reorder calls so that FillParentTags is only called once
 	if (CountDelta > 0)
 	{
@@ -685,6 +750,25 @@ void UAbilitySystemComponent::UpdateTagMap_Internal(const FGameplayTagContainer&
 			if (GameplayTagCountContainer.UpdateTagCount_DeferredParentRemoval(Tag, CountDelta, DeferredTagChangeDelegates))
 			{
 				RemovedTags.Add(Tag);
+				if (UE::Net::ShouldUseIrisReplication())
+				{
+					// Iris replication doesn't work the same so we're falling back to using the Minimal and replicated loose tags until it's fixed.
+					ENetMode NetMode = GetWorld()->GetNetMode();
+					if (NetMode == NM_DedicatedServer || NetMode == NM_ListenServer)
+					{
+						FMinimalReplicationTagCountMap& TagMap = GetMinimalReplicationTags_Mutable();
+						if (TagMap.TagMap.Contains(Tag))
+						{
+							TagMap.RemoveTag(Tag);
+						}
+
+						TagMap = GetReplicatedLooseTags_Mutable();
+						if (TagMap.TagMap.Contains(Tag))
+						{
+							TagMap.RemoveTag(Tag);
+						}
+					}
+				}
 			}
 		}
 
@@ -1091,17 +1175,6 @@ void UAbilitySystemComponent::RemoveActiveGameplayEffect_AllowClientRemoval(FAct
 	ActiveGameplayEffects.RemoveActiveGameplayEffect(Handle, StacksToRemove);
 }
 
-void UAbilitySystemComponent::OnRejectedActiveGameplayEffect(FActiveGameplayEffectHandle Handle, int32 StacksToRemove)
-{
-	constexpr bool bPredictionRejected = true;
-	ActiveGameplayEffects.RemoveActiveGameplayEffect(Handle, StacksToRemove, bPredictionRejected);
-}
-
-void UAbilitySystemComponent::OnCaughtUpActiveGameplayEffect(FActiveGameplayEffectHandle Handle, int32 StacksToRemove)
-{
-	ActiveGameplayEffects.RemoveActiveGameplayEffect(Handle, StacksToRemove);
-}
-
 void UAbilitySystemComponent::OnPredictiveGameplayEffectStackCaughtUp(FActiveGameplayEffectHandle Handle)
 {
 	ActiveGameplayEffects.OnPredictiveGameplayEffectStackCaughtUp(Handle);
diff --git a/Engine/Plugins/Runtime/GameplayAbilities/Source/GameplayAbilities/Private/GameplayEffect.cpp b/Engine/Plugins/Runtime/GameplayAbilities/Source/GameplayAbilities/Private/GameplayEffect.cpp
--- a/Engine/Plugins/Runtime/GameplayAbilities/Source/GameplayAbilities/Private/GameplayEffect.cpp
+++ b/Engine/Plugins/Runtime/GameplayAbilities/Source/GameplayAbilities/Private/GameplayEffect.cpp
@@ -2703,8 +2703,6 @@ FActiveGameplayEffect& FActiveGameplayEffect::operator=(FActiveGameplayEffect&&
 	{
 		MostRecentArrayReplicationKey = Other.MostRecentArrayReplicationKey;
 	}
-	
-	bPostPredictObject = Other.bPostPredictObject;
 
 	return *this;
 }
@@ -2742,8 +2740,6 @@ FActiveGameplayEffect& FActiveGameplayEffect::operator=(const FActiveGameplayEff
 	{
 		MostRecentArrayReplicationKey = Other.MostRecentArrayReplicationKey;
 	}
-	
-	bPostPredictObject = Other.bPostPredictObject;
 
 	return *this;
 }
@@ -2816,11 +2812,11 @@ void FActiveGameplayEffect::PostReplicatedAdd(const struct FActiveGameplayEffect
 	}
 
 	bool ShouldInvokeGameplayCueEvents = true;
-	bPostPredictObject = PredictionKey.IsLocalClientKey();
-	if (bPostPredictObject)
+	if (PredictionKey.IsLocalClientKey())
 	{
 		// PredictionKey will only be valid on the client that predicted it. So if this has a valid PredictionKey, we can assume we already predicted it and shouldn't invoke GameplayCues.
 		// We may need to do more bookkeeping here in the future. Possibly give the predicted gameplayeffect a chance to pass something off to the new replicated gameplay effect.
+		
 		if (InArray.HasPredictedEffectWithPredictedKey(PredictionKey))
 		{
 			ShouldInvokeGameplayCueEvents = false;
@@ -4443,8 +4439,7 @@ PRAGMA_ENABLE_DEPRECATION_WARNINGS
 		{
 			// Once replicated state has caught up to this prediction key, we must remove this gameplay effect.
 			constexpr int32 RemoveAllStacks = -1;
-			InPredictionKey.NewCaughtUpDelegate().BindUObject(Owner, &UAbilitySystemComponent::OnCaughtUpActiveGameplayEffect, AppliedActiveGE->Handle, RemoveAllStacks);
-			InPredictionKey.NewRejectedDelegate().BindUObject(Owner, &UAbilitySystemComponent::OnRejectedActiveGameplayEffect, AppliedActiveGE->Handle, RemoveAllStacks);
+			InPredictionKey.NewRejectOrCaughtUpDelegate(FPredictionKeyEvent::CreateUObject(Owner, &UAbilitySystemComponent::RemoveActiveGameplayEffect_AllowClientRemoval, AppliedActiveGE->Handle, RemoveAllStacks));
 		}
 		else // We have an existing Stackable GE
 		{
@@ -4470,7 +4465,7 @@ PRAGMA_ENABLE_DEPRECATION_WARNINGS
 		const bool bInvokeGameplayCueEvents = (Spec.Def->DurationPolicy != EGameplayEffectDurationType::Instant);
 		InternalOnActiveGameplayEffectAdded(*AppliedActiveGE, bInvokeGameplayCueEvents);
 	}
-	
+
 	return AppliedActiveGE;
 }
 
@@ -4512,12 +4507,6 @@ void FActiveGameplayEffectsContainer::AddActiveGameplayEffectGrantedTagsAndModif
 		return;
 	}
 
-	if (!Owner)
-	{
-		UE_LOG(LogGameplayEffects, Error, TEXT("AddActiveGameplayEffectGrantedTagsAndModifiers called with null Owner!"));
-		return;
-	}
-
 	GAMEPLAYEFFECT_SCOPE_LOCK();
 
 	// Register this ActiveGameplayEffects modifiers with our Attribute Aggregators
@@ -4535,7 +4524,7 @@ void FActiveGameplayEffectsContainer::AddActiveGameplayEffectGrantedTagsAndModif
 			const FGameplayModifierInfo &ModInfo = Effect.Spec.Def->Modifiers[ModIdx];
 
 			// skip over any modifiers for attributes that we don't have
-			if (Owner->HasAttributeSetForAttribute(ModInfo.Attribute) == false)
+			if (!Owner || Owner->HasAttributeSetForAttribute(ModInfo.Attribute) == false)
 			{
 				continue;
 			}
@@ -4554,7 +4543,7 @@ void FActiveGameplayEffectsContainer::AddActiveGameplayEffectGrantedTagsAndModif
 	}
 	else
 	{
-		if (Effect.Spec.Def->PeriodicInhibitionPolicy != EGameplayEffectPeriodInhibitionRemovedPolicy::NeverReset && Owner->IsOwnerActorAuthoritative())
+		if (Effect.Spec.Def->PeriodicInhibitionPolicy != EGameplayEffectPeriodInhibitionRemovedPolicy::NeverReset && Owner && Owner->IsOwnerActorAuthoritative())
 		{
 			FTimerManager& TimerManager = Owner->GetWorld()->GetTimerManager();
 			FTimerDelegate Delegate = FTimerDelegate::CreateUObject(Owner, &UAbilitySystemComponent::ExecutePeriodicEffect, Effect.Handle);
@@ -4569,14 +4558,11 @@ void FActiveGameplayEffectsContainer::AddActiveGameplayEffectGrantedTagsAndModif
 		}
 	}
 
-	// Local application should be responsible for applying tags
-	if (!Effect.bPostPredictObject)
-	{
-		// Update our owner with the tags this GameplayEffect grants them
-		Owner->UpdateTagMap(Effect.Spec.Def->GetGrantedTags(), 1, EGameplayTagReplicationState::CountToOwner);
-		Owner->UpdateTagMap(Effect.Spec.DynamicGrantedTags, 1, EGameplayTagReplicationState::CountToOwner);
-	}
-	
+
+	// Update our owner with the tags this GameplayEffect grants them
+	Owner->UpdateTagMap(Effect.Spec.Def->GetGrantedTags(), 1, ShouldUseMinimalReplication() ? EGameplayTagReplicationState::TagOnly : EGameplayTagReplicationState::None);
+	Owner->UpdateTagMap(Effect.Spec.DynamicGrantedTags, 1, ShouldUseMinimalReplication() ? EGameplayTagReplicationState::TagOnly : EGameplayTagReplicationState::None);
+
 	// Update our owner with the blocked ability tags this GameplayEffect adds to them
 	Owner->BlockAbilitiesWithTags(Effect.Spec.Def->GetBlockedAbilityTags());
 
@@ -4638,8 +4624,8 @@ PRAGMA_ENABLE_DEPRECATION_WARNINGS
 	Owner->OnActiveGameplayEffectAddedDelegateToSelf.Broadcast(Owner, Effect.Spec, Effect.Handle);
 }
 
-/** Called on server to remove a GameplayEffect and on owning client for Post-Prediction cleanup */
-bool FActiveGameplayEffectsContainer::RemoveActiveGameplayEffect(FActiveGameplayEffectHandle Handle, int32 StacksToRemove, bool bPredictionRejected)
+/** Called on server to remove a GameplayEffect */
+bool FActiveGameplayEffectsContainer::RemoveActiveGameplayEffect(FActiveGameplayEffectHandle Handle, int32 StacksToRemove)
 {
 	// Iterating through manually since this is a removal operation and we need to pass the index into InternalRemoveActiveGameplayEffect
 	int32 NumGameplayEffects = GetNumGameplayEffects();
@@ -4648,7 +4634,7 @@ bool FActiveGameplayEffectsContainer::RemoveActiveGameplayEffect(FActiveGameplay
 		FActiveGameplayEffect& Effect = *GetActiveGameplayEffect(ActiveGEIdx);
 		if (Effect.Handle == Handle && Effect.IsPendingRemove == false)
 		{
-			InternalRemoveActiveGameplayEffect(ActiveGEIdx, StacksToRemove, true, bPredictionRejected);
+			InternalRemoveActiveGameplayEffect(ActiveGEIdx, StacksToRemove, true);
 			return true;
 		}
 	}
@@ -4692,7 +4678,7 @@ void FActiveGameplayEffectsContainer::InternalExecutePeriodicGameplayEffect(FAct
 }
 
 /** Called by server to actually remove a GameplayEffect */
-bool FActiveGameplayEffectsContainer::InternalRemoveActiveGameplayEffect(int32 Idx, int32 StacksToRemove, bool bPrematureRemoval, bool bPredictionRejected)
+bool FActiveGameplayEffectsContainer::InternalRemoveActiveGameplayEffect(int32 Idx, int32 StacksToRemove, bool bPrematureRemoval)
 {
 	SCOPE_CYCLE_COUNTER(STAT_RemoveActiveGameplayEffect);
 
@@ -4725,7 +4711,6 @@ bool FActiveGameplayEffectsContainer::InternalRemoveActiveGameplayEffect(int32 I
 		GameplayEffectRemovalInfo.ActiveEffect = &Effect;
 		GameplayEffectRemovalInfo.StackCount = Effect.Spec.GetStackCount();
 		GameplayEffectRemovalInfo.bPrematureRemoval = bPrematureRemoval;
-		GameplayEffectRemovalInfo.bPredictionRejected = bPredictionRejected;
 		GameplayEffectRemovalInfo.EffectContext = Effect.Spec.GetEffectContext();
 
 		if (StacksToRemove > 0 && Effect.Spec.GetStackCount() > StacksToRemove)
@@ -4831,7 +4816,7 @@ void FActiveGameplayEffectsContainer::InternalOnActiveGameplayEffectRemoved(FAct
 		// Only Need to update tags and modifiers if the gameplay effect is active.
 		if (!Effect.bIsInhibited)
 		{
-			RemoveActiveGameplayEffectGrantedTagsAndModifiers(Effect, bInvokeGameplayCueEvents, GameplayEffectRemovalInfo.bPredictionRejected);
+			RemoveActiveGameplayEffectGrantedTagsAndModifiers(Effect, bInvokeGameplayCueEvents);
 		}
 
 		RemoveCustomMagnitudeExternalDependencies(Effect);
@@ -4846,7 +4831,7 @@ void FActiveGameplayEffectsContainer::InternalOnActiveGameplayEffectRemoved(FAct
 	OnActiveGameplayEffectRemovedDelegate.Broadcast(Effect);
 }
 
-void FActiveGameplayEffectsContainer::RemoveActiveGameplayEffectGrantedTagsAndModifiers(const FActiveGameplayEffect& Effect, bool bInvokeGameplayCueEvents, bool bPredictionRejected)
+void FActiveGameplayEffectsContainer::RemoveActiveGameplayEffectGrantedTagsAndModifiers(const FActiveGameplayEffect& Effect, bool bInvokeGameplayCueEvents)
 {
 	// Update AttributeAggregators: remove mods from this ActiveGE Handle
 	if (Effect.Spec.GetPeriod() <= UGameplayEffect::NO_PERIOD)
@@ -4863,13 +4848,9 @@ void FActiveGameplayEffectsContainer::RemoveActiveGameplayEffectGrantedTagsAndMo
 		}
 	}
 
-	// We're no longer predicting for remove unless it was rejected, so let the server replication handle the tags.
-	if (IsNetAuthority() || bPredictionRejected)
-	{
-		// Update gameplaytag count and broadcast delegate if we are at 0
-		Owner->UpdateTagMap(Effect.Spec.Def->GetGrantedTags(), -1, EGameplayTagReplicationState::None);
-		Owner->UpdateTagMap(Effect.Spec.DynamicGrantedTags, -1, EGameplayTagReplicationState::None);
-	}
+	// Update gameplaytag count and broadcast delegate if we are at 0
+	Owner->UpdateTagMap(Effect.Spec.Def->GetGrantedTags(), -1, ShouldUseMinimalReplication() ? EGameplayTagReplicationState::TagOnly : EGameplayTagReplicationState::None);
+	Owner->UpdateTagMap(Effect.Spec.DynamicGrantedTags, -1, ShouldUseMinimalReplication() ? EGameplayTagReplicationState::TagOnly : EGameplayTagReplicationState::None);
 
 	// Update our owner with the blocked ability tags this GameplayEffect adds to them
 	Owner->UnBlockAbilitiesWithTags(Effect.Spec.Def->GetBlockedAbilityTags());
diff --git a/Engine/Plugins/Runtime/GameplayAbilities/Source/GameplayAbilities/Private/GameplayEffectTypes.cpp b/Engine/Plugins/Runtime/GameplayAbilities/Source/GameplayAbilities/Private/GameplayEffectTypes.cpp
--- a/Engine/Plugins/Runtime/GameplayAbilities/Source/GameplayAbilities/Private/GameplayEffectTypes.cpp
+++ b/Engine/Plugins/Runtime/GameplayAbilities/Source/GameplayAbilities/Private/GameplayEffectTypes.cpp
@@ -18,10 +18,6 @@
 #endif
 
 #include "Iris/Serialization/NetBitStreamUtil.h"
-#include "Iris/Serialization/NetSerializationContext.h"
-#include "Iris/Serialization/NetSerializerDelegates.h"
-#include "Iris/Serialization/PolymorphicNetSerializerImpl.h"
-#include "Iris/ReplicationState/PropertyNetSerializerInfoRegistry.h"
 
 #include UE_INLINE_GENERATED_CPP_BY_NAME(GameplayEffectTypes)
 
@@ -654,7 +650,7 @@ bool FGameplayTagCountContainer::NetDeltaSerialize(FNetDeltaSerializeInfo & Delt
 			}
 		}
 	}
-	
+
 	if (DeltaParams.Reader)
 	{
 		FBitReader& Reader = *DeltaParams.Reader;
@@ -668,10 +664,9 @@ bool FGameplayTagCountContainer::NetDeltaSerialize(FNetDeltaSerializeInfo & Delt
 			for (; NumRemovals > 0; NumRemovals--)
 			{
 				RemovedItem.Tag.NetSerialize(Reader, DeltaParams.Map, bTagSuccess);
-				int32* Count = GameplayTagCountMap.Find(RemovedItem.Tag);
-				if (Count && *Count > 0)
+				if (GameplayTagCountMap.Contains(RemovedItem.Tag))
 				{
-					UpdateTagMap_Internal(RemovedItem.Tag, -(*Count), EGameplayTagReplicationState::None);
+					UpdateTagMap_Internal(RemovedItem.Tag, -GameplayTagCountMap[RemovedItem.Tag], EGameplayTagReplicationState::None);
 				}
 			}
 		}
@@ -688,12 +683,16 @@ bool FGameplayTagCountContainer::NetDeltaSerialize(FNetDeltaSerializeInfo & Delt
 				NewCount = 1;
 				UpdatedItem.Tag.NetSerialize(Reader, DeltaParams.Map, bTagSuccess);
 				Reader.SerializeBits(&bPackCount, 1);
+				int32& Count = GameplayTagCountMap.FindOrAdd(UpdatedItem.Tag);
 				if (bPackCount)
 				{
 					Reader.SerializeIntPacked(NewCount);
 				}
 
-				SetTagCount(UpdatedItem.Tag, NewCount, EGameplayTagReplicationState::None);
+				if (NewCount != Count)
+				{
+					UpdateTagMap_Internal(UpdatedItem.Tag, NewCount - Count, EGameplayTagReplicationState::None);
+				}
 			}
 		}
 	}
@@ -767,21 +766,10 @@ bool FGameplayTagCountContainer::UpdateExplicitTags(const FGameplayTag& Tag, con
 				ABILITY_LOG(Warning, TEXT("Trying to change existing tag \"%s\" replication state from \"%s\" to \"%s\""), *Tag.ToString(), *UEnum::GetDisplayValueAsText(TagCountItem.ReplicationState).ToString(), *UEnum::GetDisplayValueAsText(TagRepState).ToString());
 			}
 
-			// Promote our rep state to the higher level
+			// Some sort of authority check before updating this value?
 			TagCountItem.ReplicationState = TagRepState > TagCountItem.ReplicationState ? TagRepState : TagCountItem.ReplicationState;
-			TagRepState = TagCountItem.ReplicationState;
-		}
-
-		// make sure clients only have the 1 tag if it's tag Only.
-		if (Owner && Owner->GetOwnerRole() != ROLE_Authority && CountDelta > 0  && TagCountItem.ReplicationState == EGameplayTagReplicationState::TagOnly)
-		{
-			TagCountItem.Count = 1;
-		}
-		else
-		{
-			TagCountItem.Count += CountDelta;
 		}
-		
+		TagCountItem.Count += CountDelta;
 		if (TagCountItem.Count <= 0)
 		{
 			Items.RemoveAtSwap(TagCountIndex);
@@ -802,7 +790,7 @@ bool FGameplayTagCountContainer::UpdateExplicitTags(const FGameplayTag& Tag, con
 	return true;
 }
 
-bool FGameplayTagCountContainer::GatherTagChangeDelegates(const FGameplayTag& Tag, const int32 CountDelta, TArray<FDeferredTagChangeDelegate>& TagChangeDelegates, EGameplayTagReplicationState TagRepState)
+bool FGameplayTagCountContainer::GatherTagChangeDelegates(const FGameplayTag& Tag, const int32 CountDelta, TArray<FDeferredTagChangeDelegate>& TagChangeDelegates)
 {
 	// Check if change delegates are required to fire for the tag or any of its parents based on the count change
 	FGameplayTagContainer TagAndParentsContainer = Tag.GetGameplayTagParents();
@@ -815,13 +803,10 @@ bool FGameplayTagCountContainer::GatherTagChangeDelegates(const FGameplayTag& Ta
 		int32& TagCountRef = GameplayTagCountMap.FindOrAdd(CurTag);
 
 		const int32 OldCount = TagCountRef;
-		
-		const bool ClampTagOnlyCount = Owner && Owner->GetOwnerRole() != ROLE_Authority && CountDelta > 0  && TagRepState == EGameplayTagReplicationState::TagOnly;
-		
+
 		// Apply the delta to the count in the map
 		int32 NewTagCount = FMath::Max(OldCount + CountDelta, 0);
-		
-		TagCountRef = ClampTagOnlyCount ? FMath::Min(NewTagCount, 1) : NewTagCount;
+		TagCountRef = NewTagCount;
 
 		// If a significant change (new addition or total removal) occurred, trigger related delegates
 		const bool SignificantChange = (OldCount == 0 || NewTagCount == 0);
@@ -829,9 +814,9 @@ bool FGameplayTagCountContainer::GatherTagChangeDelegates(const FGameplayTag& Ta
 		if (SignificantChange)
 		{
 			TagChangeDelegates.AddDefaulted();
-			TagChangeDelegates.Last().BindLambda([Delegate = OnAnyTagChangeDelegate, CurTag, TagCountRef]()
+			TagChangeDelegates.Last().BindLambda([Delegate = OnAnyTagChangeDelegate, CurTag, NewTagCount]()
 			{
-				Delegate.Broadcast(CurTag, TagCountRef);
+				Delegate.Broadcast(CurTag, NewTagCount);
 			});
 		}
 
@@ -839,17 +824,17 @@ bool FGameplayTagCountContainer::GatherTagChangeDelegates(const FGameplayTag& Ta
 		if (DelegateInfo)
 		{
 			TagChangeDelegates.AddDefaulted();
-			TagChangeDelegates.Last().BindLambda([Delegate = DelegateInfo->OnAnyChange, CurTag, TagCountRef]()
+			TagChangeDelegates.Last().BindLambda([Delegate = DelegateInfo->OnAnyChange, CurTag, NewTagCount]()
 			{
-				Delegate.Broadcast(CurTag, TagCountRef);
+				Delegate.Broadcast(CurTag, NewTagCount);
 			});
 
 			if (SignificantChange)
 			{
 				TagChangeDelegates.AddDefaulted();
-				TagChangeDelegates.Last().BindLambda([Delegate = DelegateInfo->OnNewOrRemove, CurTag, TagCountRef]()
+				TagChangeDelegates.Last().BindLambda([Delegate = DelegateInfo->OnNewOrRemove, CurTag, NewTagCount]()
 				{
-					Delegate.Broadcast(CurTag, TagCountRef);
+					Delegate.Broadcast(CurTag, NewTagCount);
 				});
 			}
 		}
@@ -864,11 +849,6 @@ bool FGameplayTagCountContainer::UpdateTagMap_Internal(const FGameplayTag& Tag,
 	{
 		return false;
 	}
-	
-	if (!ensureMsgf(CountDelta != 0, TEXT("FGameplayTagCountContainer Modifying tag count by 0 for tag %s"), *Tag.ToString()))
-	{
-		return false;
-	}
 
 	if (!UpdateExplicitTags(Tag, CountDelta, false, TagRepState))
 	{
@@ -876,7 +856,7 @@ bool FGameplayTagCountContainer::UpdateTagMap_Internal(const FGameplayTag& Tag,
 	}
 
 	TArray<FDeferredTagChangeDelegate> DeferredTagChangeDelegates;
-	bool bSignificantChange = GatherTagChangeDelegates(Tag, CountDelta, DeferredTagChangeDelegates, TagRepState);
+	bool bSignificantChange = GatherTagChangeDelegates(Tag, CountDelta, DeferredTagChangeDelegates);
 	for (FDeferredTagChangeDelegate& Delegate : DeferredTagChangeDelegates)
 	{
 		Delegate.Execute();
@@ -892,7 +872,7 @@ bool FGameplayTagCountContainer::UpdateTagMapDeferredParentRemoval_Internal(cons
 		return false;
 	}
 
-	return GatherTagChangeDelegates(Tag, CountDelta, DeferredTagChangeDelegates, TagRepState);
+	return GatherTagChangeDelegates(Tag, CountDelta, DeferredTagChangeDelegates);
 }
 
 FGameplayTagBlueprintPropertyMap::FGameplayTagBlueprintPropertyMap()
diff --git a/Engine/Plugins/Runtime/GameplayAbilities/Source/GameplayAbilities/Public/AbilitySystemComponent.h b/Engine/Plugins/Runtime/GameplayAbilities/Source/GameplayAbilities/Public/AbilitySystemComponent.h
--- a/Engine/Plugins/Runtime/GameplayAbilities/Source/GameplayAbilities/Public/AbilitySystemComponent.h
+++ b/Engine/Plugins/Runtime/GameplayAbilities/Source/GameplayAbilities/Public/AbilitySystemComponent.h
@@ -1847,10 +1847,6 @@ protected:
 	/** A version of RemoveActiveGameplayEffect that allows us to remove it even when we don't have authority */
 	UE_API void RemoveActiveGameplayEffect_AllowClientRemoval(FActiveGameplayEffectHandle Handle, int32 StacksToRemove = -1);
 
-	UE_API void OnCaughtUpActiveGameplayEffect(FActiveGameplayEffectHandle Handle, int32 StacksToRemove = -1);
-
-	UE_API void OnRejectedActiveGameplayEffect(FActiveGameplayEffectHandle Handle, int32 StacksToRemove = -1);
-
 	UE_API void OnPredictiveGameplayEffectStackCaughtUp(FActiveGameplayEffectHandle Handle);
 
 	/** Contains all of the gameplay effects that are currently active on this component */
diff --git a/Engine/Plugins/Runtime/GameplayAbilities/Source/GameplayAbilities/Public/GameplayEffect.h b/Engine/Plugins/Runtime/GameplayAbilities/Source/GameplayAbilities/Public/GameplayEffect.h
--- a/Engine/Plugins/Runtime/GameplayAbilities/Source/GameplayAbilities/Public/GameplayEffect.h
+++ b/Engine/Plugins/Runtime/GameplayAbilities/Source/GameplayAbilities/Public/GameplayEffect.h
@@ -1445,9 +1445,6 @@ struct FActiveGameplayEffect : public FFastArraySerializerItem
 	
 	/** All the bindable events for this active effect (bundled to allow easier non-const access to these events via the ASC) */
 	FActiveGameplayEffectEvents EventSet;
-
-	/** Signifies an active effect that was predicted and has now been replicated back to the client */
-	bool bPostPredictObject = false;
 };
 
 DECLARE_DELEGATE_RetVal_OneParam(bool, FActiveGameplayEffectQueryCustomMatch, const FActiveGameplayEffect&);
@@ -1697,7 +1694,7 @@ struct FActiveGameplayEffectsContainer : public FFastArraySerializer
 	
 	UE_API void ExecutePeriodicGameplayEffect(FActiveGameplayEffectHandle Handle);	// This should not be outward facing to the skill system API, should only be called by the owning AbilitySystemComponent
 
-	UE_API bool RemoveActiveGameplayEffect(FActiveGameplayEffectHandle Handle, int32 StacksToRemove, bool bPredictionRejected = false);
+	UE_API bool RemoveActiveGameplayEffect(FActiveGameplayEffectHandle Handle, int32 StacksToRemove);
 
 	UE_API void GetGameplayEffectStartTimeAndDuration(FActiveGameplayEffectHandle Handle, float& EffectStartTime, float& EffectDuration) const;
 
@@ -1927,13 +1924,13 @@ private:
 	UE_API void InternalExecutePeriodicGameplayEffect(FActiveGameplayEffect& ActiveEffect);
 
 	/** Called internally to actually remove a GameplayEffect or to reduce its StackCount. Returns true if we resized our internal GameplayEffect array. */
-	UE_API bool InternalRemoveActiveGameplayEffect(int32 Idx, int32 StacksToRemove, bool bPrematureRemoval, bool bPredictionRejected = false);
+	UE_API bool InternalRemoveActiveGameplayEffect(int32 Idx, int32 StacksToRemove, bool bPrematureRemoval);
 	
 	/** Called both in server side creation and replication creation/deletion */
 	UE_API void InternalOnActiveGameplayEffectAdded(FActiveGameplayEffect& Effect, bool bInvokeGameplayCueEvents);
 	UE_API void InternalOnActiveGameplayEffectRemoved(FActiveGameplayEffect& Effect, bool bInvokeGameplayCueEvents, const FGameplayEffectRemovalInfo& GameplayEffectRemovalInfo);
 
-	UE_API void RemoveActiveGameplayEffectGrantedTagsAndModifiers(const FActiveGameplayEffect& Effect, bool bInvokeGameplayCueEvents, bool bPredictionRejected = false);
+	UE_API void RemoveActiveGameplayEffectGrantedTagsAndModifiers(const FActiveGameplayEffect& Effect, bool bInvokeGameplayCueEvents);
 	UE_API void AddActiveGameplayEffectGrantedTagsAndModifiers(FActiveGameplayEffect& Effect, bool bInvokeGameplayCueEvents);
 
 	/** Updates tag dependency map when a GameplayEffect is removed */
diff --git a/Engine/Plugins/Runtime/GameplayAbilities/Source/GameplayAbilities/Public/GameplayEffectTypes.h b/Engine/Plugins/Runtime/GameplayAbilities/Source/GameplayAbilities/Public/GameplayEffectTypes.h
--- a/Engine/Plugins/Runtime/GameplayAbilities/Source/GameplayAbilities/Public/GameplayEffectTypes.h
+++ b/Engine/Plugins/Runtime/GameplayAbilities/Source/GameplayAbilities/Public/GameplayEffectTypes.h
@@ -36,10 +36,8 @@ struct FGameplayEffectModCallbackData;
 struct FGameplayEffectSpec;
 struct FHitResult;
 struct FMinimalReplicationTagCountMapForNetSerializer;
-struct FNetGameplayTagCountContainerStateForNetSerialize;
 namespace UE::Net
 {
-	struct FGameplayTagCountContainerNetSerializer;
 	struct FGameplayEffectContextHandleAccessorForNetSerializer;
 	class FMinimalReplicationTagCountMapReplicationFragment;
 }
@@ -821,10 +819,6 @@ struct FGameplayEffectRemovalInfo
 	UPROPERTY(VisibleAnywhere, BlueprintReadOnly, Category="Removal")
 	bool bPrematureRemoval = false;
 
-	/** True when the effect is being removed due to prediction rejection */
-	UPROPERTY(VisibleAnywhere, BlueprintReadOnly, Category="Removal")
-	bool bPredictionRejected = false;
-
 	/** Number of Stacks this gameplay effect had before it was removed. */
 	UPROPERTY(VisibleAnywhere, BlueprintReadOnly, Category="Removal")
 	int32 StackCount = 0;
@@ -832,7 +826,7 @@ struct FGameplayEffectRemovalInfo
 	/** Actor this gameplay effect was targeting. */
 	UPROPERTY(VisibleAnywhere, BlueprintReadOnly, Category="Removal")
 	FGameplayEffectContextHandle EffectContext;
-	
+
 	/** The Effect being Removed */
 	const FActiveGameplayEffect* ActiveEffect = nullptr;
 };
@@ -1053,7 +1047,7 @@ enum class EGameplayTagReplicationState : uint8
 {
 	None,				// Tag will not replicate
 	TagOnly,			// Tag will replicate to all without count
-	CountToOwner,		// Tag will replicate to all and Count will replicate to owner, Functions same as TagAndCountToAll with Iris
+	CountToOwner,		// Tag will replicate to all and Count will replicate to owner
 	TagAndCountToAll	// Tag and Count will replicate to all
 };
 
@@ -1352,31 +1346,8 @@ struct FGameplayTagCountContainer
 
 	UE_API void SetOwner(UAbilitySystemComponent* ASC);
 
-	bool operator==(const FGameplayTagCountContainer& Other) const
-	{
-		if (Items.Num() != Other.Items.Num())
-		{
-			return false;
-		}
-
-		int32 Idx = INDEX_NONE;
-		for (const FGameplayTagCountItem& Item : Items)
-		{
-			Idx = Other.Items.Find(Item);
-			if (Idx == INDEX_NONE || Item.Count != Other.Items[Idx].Count)
-			{
-				return false;
-			}
-		}
-
-		return true;
-	}
-
 private:
 
-	friend UE::Net::FGameplayTagCountContainerNetSerializer;
-	friend FNetGameplayTagCountContainerStateForNetSerialize;
-
 	struct FDelegateInfo
 	{
 		FOnGameplayEffectTagCountChanged	OnNewOrRemove;
@@ -1412,7 +1383,7 @@ private:
 	UE_API bool UpdateExplicitTags(const FGameplayTag& Tag, int32 CountDelta, bool bDeferParentTagsOnRemove, EGameplayTagReplicationState TagRepState = EGameplayTagReplicationState::None);
 
 	/** Internal helper function to collect the delegates that need to be called when Tag has its count changed by CountDelta. */
-	UE_API bool GatherTagChangeDelegates(const FGameplayTag& Tag, int32 CountDelta, TArray<FDeferredTagChangeDelegate>& TagChangeDelegates, EGameplayTagReplicationState TagRepState);
+	UE_API bool GatherTagChangeDelegates(const FGameplayTag& Tag, int32 CountDelta, TArray<FDeferredTagChangeDelegate>& TagChangeDelegates);
 };
 
 template<>
diff --git a/Engine/Plugins/Runtime/GameplayAbilities/Source/GameplayAbilities/Public/Serialization/GameplayTagCountContainerNetSerializer.cpp b/Engine/Plugins/Runtime/GameplayAbilities/Source/GameplayAbilities/Public/Serialization/GameplayTagCountContainerNetSerializer.cpp
deleted file mode 100644
diff --git a/Engine/Plugins/Runtime/GameplayAbilities/Source/GameplayAbilities/Public/Serialization/GameplayTagCountContainerNetSerializer.h b/Engine/Plugins/Runtime/GameplayAbilities/Source/GameplayAbilities/Public/Serialization/GameplayTagCountContainerNetSerializer.h
deleted file mode 100644
