Fix MinimalReplicationTags getting stomped.

#rnx

diff --git a/Engine/Plugins/Runtime/GameplayAbilities/Source/GameplayAbilities/Private/AbilitySystemComponent.cpp b/Engine/Plugins/Runtime/GameplayAbilities/Source/GameplayAbilities/Private/AbilitySystemComponent.cpp
--- a/Engine/Plugins/Runtime/GameplayAbilities/Source/GameplayAbilities/Private/AbilitySystemComponent.cpp
+++ b/Engine/Plugins/Runtime/GameplayAbilities/Source/GameplayAbilities/Private/AbilitySystemComponent.cpp
@@ -681,16 +681,16 @@ void UAbilitySystemComponent::UpdateTagMapSingle_Internal(const FGameplayTag& Ta
 				ENetMode NetMode = GetWorld()->GetNetMode();
 				if (NetMode == NM_DedicatedServer || NetMode == NM_ListenServer)
 				{
-					FMinimalReplicationTagCountMap& TagMap = GetMinimalReplicationTags_Mutable();
-					if (TagMap.TagMap.Contains(Tag))
+					FMinimalReplicationTagCountMap& MTagMap = GetMinimalReplicationTags_Mutable();
+					if (MTagMap.TagMap.Contains(Tag))
 					{
-						TagMap.RemoveTag(Tag);
+						MTagMap.RemoveTag(Tag);
 					}
 
-					TagMap = GetReplicatedLooseTags_Mutable();
-					if (TagMap.TagMap.Contains(Tag))
+					FMinimalReplicationTagCountMap& RLTagMap = GetReplicatedLooseTags_Mutable();
+					if (RLTagMap.TagMap.Contains(Tag))
 					{
-						TagMap.RemoveTag(Tag);
+						RLTagMap.RemoveTag(Tag);
 					}
 				}
 			}
@@ -756,16 +756,16 @@ void UAbilitySystemComponent::UpdateTagMap_Internal(const FGameplayTagContainer&
 					ENetMode NetMode = GetWorld()->GetNetMode();
 					if (NetMode == NM_DedicatedServer || NetMode == NM_ListenServer)
 					{
-						FMinimalReplicationTagCountMap& TagMap = GetMinimalReplicationTags_Mutable();
-						if (TagMap.TagMap.Contains(Tag))
+						FMinimalReplicationTagCountMap& MTagMap = GetMinimalReplicationTags_Mutable();
+						if (MTagMap.TagMap.Contains(Tag))
 						{
-							TagMap.RemoveTag(Tag);
+							MTagMap.RemoveTag(Tag);
 						}
 
-						TagMap = GetReplicatedLooseTags_Mutable();
-						if (TagMap.TagMap.Contains(Tag))
+						FMinimalReplicationTagCountMap& RLTagMap = GetReplicatedLooseTags_Mutable();
+						if (RLTagMap.TagMap.Contains(Tag))
 						{
-							TagMap.RemoveTag(Tag);
+							RLTagMap.RemoveTag(Tag);
 						}
 					}
 				}
