SetTagCount needs to check explicit tags and we should make sure our replicated tags don't try to assign negative values.

#rnx

diff --git a/Engine/Plugins/Runtime/GameplayAbilities/Source/GameplayAbilities/Private/AbilitySystemComponent.cpp b/Engine/Plugins/Runtime/GameplayAbilities/Source/GameplayAbilities/Private/AbilitySystemComponent.cpp
--- a/Engine/Plugins/Runtime/GameplayAbilities/Source/GameplayAbilities/Private/AbilitySystemComponent.cpp
+++ b/Engine/Plugins/Runtime/GameplayAbilities/Source/GameplayAbilities/Private/AbilitySystemComponent.cpp
@@ -655,7 +655,7 @@ void UAbilitySystemComponent::UpdateTagMapSingle_Internal(const FGameplayTag& Ta
 			else if (CountDelta < 0 || TagRepState > EGameplayTagReplicationState::TagOnly)
 			{
 				int32& TagCount = GetReplicatedLooseTags_Mutable().TagMap.FindOrAdd(Tag);
-				TagCount += CountDelta;
+				GetReplicatedLooseTags_Mutable().SetTagCount(Tag, FMath::Max(0, TagCount + CountDelta));
 			}
 		}
 	}
@@ -719,7 +719,7 @@ void UAbilitySystemComponent::UpdateTagMap_Internal(const FGameplayTagContainer&
 				{
 					const FGameplayTag& Tag = *TagIt;
 					int32& TagCount = GetReplicatedLooseTags_Mutable().TagMap.FindOrAdd(Tag);
-					TagCount += CountDelta;
+					GetReplicatedLooseTags_Mutable().SetTagCount(Tag, FMath::Max(0, TagCount + CountDelta));
 				}
 			}
 		}
diff --git a/Engine/Plugins/Runtime/GameplayAbilities/Source/GameplayAbilities/Public/AbilitySystemComponent.h b/Engine/Plugins/Runtime/GameplayAbilities/Source/GameplayAbilities/Public/AbilitySystemComponent.h
--- a/Engine/Plugins/Runtime/GameplayAbilities/Source/GameplayAbilities/Public/AbilitySystemComponent.h
+++ b/Engine/Plugins/Runtime/GameplayAbilities/Source/GameplayAbilities/Public/AbilitySystemComponent.h
@@ -615,7 +615,7 @@ class UAbilitySystemComponent : public UGameplayTasksComponent, public IGameplay
 	/** Forcibly sets the number of instances of a given tag */
 	inline void SetTagMapCount(const FGameplayTag& Tag, int32 NewCount, EGameplayTagReplicationState TagRepState = EGameplayTagReplicationState::None)
 	{
-		const int32 CurrentCount = GameplayTagCountContainer.GetTagCount(Tag);
+		const int32 CurrentCount = GameplayTagCountContainer.GetExplicitTagCount(Tag);
 		UpdateTagMapSingle_Internal(Tag, NewCount - CurrentCount, TagRepState);
 	}
 
