Updating the System.AbilitySystem.AbilitySystemComponent.ActivateAbilityFlow test to give the player controller a LocalPlayer, because that is now a requirement to identiy a local controller as of CL 44795882

#tests System.AbilitySystem.AbilitySystemComponent.ActivateAbilityFlow
#rnx

diff --git a/Engine/Plugins/Runtime/GameplayAbilities/Source/GameplayAbilities/Private/Tests/AbilitySystemComponentTests.cpp b/Engine/Plugins/Runtime/GameplayAbilities/Source/GameplayAbilities/Private/Tests/AbilitySystemComponentTests.cpp
--- a/Engine/Plugins/Runtime/GameplayAbilities/Source/GameplayAbilities/Private/Tests/AbilitySystemComponentTests.cpp
+++ b/Engine/Plugins/Runtime/GameplayAbilities/Source/GameplayAbilities/Private/Tests/AbilitySystemComponentTests.cpp
@@ -11,6 +11,7 @@
 #include "AbilitySystemTestPawn.h"
 #include "AbilitySystemGlobals.h"
 #include "NativeGameplayTags.h"
+#include "Engine/LocalPlayer.h"
 
 #if WITH_EDITOR
 
@@ -105,6 +106,8 @@ public:
 		SourceActor = World->SpawnActor<AAbilitySystemTestPawn>();
 		SourceController = World->SpawnActor<APlayerController>();
 		SourceController->Possess(SourceActor);
+		SourceLocalPlayer = NewObject<ULocalPlayer>(GEngine, ULocalPlayer::StaticClass());
+		SourceController->SetPlayer(SourceLocalPlayer);
 		SourceComponent = SourceActor->GetAbilitySystemComponent();
 
 		// set up the destination actor
@@ -229,7 +232,8 @@ private:
 
 	AAbilitySystemTestPawn* SourceActor;
 	UAbilitySystemComponent* SourceComponent;
-	AController* SourceController;
+	APlayerController* SourceController;
+	ULocalPlayer* SourceLocalPlayer;
 
 	AAbilitySystemTestPawn* DestActor;
 	UAbilitySystemComponent* DestComponent;
