Disabling TagCountContainer net serializer.

diff --git a/Engine/Plugins/Runtime/GameplayAbilities/Source/GameplayAbilities/Private/GameplayEffectTypes.cpp b/Engine/Plugins/Runtime/GameplayAbilities/Source/GameplayAbilities/Private/GameplayEffectTypes.cpp
--- a/Engine/Plugins/Runtime/GameplayAbilities/Source/GameplayAbilities/Private/GameplayEffectTypes.cpp
+++ b/Engine/Plugins/Runtime/GameplayAbilities/Source/GameplayAbilities/Private/GameplayEffectTypes.cpp
@@ -53,7 +53,7 @@ static FAutoConsoleVariableRef CVarUseMultiplyCompoundFix(
 
 TAutoConsoleVariable<bool> CVarReplicateTagCountContainerWithIris(
 		TEXT("AbilitySystem.Fix.ReplicateTagCountContainerWithIris"),
-		true,
+		false,
 		TEXT("Default: True.  When true will use Iris serialization for Gameplay Tag Count Container enabling tag rep state promotion with iris."),
 		ECVF_Default);
 
