Shoring up inconsistencies in tag counts on the server.

#FailedValidations Preflight

diff --git a/Engine/Plugins/Runtime/GameplayAbilities/Source/GameplayAbilities/Private/Abilities/GameplayAbility.cpp b/Engine/Plugins/Runtime/GameplayAbilities/Source/GameplayAbilities/Private/Abilities/GameplayAbility.cpp
--- a/Engine/Plugins/Runtime/GameplayAbilities/Source/GameplayAbilities/Private/Abilities/GameplayAbility.cpp
+++ b/Engine/Plugins/Runtime/GameplayAbilities/Source/GameplayAbilities/Private/Abilities/GameplayAbility.cpp
@@ -867,7 +867,7 @@ void UGameplayAbility::EndAbility(const FGameplayAbilitySpecHandle Handle, const
 			}
 
 			// Remove tags we added to owner
-			AbilitySystemComponent->RemoveLooseGameplayTags(ActivationOwnedTags);
+			AbilitySystemComponent->RemoveLooseGameplayTags(ActivationOwnedTags, 1, UAbilitySystemGlobals::Get().ShouldReplicateActivationOwnedTags() ? EGameplayTagReplicationState::CountToOwner : EGameplayTagReplicationState::None);
 
 			// Remove tracked GameplayCues that we added
 			for (FGameplayTag& GameplayCueTag : TrackedGameplayCues)
diff --git a/Engine/Plugins/Runtime/GameplayAbilities/Source/GameplayAbilities/Private/AbilitySystemBlueprintLibrary.cpp b/Engine/Plugins/Runtime/GameplayAbilities/Source/GameplayAbilities/Private/AbilitySystemBlueprintLibrary.cpp
--- a/Engine/Plugins/Runtime/GameplayAbilities/Source/GameplayAbilities/Private/AbilitySystemBlueprintLibrary.cpp
+++ b/Engine/Plugins/Runtime/GameplayAbilities/Source/GameplayAbilities/Private/AbilitySystemBlueprintLibrary.cpp
@@ -1360,7 +1360,7 @@ bool UAbilitySystemBlueprintLibrary::RemoveLooseGameplayTags(AActor* Actor, cons
 {
 	if (UAbilitySystemComponent* AbilitySysComp = GetAbilitySystemComponent(Actor))
 	{
-		AbilitySysComp->RemoveLooseGameplayTags(GameplayTags, 1);
+		AbilitySysComp->RemoveLooseGameplayTags(GameplayTags, 1, bShouldReplicate ? EGameplayTagReplicationState::CountToOwner : EGameplayTagReplicationState::None);
 
 		return true;
 	}
@@ -1380,11 +1380,11 @@ bool UAbilitySystemBlueprintLibrary::AddGameplayTags(AActor* Actor, const FGamep
 	return false;
 }
 
-bool UAbilitySystemBlueprintLibrary::RemoveGameplayTags(AActor* Actor, const FGameplayTagContainer& GameplayTags)
+bool UAbilitySystemBlueprintLibrary::RemoveGameplayTags(AActor* Actor, const FGameplayTagContainer& GameplayTags, EGameplayTagReplicationState ReplicationRule)
 {
 	if (UAbilitySystemComponent* AbilitySysComp = GetAbilitySystemComponent(Actor))
 	{
-		AbilitySysComp->RemoveLooseGameplayTags(GameplayTags, 1);
+		AbilitySysComp->RemoveLooseGameplayTags(GameplayTags, 1, ReplicationRule);
 
 		return true;
 	}
diff --git a/Engine/Plugins/Runtime/GameplayAbilities/Source/GameplayAbilities/Private/AbilitySystemComponent.cpp b/Engine/Plugins/Runtime/GameplayAbilities/Source/GameplayAbilities/Private/AbilitySystemComponent.cpp
--- a/Engine/Plugins/Runtime/GameplayAbilities/Source/GameplayAbilities/Private/AbilitySystemComponent.cpp
+++ b/Engine/Plugins/Runtime/GameplayAbilities/Source/GameplayAbilities/Private/AbilitySystemComponent.cpp
@@ -649,12 +649,10 @@ void UAbilitySystemComponent::UpdateTagMapSingle_Internal(const FGameplayTag& Ta
 		{
 			if (TagRepState == EGameplayTagReplicationState::TagOnly)
 			{
-				if (CountDelta > 0)
-				{
-					GetMinimalReplicationTags_Mutable().AddTag(Tag);
-				}
+				int32& TagCount = GetMinimalReplicationTags_Mutable().TagMap.FindOrAdd(Tag);
+				GetMinimalReplicationTags_Mutable().SetTagCount(Tag, FMath::Max(0, TagCount + CountDelta));
 			}
-			else if (CountDelta < 0 || TagRepState > EGameplayTagReplicationState::TagOnly)
+			else if (TagRepState > EGameplayTagReplicationState::TagOnly)
 			{
 				int32& TagCount = GetReplicatedLooseTags_Mutable().TagMap.FindOrAdd(Tag);
 				GetReplicatedLooseTags_Mutable().SetTagCount(Tag, FMath::Max(0, TagCount + CountDelta));
@@ -673,9 +671,9 @@ void UAbilitySystemComponent::UpdateTagMapSingle_Internal(const FGameplayTag& Ta
 	else if (CountDelta < 0)
 	{
 		TArray<FDeferredTagChangeDelegate> DeferredTagChangeDelegates;
-		
-		if (GameplayTagCountContainer.UpdateTagCount(Tag, CountDelta, TagRepState))
+		if (GameplayTagCountContainer.UpdateTagCount_DeferredParentRemoval(Tag, CountDelta, DeferredTagChangeDelegates))
 		{
+			GameplayTagCountContainer.FillParentTags();
 			OnTagUpdated(Tag, false);
 			if (UE::Net::ShouldUseIrisReplication() && !CVarReplicateTagCountContainerWithIris.GetValueOnAnyThread())
 			{
@@ -696,6 +694,11 @@ void UAbilitySystemComponent::UpdateTagMapSingle_Internal(const FGameplayTag& Ta
 					}
 				}
 			}
+
+			for (FDeferredTagChangeDelegate& Delegate : DeferredTagChangeDelegates)
+			{
+				Delegate.Execute();
+			}
 		}
 	}
 	
@@ -715,14 +718,16 @@ void UAbilitySystemComponent::UpdateTagMap_Internal(const FGameplayTagContainer&
 		{
 			if (TagRepState == EGameplayTagReplicationState::TagOnly)
 			{
-				if (CountDelta > 0)
+				for (TArray<FGameplayTag>::TConstIterator TagIt = Container.CreateConstIterator(); TagIt; ++TagIt)
 				{
-					GetMinimalReplicationTags_Mutable().AddTags(Container);
+					const FGameplayTag& Tag = *TagIt;
+					int32& TagCount = GetMinimalReplicationTags_Mutable().TagMap.FindOrAdd(Tag);
+					GetMinimalReplicationTags_Mutable().SetTagCount(Tag, FMath::Max(0, TagCount + CountDelta));
 				}
 			}
-			else if (CountDelta < 0 || TagRepState > EGameplayTagReplicationState::TagOnly)
+			else if (TagRepState > EGameplayTagReplicationState::TagOnly)
 			{
-				for (auto TagIt = Container.CreateConstIterator(); TagIt; ++TagIt)
+				for (TArray<FGameplayTag>::TConstIterator TagIt = Container.CreateConstIterator(); TagIt; ++TagIt)
 				{
 					const FGameplayTag& Tag = *TagIt;
 					int32& TagCount = GetReplicatedLooseTags_Mutable().TagMap.FindOrAdd(Tag);
@@ -735,7 +740,7 @@ void UAbilitySystemComponent::UpdateTagMap_Internal(const FGameplayTagContainer&
 	// For removal, reorder calls so that FillParentTags is only called once
 	if (CountDelta > 0)
 	{
-		for (auto TagIt = Container.CreateConstIterator(); TagIt; ++TagIt)
+		for (TArray<FGameplayTag>::TConstIterator TagIt = Container.CreateConstIterator(); TagIt; ++TagIt)
 		{
 			const FGameplayTag& Tag = *TagIt;
 			if (GameplayTagCountContainer.UpdateTagCount(Tag, CountDelta, TagRepState))
@@ -751,7 +756,7 @@ void UAbilitySystemComponent::UpdateTagMap_Internal(const FGameplayTagContainer&
 		RemovedTags.Reserve(Container.Num()); // pre-allocate max number (if all are removed)
 		TArray<FDeferredTagChangeDelegate> DeferredTagChangeDelegates;
 
-		for (auto TagIt = Container.CreateConstIterator(); TagIt; ++TagIt)
+		for (TArray<FGameplayTag>::TConstIterator TagIt = Container.CreateConstIterator(); TagIt; ++TagIt)
 		{
 			const FGameplayTag& Tag = *TagIt;
 			if (GameplayTagCountContainer.UpdateTagCount_DeferredParentRemoval(Tag, CountDelta, DeferredTagChangeDelegates))
@@ -883,7 +888,7 @@ void UAbilitySystemComponent::ResetTagMap()
 
 void UAbilitySystemComponent::NotifyTagMap_StackCountChange(const FGameplayTagContainer& Container)
 {
-	for (auto TagIt = Container.CreateConstIterator(); TagIt; ++TagIt)
+	for (TArray<FGameplayTag>::TConstIterator TagIt = Container.CreateConstIterator(); TagIt; ++TagIt)
 	{
 		const FGameplayTag& Tag = *TagIt;
 		GameplayTagCountContainer.Notify_StackCountChange(Tag);
diff --git a/Engine/Plugins/Runtime/GameplayAbilities/Source/GameplayAbilities/Private/GameplayEffect.cpp b/Engine/Plugins/Runtime/GameplayAbilities/Source/GameplayAbilities/Private/GameplayEffect.cpp
--- a/Engine/Plugins/Runtime/GameplayAbilities/Source/GameplayAbilities/Private/GameplayEffect.cpp
+++ b/Engine/Plugins/Runtime/GameplayAbilities/Source/GameplayAbilities/Private/GameplayEffect.cpp
@@ -4872,8 +4872,8 @@ void FActiveGameplayEffectsContainer::RemoveActiveGameplayEffectGrantedTagsAndMo
 	if (IsNetAuthority() || bPredictionRejected || bIrisWithoutTagCountSerializer)
 	{
 		// Update gameplaytag count and broadcast delegate if we are at 0
-		Owner->UpdateTagMap(Effect.Spec.Def->GetGrantedTags(), -1, EGameplayTagReplicationState::None);
-		Owner->UpdateTagMap(Effect.Spec.DynamicGrantedTags, -1, EGameplayTagReplicationState::None);
+		Owner->UpdateTagMap(Effect.Spec.Def->GetGrantedTags(), -1, bIrisWithoutTagCountSerializer ? EGameplayTagReplicationState::TagOnly : EGameplayTagReplicationState::CountToOwner);
+		Owner->UpdateTagMap(Effect.Spec.DynamicGrantedTags, -1, bIrisWithoutTagCountSerializer ? EGameplayTagReplicationState::TagOnly : EGameplayTagReplicationState::CountToOwner);
 	}
 
 	// Update our owner with the blocked ability tags this GameplayEffect adds to them
diff --git a/Engine/Plugins/Runtime/GameplayAbilities/Source/GameplayAbilities/Public/AbilitySystemBlueprintLibrary.h b/Engine/Plugins/Runtime/GameplayAbilities/Source/GameplayAbilities/Public/AbilitySystemBlueprintLibrary.h
--- a/Engine/Plugins/Runtime/GameplayAbilities/Source/GameplayAbilities/Public/AbilitySystemBlueprintLibrary.h
+++ b/Engine/Plugins/Runtime/GameplayAbilities/Source/GameplayAbilities/Public/AbilitySystemBlueprintLibrary.h
@@ -495,7 +495,7 @@ public:
 
 	/** Manually removes a set of tags from a given actor */
 	UFUNCTION(BlueprintCallable, Category="Ability|GameplayEffect")
-	static UE_API bool RemoveGameplayTags(AActor* Actor, const FGameplayTagContainer& GameplayTags);
+	static UE_API bool RemoveGameplayTags(AActor* Actor, const FGameplayTagContainer& GameplayTags, EGameplayTagReplicationState ReplicationRule = EGameplayTagReplicationState::None);
 
 	/** Get the GameplayEffect definition from a GameplayEffectSpecHandle. */
 	UFUNCTION(BlueprintCallable, Category = "Ability|GameplayEffect")
diff --git a/Engine/Plugins/Runtime/GameplayAbilities/Source/GameplayAbilities/Public/AbilitySystemComponent.h b/Engine/Plugins/Runtime/GameplayAbilities/Source/GameplayAbilities/Public/AbilitySystemComponent.h
--- a/Engine/Plugins/Runtime/GameplayAbilities/Source/GameplayAbilities/Public/AbilitySystemComponent.h
+++ b/Engine/Plugins/Runtime/GameplayAbilities/Source/GameplayAbilities/Public/AbilitySystemComponent.h
@@ -657,14 +657,14 @@ class UAbilitySystemComponent : public UGameplayTasksComponent, public IGameplay
 		UpdateTagMap(GameplayTags, Count, TagRepState);
 	}
 
-	inline void RemoveLooseGameplayTag(const FGameplayTag& GameplayTag, int32 Count = 1)
+	inline void RemoveLooseGameplayTag(const FGameplayTag& GameplayTag, int32 Count = 1, EGameplayTagReplicationState TagRepState = EGameplayTagReplicationState::None)
 	{
-		UpdateTagMap(GameplayTag, -Count, EGameplayTagReplicationState::None);
+		UpdateTagMap(GameplayTag, -Count, TagRepState);
 	}
 
-	inline void RemoveLooseGameplayTags(const FGameplayTagContainer& GameplayTags, int32 Count = 1)
+	inline void RemoveLooseGameplayTags(const FGameplayTagContainer& GameplayTags, int32 Count = 1, EGameplayTagReplicationState TagRepState = EGameplayTagReplicationState::None)
 	{
-		UpdateTagMap(GameplayTags, -Count, EGameplayTagReplicationState::None);
+		UpdateTagMap(GameplayTags, -Count, TagRepState);
 	}
 
 	inline void SetLooseGameplayTagCount(const FGameplayTag& GameplayTag, int32 NewCount, EGameplayTagReplicationState TagRepState = EGameplayTagReplicationState::None)
