[Engine] AbilityTask added PreDestroyFromReplication. Speculative fix to prevent AbilityTasks that have dangling pointers to their GameplayAbility that issued them. Normally handled by GameplayAbility EndAbility with EndTask's call to TaskOwnerEnded.

#tests Used FlowberryFizz until out of ammo in PIE shooting range to trigger SetNextMontageSectionandWait with 2 players. Garbage reference warning encountered 0/18 times. Initial callout encountered 3/30 times. With fix neither the garbage reference nor the PreDestroyFromReplication was called 0/12 times. Attempted intentionally causing the client desync by shutting off the logic issuing the Task from the client, but the task did not use PreDestroyFromReplication 0/8 times and the player returned to their expected state after FlowberryFizz ended.

#rnx

diff --git a/Engine/Plugins/Runtime/GameplayAbilities/Source/GameplayAbilities/Private/Abilities/Tasks/AbilityTask.cpp b/Engine/Plugins/Runtime/GameplayAbilities/Source/GameplayAbilities/Private/Abilities/Tasks/AbilityTask.cpp
--- a/Engine/Plugins/Runtime/GameplayAbilities/Source/GameplayAbilities/Private/Abilities/Tasks/AbilityTask.cpp
+++ b/Engine/Plugins/Runtime/GameplayAbilities/Source/GameplayAbilities/Private/Abilities/Tasks/AbilityTask.cpp
@@ -163,6 +163,12 @@ void UAbilityTask::BeginDestroy()
 	}
 }
 
+void UAbilityTask::PreDestroyFromReplication()
+{
+	// Clear the pointer without causing sideffects that a specialized EndTask() could cause.
+	Ability = nullptr;
+}
+
 FGameplayAbilitySpecHandle UAbilityTask::GetAbilitySpecHandle() const
 {
 	return Ability ? Ability->GetCurrentAbilitySpecHandle() : FGameplayAbilitySpecHandle();
diff --git a/Engine/Plugins/Runtime/GameplayAbilities/Source/GameplayAbilities/Public/Abilities/Tasks/AbilityTask.h b/Engine/Plugins/Runtime/GameplayAbilities/Source/GameplayAbilities/Public/Abilities/Tasks/AbilityTask.h
--- a/Engine/Plugins/Runtime/GameplayAbilities/Source/GameplayAbilities/Public/Abilities/Tasks/AbilityTask.h
+++ b/Engine/Plugins/Runtime/GameplayAbilities/Source/GameplayAbilities/Public/Abilities/Tasks/AbilityTask.h
@@ -93,6 +93,12 @@ class UAbilityTask : public UGameplayTask
 
 	UE_API virtual void OnDestroy(bool bInOwnerFinished) override;
 	UE_API virtual void BeginDestroy() override;
+
+	/**
+	 * Clears potential dangling pointers. Added to prevent dangling pointer to GameplayAbility that issued the task.
+	 * When overriding use EndTask() instead.
+	 */
+	UE_API virtual void PreDestroyFromReplication() override;
 	
 	/** Returns spec handle for owning ability */
 	UE_API FGameplayAbilitySpecHandle GetAbilitySpecHandle() const;
