Get World can return null, assume that we're standalone if that ever happens.

diff --git a/Engine/Plugins/Runtime/GameplayAbilities/Source/GameplayAbilities/Private/AbilitySystemComponent.cpp b/Engine/Plugins/Runtime/GameplayAbilities/Source/GameplayAbilities/Private/AbilitySystemComponent.cpp
--- a/Engine/Plugins/Runtime/GameplayAbilities/Source/GameplayAbilities/Private/AbilitySystemComponent.cpp
+++ b/Engine/Plugins/Runtime/GameplayAbilities/Source/GameplayAbilities/Private/AbilitySystemComponent.cpp
@@ -644,7 +644,8 @@ void UAbilitySystemComponent::UpdateTagMapSingle_Internal(const FGameplayTag& Ta
 	if (UE::Net::ShouldUseIrisReplication() && !CVarReplicateTagCountContainerWithIris.GetValueOnAnyThread())
 	{
 		// Iris replication doesn't work the same so we're falling back to using the Minimal and replicated loose tags until it's fixed.
-		ENetMode NetMode = GetWorld()->GetNetMode();
+		const UWorld* World = GetWorld();
+		ENetMode NetMode = World ? World->GetNetMode() : NM_Standalone;
 		if (NetMode == NM_DedicatedServer || NetMode == NM_ListenServer)
 		{
 			if (TagRepState == EGameplayTagReplicationState::TagOnly)
@@ -678,7 +679,8 @@ void UAbilitySystemComponent::UpdateTagMapSingle_Internal(const FGameplayTag& Ta
 			if (UE::Net::ShouldUseIrisReplication() && !CVarReplicateTagCountContainerWithIris.GetValueOnAnyThread())
 			{
 				// Iris replication doesn't work the same so we're falling back to using the Minimal and replicated loose tags until it's fixed.
-				ENetMode NetMode = GetWorld()->GetNetMode();
+				const UWorld* World = GetWorld();
+				ENetMode NetMode = World ? World->GetNetMode() : NM_Standalone;
 				if (NetMode == NM_DedicatedServer || NetMode == NM_ListenServer)
 				{
 					FMinimalReplicationTagCountMap& MTagMap = GetMinimalReplicationTags_Mutable();
@@ -713,7 +715,8 @@ void UAbilitySystemComponent::UpdateTagMap_Internal(const FGameplayTagContainer&
 	if (UE::Net::ShouldUseIrisReplication() && !CVarReplicateTagCountContainerWithIris.GetValueOnAnyThread())
 	{
 		// Iris replication doesn't work the same so we're falling back to using the Minimal and replicated loose tags until it's fixed.
-		ENetMode NetMode = GetWorld()->GetNetMode();
+		const UWorld* World = GetWorld();
+		ENetMode NetMode = World ? World->GetNetMode() : NM_Standalone;
 		if (NetMode == NM_DedicatedServer || NetMode == NM_ListenServer)
 		{
 			if (TagRepState == EGameplayTagReplicationState::TagOnly)
@@ -765,7 +768,8 @@ void UAbilitySystemComponent::UpdateTagMap_Internal(const FGameplayTagContainer&
 				if (UE::Net::ShouldUseIrisReplication() && !CVarReplicateTagCountContainerWithIris.GetValueOnAnyThread())
 				{
 					// Iris replication doesn't work the same so we're falling back to using the Minimal and replicated loose tags until it's fixed.
-					ENetMode NetMode = GetWorld()->GetNetMode();
+					const UWorld* World = GetWorld();
+					ENetMode NetMode = World ? World->GetNetMode() : NM_Standalone;
 					if (NetMode == NM_DedicatedServer || NetMode == NM_ListenServer)
 					{
 						FMinimalReplicationTagCountMap& MTagMap = GetMinimalReplicationTags_Mutable();
