Fix for incorrect Deleter being used on FGameplayAbilityTargetData when calling FilterTargetData, and exposed FGameplayAbilityTargetDataDeleter for this purpose

#tests Built and tested on Windows Server with Cooked Data
[RN] When creating TSharedPtr to FGameplayAbilityTargetData, please use the FGameplayAbilityTargetDataDeleter

diff --git a/Engine/Plugins/Runtime/GameplayAbilities/Source/GameplayAbilities/Private/AbilitySystemBlueprintLibrary.cpp b/Engine/Plugins/Runtime/GameplayAbilities/Source/GameplayAbilities/Private/AbilitySystemBlueprintLibrary.cpp
--- a/Engine/Plugins/Runtime/GameplayAbilities/Source/GameplayAbilities/Private/AbilitySystemBlueprintLibrary.cpp
+++ b/Engine/Plugins/Runtime/GameplayAbilities/Source/GameplayAbilities/Private/AbilitySystemBlueprintLibrary.cpp
@@ -446,7 +446,7 @@ FGameplayAbilityTargetDataHandle UAbilitySystemBlueprintLibrary::FilterTargetDat
 				FGameplayAbilityTargetData* NewData = (FGameplayAbilityTargetData*)FMemory::Malloc(ScriptStruct->GetCppStructOps()->GetSize());
 				ScriptStruct->InitializeStruct(NewData);
 				ScriptStruct->CopyScriptStruct(NewData, UnfilteredData);
-				ReturnDataHandle.Data.Add(TSharedPtr<FGameplayAbilityTargetData>(NewData));
+				ReturnDataHandle.Data.Add(TSharedPtr<FGameplayAbilityTargetData>(NewData, FGameplayAbilityTargetDataDeleter()));
 				if (FilteredActors.Num() < UnfilteredActors.Num())
 				{
 					//We have lost some, but not all, of our actors, so replace the array. This should only be possible with targeting types that permit actor-array setting.
diff --git a/Engine/Plugins/Runtime/GameplayAbilities/Source/GameplayAbilities/Private/GameplayAbilityTargetTypes.cpp b/Engine/Plugins/Runtime/GameplayAbilities/Source/GameplayAbilities/Private/GameplayAbilityTargetTypes.cpp
--- a/Engine/Plugins/Runtime/GameplayAbilities/Source/GameplayAbilities/Private/GameplayAbilityTargetTypes.cpp
+++ b/Engine/Plugins/Runtime/GameplayAbilities/Source/GameplayAbilities/Private/GameplayAbilityTargetTypes.cpp
@@ -182,17 +182,15 @@ FGameplayAbilityTargetDataHandle FGameplayAbilityTargetingLocationInfo::MakeTarg
 #define TARGETDATAHANDLE_SAFE_NET_SERIALIZE 1
 #endif
 
-struct FGameplayAbilityTargetDataDeleter
+
+void FGameplayAbilityTargetDataDeleter::operator()(FGameplayAbilityTargetData* Object) const
 {
-	FORCEINLINE void operator()(FGameplayAbilityTargetData* Object) const
-	{
-		check(Object);
-		UScriptStruct* ScriptStruct = Object->GetScriptStruct();
-		check(ScriptStruct);
-		ScriptStruct->DestroyStruct(Object);
-		FMemory::Free(Object);
-	}
-};
+	check(Object);
+	UScriptStruct* ScriptStruct = Object->GetScriptStruct();
+	check(ScriptStruct);
+	ScriptStruct->DestroyStruct(Object);
+	FMemory::Free(Object);
+}
 
 bool FGameplayAbilityTargetDataHandle::NetSerialize(FArchive& Ar, class UPackageMap* Map, bool& bOutSuccess)
 {
diff --git a/Engine/Plugins/Runtime/GameplayAbilities/Source/GameplayAbilities/Public/Abilities/GameplayAbilityTargetTypes.h b/Engine/Plugins/Runtime/GameplayAbilities/Source/GameplayAbilities/Public/Abilities/GameplayAbilityTargetTypes.h
--- a/Engine/Plugins/Runtime/GameplayAbilities/Source/GameplayAbilities/Public/Abilities/GameplayAbilityTargetTypes.h
+++ b/Engine/Plugins/Runtime/GameplayAbilities/Source/GameplayAbilities/Public/Abilities/GameplayAbilityTargetTypes.h
@@ -42,6 +42,15 @@ namespace EGameplayTargetingConfirmation
 	};
 }
 
+/** 
+ * This type should be provided as the Deleter when creating a TSharedPtr<FGameplayAbilityTargetData> so that
+ * the Object's ScriptStruct data will be destructed correctly.
+ */
+struct FGameplayAbilityTargetDataDeleter
+{
+	UE_API void operator()(FGameplayAbilityTargetData* Object) const;
+};
+
 /**
  *	A generic structure for targeting data. We want generic functions to produce this data and other generic
  *	functions to consume this data.
