[GAS] Using Generic Replication, the GameplayTagCountContainer can get stuck in a state where it constantly replicates a None tag if any of the tags were set to non-replicated.

diff --git a/Engine/Plugins/Runtime/GameplayAbilities/Source/GameplayAbilities/Private/GameplayEffectTypes.cpp b/Engine/Plugins/Runtime/GameplayAbilities/Source/GameplayAbilities/Private/GameplayEffectTypes.cpp
--- a/Engine/Plugins/Runtime/GameplayAbilities/Source/GameplayAbilities/Private/GameplayEffectTypes.cpp
+++ b/Engine/Plugins/Runtime/GameplayAbilities/Source/GameplayAbilities/Private/GameplayEffectTypes.cpp
@@ -559,53 +559,56 @@ bool FGameplayTagCountContainer::NetDeltaSerialize(FNetDeltaSerializeInfo & Delt
 
 		check(DeltaParams.NewState);
 		*DeltaParams.NewState = TSharedPtr<INetDeltaBaseState>(NewState);
-		NewState->TagStates.SetNum(Items.Num());
+		NewState->TagStates.Reserve(Items.Num());
 		ChangeIdx.Reserve(Items.Num());
 
 		// We set ChangelistHistory and LastAckedHistory to 0, as that's the default RepLayout expects.
 		// Setting them to anything else may cause issues with initial sends.
 		uint32 OldChangelistHistory = 0;
 		uint32 OldLastAckedHistory = 0;
-
+		
 		TArray<FNetGameplayTagCountContainerState::FGameplayTagCountStateItem> Removals = OldState ? OldState->TagStates : TArray<FNetGameplayTagCountContainerState::FGameplayTagCountStateItem>{};
 
 		// Iterate over our current items and build our new state data while removing valid entries from the Removals list so we know what's missing
-		for (int32 Idx = 0; Idx < Items.Num(); Idx++)
+		for (const FGameplayTagCountItem& CountItem : Items)
 		{
-			FGameplayTagCountItem& CountItem = Items[Idx];
-			FNetGameplayTagCountContainerState::FGameplayTagCountStateItem& NewItem = NewState->TagStates[Idx];
+			// Find this Item in the OldState
+			FNetGameplayTagCountContainerState::FGameplayTagCountStateItem* OldItem = Removals.FindByPredicate(
+				[&CountItem](FNetGameplayTagCountContainerState::FGameplayTagCountStateItem& Item)
+				{
+					return Item.Tag == CountItem.Tag;
+				});
+
+			// Warn if the ReplicationState ever changed; it's an indication that the code that uses it has conflicting tag count requirements
+			UE_CLOG(OldItem && OldItem->ReplicationState > CountItem.ReplicationState, LogAbilitySystem, Warning,
+				TEXT("Existing tag: %s with \"%s\" replication state now serializing with state \"%s\""), 
+				*CountItem.Tag.ToString(), *UEnum::GetDisplayValueAsText(OldItem->ReplicationState).ToString(), *UEnum::GetDisplayValueAsText(CountItem.ReplicationState).ToString());
 
 			if (CountItem.ReplicationState == EGameplayTagReplicationState::None)
 			{
+				// We don't replicate this tag, so don't save its state
 				continue;
 			}
-		
-			bPackCount = CountItem.ReplicationState >= (bIsOwner ? EGameplayTagReplicationState::CountToOwner : EGameplayTagReplicationState::TagAndCountToAll);
 
+			// Record the NewState for the CountItem
+			FNetGameplayTagCountContainerState::FGameplayTagCountStateItem& NewItem = NewState->TagStates.AddDefaulted_GetRef();
 			NewItem.Tag = CountItem.Tag;
 			NewItem.ReplicationState = CountItem.ReplicationState;
-			NewItem.Count = bPackCount ? CountItem.Count : 1;
 
-			FNetGameplayTagCountContainerState::FGameplayTagCountStateItem* OldItem = Removals.FindByPredicate(
-				[NewItem] (FNetGameplayTagCountContainerState::FGameplayTagCountStateItem& Item)
-			{
-				return Item.Tag == NewItem.Tag;
-			});
+			bPackCount = CountItem.ReplicationState >= (bIsOwner ? EGameplayTagReplicationState::CountToOwner : EGameplayTagReplicationState::TagAndCountToAll);
+			NewItem.Count = bPackCount ? CountItem.Count : 1;
 			
 			if (!OldItem)
 			{
-				ChangeIdx.Add(Idx);
+				// This was a tag added
+				ChangeIdx.Add(NewState->TagStates.Num() - 1);
 				continue;
 			}
 
-			if (OldItem->ReplicationState > NewItem.ReplicationState)
-			{
-				ABILITY_LOG(Warning, TEXT("Existing tag: %s with \"%s\" replication state now serializing with state \"%s\""), 
-				*NewItem.Tag.ToString(), *UEnum::GetDisplayValueAsText(OldItem->ReplicationState).ToString(), *UEnum::GetDisplayValueAsText(NewItem.ReplicationState).ToString());
-			}
 			if (bPackCount && OldItem->Count != NewItem.Count)
 			{
-				ChangeIdx.Add(Idx);
+				// This was a tag count change
+				ChangeIdx.Add(NewState->TagStates.Num() - 1);
 			}
 
 			Removals.RemoveSingleSwap(*OldItem);
