Unshelved from pending changelist '47949096':

Force net update when a predicted Gameplay Effect is applied so that short duration effects are processed correctly.

#swarm https://p4-swarm.epicgames.net/reviews/47949096
[RN] Fixed tags getting stuck on clients for short duration predicted GEs

diff --git a/Engine/Plugins/Runtime/GameplayAbilities/Source/GameplayAbilities/Private/Abilities/GameplayAbility.cpp b/Engine/Plugins/Runtime/GameplayAbilities/Source/GameplayAbilities/Private/Abilities/GameplayAbility.cpp
--- a/Engine/Plugins/Runtime/GameplayAbilities/Source/GameplayAbilities/Private/Abilities/GameplayAbility.cpp
+++ b/Engine/Plugins/Runtime/GameplayAbilities/Source/GameplayAbilities/Private/Abilities/GameplayAbility.cpp
@@ -2057,8 +2057,14 @@ FActiveGameplayEffectHandle UGameplayAbility::ApplyGameplayEffectSpecToOwner(con
 	if (SpecHandle.IsValid() && (HasAuthorityOrPredictionKey(ActorInfo, &ActivationInfo)))
 	{
 		UAbilitySystemComponent* const AbilitySystemComponent = ActorInfo->AbilitySystemComponent.Get();
-		return AbilitySystemComponent->ApplyGameplayEffectSpecToSelf(*SpecHandle.Data.Get(), AbilitySystemComponent->GetPredictionKeyForNewAction());
+		FActiveGameplayEffectHandle ReturnHandle = AbilitySystemComponent->ApplyGameplayEffectSpecToSelf(*SpecHandle.Data.Get(), AbilitySystemComponent->GetPredictionKeyForNewAction());
 
+		// Force a net update for predicted effects so that short durations are replicated correctly
+		if (ReturnHandle.IsValid() && ActorInfo && ActorInfo->OwnerActor.IsValid() && ActorInfo->OwnerActor->HasAuthority())
+		{
+			AbilitySystemComponent->GetOwner()->ForceNetUpdate();
+		}
+		return ReturnHandle;
 	}
 	return FActiveGameplayEffectHandle();
 }
