Add API call for getting the next montage section name.

preflight 68bf3a079910076443bda836 failed on unrelated errors

diff --git a/Engine/Plugins/Runtime/GameplayAbilities/Source/GameplayAbilities/Private/AbilitySystemComponent_Abilities.cpp b/Engine/Plugins/Runtime/GameplayAbilities/Source/GameplayAbilities/Private/AbilitySystemComponent_Abilities.cpp
--- a/Engine/Plugins/Runtime/GameplayAbilities/Source/GameplayAbilities/Private/AbilitySystemComponent_Abilities.cpp
+++ b/Engine/Plugins/Runtime/GameplayAbilities/Source/GameplayAbilities/Private/AbilitySystemComponent_Abilities.cpp
@@ -3543,6 +3543,21 @@ void UAbilitySystemComponent::CurrentMontageSetNextSectionName(FName FromSection
 	}
 }
 
+FName UAbilitySystemComponent::CurrentMontageGetNextSectionName() const
+{
+	UAnimInstance* AnimInstance = AbilityActorInfo.IsValid() ? AbilityActorInfo->GetAnimInstance() : nullptr;
+	if (LocalAnimMontageInfo.AnimMontage && AnimInstance)
+	{
+		FAnimMontageInstance* MontageInstance = AnimInstance->GetActiveInstanceForMontage(LocalAnimMontageInfo.AnimMontage);
+		if (MontageInstance)
+		{
+			return MontageInstance->GetNextSection();
+		}
+	}
+
+	return NAME_None;
+}
+
 void UAbilitySystemComponent::CurrentMontageSetPlayRate(float InPlayRate)
 {
 	UAnimInstance* AnimInstance = AbilityActorInfo.IsValid() ? AbilityActorInfo->GetAnimInstance() : nullptr;
diff --git a/Engine/Plugins/Runtime/GameplayAbilities/Source/GameplayAbilities/Public/AbilitySystemComponent.h b/Engine/Plugins/Runtime/GameplayAbilities/Source/GameplayAbilities/Public/AbilitySystemComponent.h
--- a/Engine/Plugins/Runtime/GameplayAbilities/Source/GameplayAbilities/Public/AbilitySystemComponent.h
+++ b/Engine/Plugins/Runtime/GameplayAbilities/Source/GameplayAbilities/Public/AbilitySystemComponent.h
@@ -1465,6 +1465,9 @@ class UAbilitySystemComponent : public UGameplayTasksComponent, public IGameplay
 	/** Sets current montages next section name. Expectation is caller should only be stopping it if they are the current animating ability (or have good reason not to check) */
 	UE_API virtual void CurrentMontageSetNextSectionName(FName FromSectionName, FName ToSectionName);
 
+	/** Returns the next section name for the current montage. */
+	UE_API FName CurrentMontageGetNextSectionName() const;
+
 	/** Sets current montage's play rate */
 	UE_API virtual void CurrentMontageSetPlayRate(float InPlayRate);
 
