Make AreAbilityTagsBlocked BlueprintCallable

diff --git a/Engine/Plugins/Runtime/GameplayAbilities/Source/GameplayAbilities/Public/AbilitySystemComponent.h b/Engine/Plugins/Runtime/GameplayAbilities/Source/GameplayAbilities/Public/AbilitySystemComponent.h
--- a/Engine/Plugins/Runtime/GameplayAbilities/Source/GameplayAbilities/Public/AbilitySystemComponent.h
+++ b/Engine/Plugins/Runtime/GameplayAbilities/Source/GameplayAbilities/Public/AbilitySystemComponent.h
@@ -1077,6 +1077,7 @@ class UAbilitySystemComponent : public UGameplayTasksComponent, public IGameplay
 	virtual void HandleChangeAbilityCanBeCanceled(const FGameplayTagContainer& AbilityTags, UGameplayAbility* RequestingAbility, bool bCanBeCanceled) {}
 
 	/** Returns true if any passed in tags are blocked */
+	UFUNCTION(BlueprintCallable, Category="Abilities")
 	UE_API virtual bool AreAbilityTagsBlocked(const FGameplayTagContainer& Tags) const;
 
 	/** Block or cancel blocking for specific ability tags */
