Fix GC nullptr crash
#rnx

diff --git a/Engine/Plugins/Runtime/GameplayAbilities/Source/GameplayAbilities/Private/AbilitySystemComponent.cpp b/Engine/Plugins/Runtime/GameplayAbilities/Source/GameplayAbilities/Private/AbilitySystemComponent.cpp
--- a/Engine/Plugins/Runtime/GameplayAbilities/Source/GameplayAbilities/Private/AbilitySystemComponent.cpp
+++ b/Engine/Plugins/Runtime/GameplayAbilities/Source/GameplayAbilities/Private/AbilitySystemComponent.cpp
@@ -644,7 +644,8 @@ void UAbilitySystemComponent::UpdateTagMapSingle_Internal(const FGameplayTag& Ta
 	if (UE::Net::ShouldUseIrisReplication() && !CVarReplicateTagCountContainerWithIris.GetValueOnAnyThread())
 	{
 		// Iris replication doesn't work the same so we're falling back to using the Minimal and replicated loose tags until it's fixed.
-		ENetMode NetMode = GetWorld()->GetNetMode();
+		const UWorld* World = GetWorld();
+		ENetMode NetMode = World ? World->GetNetMode() : NM_Standalone;
 		if (NetMode == NM_DedicatedServer || NetMode == NM_ListenServer)
 		{
 			if (TagRepState == EGameplayTagReplicationState::TagOnly)
