ReloadAttributeDefaults when the curve table changes, such as after applying a hotfix

diff --git a/Engine/Plugins/Runtime/GameplayAbilities/Source/GameplayAbilities/Private/AbilitySystemGlobals.cpp b/Engine/Plugins/Runtime/GameplayAbilities/Source/GameplayAbilities/Private/AbilitySystemGlobals.cpp
--- a/Engine/Plugins/Runtime/GameplayAbilities/Source/GameplayAbilities/Private/AbilitySystemGlobals.cpp
+++ b/Engine/Plugins/Runtime/GameplayAbilities/Source/GameplayAbilities/Private/AbilitySystemGlobals.cpp
@@ -490,6 +490,10 @@ void UAbilitySystemGlobals::AddAttributeDefaultTables(const FName OwnerName, con
 			if (AttribTable)
 			{
 				GlobalAttributeDefaultsTables.AddUnique(AttribTable);
+				if (!AttribTable->OnCurveTableChanged().IsBoundToObject(this))
+				{
+					AttribTable->OnCurveTableChanged().AddUObject(this, &ThisClass::OnAttributeDefaultTableChanged);
+				}
 				bModified = true;
 			}
 		}
@@ -526,6 +530,7 @@ void UAbilitySystemGlobals::RemoveAttributeDefaultTables(const FName OwnerName,
 						{
 							if (GlobalAttributeDefaultsTables.Remove(AttribTable) > 0)
 							{
+								AttribTable->OnCurveTableChanged().RemoveAll(this);
 								bModified = true;
 							}
 						}
@@ -541,6 +546,11 @@ void UAbilitySystemGlobals::RemoveAttributeDefaultTables(const FName OwnerName,
 	}
 }
 
+void UAbilitySystemGlobals::OnAttributeDefaultTableChanged()
+{
+	ReloadAttributeDefaults();
+}
+
 TArray<FSoftObjectPath> UAbilitySystemGlobals::GetGlobalAttributeSetDefaultsTablePaths() const
 {
 	TArray<FSoftObjectPath> AttribSetDefaultsTables;
@@ -570,6 +580,10 @@ void UAbilitySystemGlobals::InitAttributeDefaults()
 			if (ensureMsgf(AttribTable, TEXT("Could not load Global AttributeSet Defaults Table: %s"), *AttribSetDefaultsTablePath.ToString()))
 			{
 				GlobalAttributeDefaultsTables.AddUnique(AttribTable);
+				if (!AttribTable->OnCurveTableChanged().IsBoundToObject(this))
+				{
+					AttribTable->OnCurveTableChanged().AddUObject(this, &ThisClass::OnAttributeDefaultTableChanged);
+				}
 			}
 		}
 	}
diff --git a/Engine/Plugins/Runtime/GameplayAbilities/Source/GameplayAbilities/Public/AbilitySystemGlobals.h b/Engine/Plugins/Runtime/GameplayAbilities/Source/GameplayAbilities/Public/AbilitySystemGlobals.h
--- a/Engine/Plugins/Runtime/GameplayAbilities/Source/GameplayAbilities/Public/AbilitySystemGlobals.h
+++ b/Engine/Plugins/Runtime/GameplayAbilities/Source/GameplayAbilities/Public/AbilitySystemGlobals.h
@@ -433,6 +433,8 @@ protected:
 	template <class T>
 	T* InternalGetLoadTable(T*& Table, FString TableName);
 
+	void OnAttributeDefaultTableChanged();
+
 #if WITH_EDITOR
 	UE_API void OnTableReimported(UObject* InObject);
 	UE_API void OnPreBeginPIE(const bool bIsSimulatingInEditor);
