PR #13301: Remove LogTemp in FindAllAbilitySpecsFromInputID

#rnx

diff --git a/Engine/Plugins/Runtime/GameplayAbilities/Source/GameplayAbilities/Private/AbilitySystemComponent_Abilities.cpp b/Engine/Plugins/Runtime/GameplayAbilities/Source/GameplayAbilities/Private/AbilitySystemComponent_Abilities.cpp
--- a/Engine/Plugins/Runtime/GameplayAbilities/Source/GameplayAbilities/Private/AbilitySystemComponent_Abilities.cpp
+++ b/Engine/Plugins/Runtime/GameplayAbilities/Source/GameplayAbilities/Private/AbilitySystemComponent_Abilities.cpp
@@ -1025,8 +1025,6 @@ void UAbilitySystemComponent::FindAllAbilitySpecsFromInputID(int32 InputID, TArr
 			}
 		}
 	}
-
-	UE_LOG(LogTemp, Log, TEXT("OutAbilitySpecs count: %i"), OutAbilitySpecs.Num());
 }
 
 FGameplayAbilitySpec UAbilitySystemComponent::BuildAbilitySpecFromClass(TSubclassOf<UGameplayAbility> AbilityClass, int32 Level /*= 0*/, int32 InputID /*= -1*/)
