Don't attempt to draw debug info if the ASC Actor Info is invalid.

diff --git a/Engine/Plugins/Runtime/GameplayAbilities/Source/GameplayAbilities/Private/AbilitySystemDebugHUD.cpp b/Engine/Plugins/Runtime/GameplayAbilities/Source/GameplayAbilities/Private/AbilitySystemDebugHUD.cpp
--- a/Engine/Plugins/Runtime/GameplayAbilities/Source/GameplayAbilities/Private/AbilitySystemDebugHUD.cpp
+++ b/Engine/Plugins/Runtime/GameplayAbilities/Source/GameplayAbilities/Private/AbilitySystemDebugHUD.cpp
@@ -582,6 +582,12 @@ void AAbilitySystemDebugHUD::DrawAbilityDebugInfo(UCanvas* InCanvas, APlayerCont
 		{
 			if (ASC->GetWorld() == PCWorld)
 			{
+				// We need to have actor info for valid debug info.
+				if (!ASC->AbilityActorInfo.IsValid())
+				{
+					continue;
+				}
+				
 				const AActor* AvatarActor = ASC->GetAvatarActor();
 				if (!ShouldDisplayDebugForActor(InCanvas, AvatarActor, CameraPosition, CameraDir))
 				{
