Add error if UAssetTagsGameplayEffectComponent tries to update its inheritable tags with reference to itself.
Speculative fix to avoid modifying array while iterating it in UpdateInheritedTagProperties.

diff --git a/Engine/Plugins/Runtime/GameplayAbilities/Source/GameplayAbilities/Private/GameplayEffectComponents/AssetTagsGameplayEffectComponent.cpp b/Engine/Plugins/Runtime/GameplayAbilities/Source/GameplayAbilities/Private/GameplayEffectComponents/AssetTagsGameplayEffectComponent.cpp
--- a/Engine/Plugins/Runtime/GameplayAbilities/Source/GameplayAbilities/Private/GameplayEffectComponents/AssetTagsGameplayEffectComponent.cpp
+++ b/Engine/Plugins/Runtime/GameplayAbilities/Source/GameplayAbilities/Private/GameplayEffectComponents/AssetTagsGameplayEffectComponent.cpp
@@ -1,6 +1,8 @@
 // Copyright Epic Games, Inc. All Rights Reserved.
 
 #include "GameplayEffectComponents/AssetTagsGameplayEffectComponent.h"
+#include "AbilitySystemLog.h"
+#include "Logging/StructuredLog.h"
 
 #include UE_INLINE_GENERATED_CPP_BY_NAME(AssetTagsGameplayEffectComponent)
 
@@ -13,7 +15,14 @@ void UAssetTagsGameplayEffectComponent::PostInitProperties()
 	// Try to find the parent and update the inherited tags.  We do this 'early' because this is the only function
 	// we can use for creation of the object (and this runs post-constructor).
 	const UAssetTagsGameplayEffectComponent* Parent = FindParentComponent(*this);
-	InheritableAssetTags.UpdateInheritedTagProperties(Parent ? &Parent->InheritableAssetTags : nullptr);
+	if (Parent == this)
+	{
+		UE_LOGFMT(LogAbilitySystem, Warning, "UAssetTagsGameplayEffectComponent got self for parent compontnt when trying to update inherited tag container. {Component}", FSoftObjectPath(this));
+	}
+	else 
+	{
+		InheritableAssetTags.UpdateInheritedTagProperties(Parent ? &Parent->InheritableAssetTags : nullptr);
+	}
 
 #if WITH_EDITORONLY_DATA
 	if (GetOwner())
