[RemoteObjects] Early fixes to GAS to allow at least some migration of the AbilitySystemComponent.  Don't expect the abilities to actually work, but we won't crash!

#rnx

diff --git a/Engine/Plugins/Runtime/GameplayAbilities/Source/GameplayAbilities/Private/Abilities/GameplayAbility.cpp b/Engine/Plugins/Runtime/GameplayAbilities/Source/GameplayAbilities/Private/Abilities/GameplayAbility.cpp
--- a/Engine/Plugins/Runtime/GameplayAbilities/Source/GameplayAbilities/Private/Abilities/GameplayAbility.cpp
+++ b/Engine/Plugins/Runtime/GameplayAbilities/Source/GameplayAbilities/Private/Abilities/GameplayAbility.cpp
@@ -19,6 +19,10 @@
 #include "Misc/DataValidation.h"
 #include "UObject/Package.h"
 
+#if UE_WITH_REMOTE_OBJECT_HANDLE
+#include "UObject/UObjectMigrationContext.h"
+#endif
+
 #define LOCTEXT_NAMESPACE "GameplayAbility"
 
 #include "Iris/ReplicationSystem/ReplicationFragmentUtil.h"
@@ -2325,4 +2329,17 @@ void UGameplayAbility::RegisterReplicationFragments(UE::Net::FFragmentRegistrati
 	UE::Net::FReplicationFragmentUtil::CreateAndRegisterFragmentsForObject(this, Context, RegistrationFlags);
 }
 
+#if UE_WITH_REMOTE_OBJECT_HANDLE
+void UGameplayAbility::PostMigrate(const FUObjectMigrationContext& MigrationContext)
+{
+	Super::PostMigrate(MigrationContext);
+
+	if (MigrationContext.MigrationSide == EObjectMigrationSide::Receive)
+	{
+		// Setup our CurrentActorInfo
+		PostNetInit();
+	}
+}
+#endif
+
 #undef LOCTEXT_NAMESPACE
\ No newline at end of file
diff --git a/Engine/Plugins/Runtime/GameplayAbilities/Source/GameplayAbilities/Private/AbilitySystemComponent.cpp b/Engine/Plugins/Runtime/GameplayAbilities/Source/GameplayAbilities/Private/AbilitySystemComponent.cpp
--- a/Engine/Plugins/Runtime/GameplayAbilities/Source/GameplayAbilities/Private/AbilitySystemComponent.cpp
+++ b/Engine/Plugins/Runtime/GameplayAbilities/Source/GameplayAbilities/Private/AbilitySystemComponent.cpp
@@ -19,6 +19,7 @@
 #include "Net/Core/PushModel/PushModel.h"
 #include "HAL/IConsoleManager.h"
 #include "Iris/IrisConfig.h"
+#include "Serialization/Archive.h"
 
 #include UE_INLINE_GENERATED_CPP_BY_NAME(AbilitySystemComponent)
 
@@ -192,6 +193,12 @@ void UAbilitySystemComponent::OnRegister()
 {
 	Super::OnRegister();
 
+	/** Allocate an AbilityActorInfo. Note: this goes through a global function and is a SharedPtr so projects can make their own AbilityActorInfo */
+	if (!AbilityActorInfo.IsValid())
+	{
+		AbilityActorInfo = TSharedPtr<FGameplayAbilityActorInfo>(UAbilitySystemGlobals::Get().AllocAbilityActorInfo());
+	}
+
 	// Cached off netrole to avoid constant checking on owning actor
 	CacheIsNetSimulated();
 
@@ -222,12 +229,6 @@ PRAGMA_DISABLE_DEPRECATION_WARNINGS
 	ReplicatedLooseTags.Owner = this;
 PRAGMA_ENABLE_DEPRECATION_WARNINGS
 
-	/** Allocate an AbilityActorInfo. Note: this goes through a global function and is a SharedPtr so projects can make their own AbilityActorInfo */
-	if(!AbilityActorInfo.IsValid())
-	{
-		AbilityActorInfo = TSharedPtr<FGameplayAbilityActorInfo>(UAbilitySystemGlobals::Get().AllocAbilityActorInfo());
-	}
-
 	// Ensure bDestroyActiveStateInitiated is clear in case component is re-entering play 
 	bDestroyActiveStateInitiated = false;
 }
@@ -247,6 +248,57 @@ void UAbilitySystemComponent::BeginPlay()
 	CacheIsNetSimulated();
 }
 
+void UAbilitySystemComponent::Serialize(FArchive& Ar)
+{
+	Super::Serialize(Ar);
+
+	if (Ar.IsMigratingRemoteObjects())
+	{
+		if (Ar.IsSaving())
+		{
+			auto SaveTagToGASpecMap = [&Ar](TMap<FGameplayTag, TArray<FGameplayAbilitySpecHandle>>& Map)
+			{
+				int32 NumEntries = Map.Num();
+				Ar << NumEntries;
+				for (TPair<FGameplayTag, TArray<FGameplayAbilitySpecHandle>>& Pair : Map)
+				{
+					FGameplayTag& GameplayTag = Pair.Key;
+					TArray<FGameplayAbilitySpecHandle>& GASpecs = Pair.Value;
+
+					Ar << GameplayTag;
+					Ar << GASpecs;
+				}
+			};
+
+			SaveTagToGASpecMap(GameplayEventTriggeredAbilities);
+			SaveTagToGASpecMap(OwnedTagTriggeredAbilities);
+		}
+		else if (Ar.IsLoading())
+		{
+			auto LoadTagToGASpecMap = [&Ar](TMap<FGameplayTag, TArray<FGameplayAbilitySpecHandle>>& Map)
+			{
+				int32 NumEntries = Map.Num();
+				Ar << NumEntries;
+
+				Map.Reset();
+
+				FGameplayTag GameplayTag;
+				TArray<FGameplayAbilitySpecHandle> GASpecs;
+
+				for (int Idx = 0; Idx < NumEntries; ++Idx)
+				{
+					Ar << GameplayTag;
+					Ar << GASpecs;
+					Map.Emplace(GameplayTag, GASpecs);
+				}
+			};
+
+			LoadTagToGASpecMap(GameplayEventTriggeredAbilities);
+			LoadTagToGASpecMap(OwnedTagTriggeredAbilities);
+		}
+	}
+}
+
 void UAbilitySystemComponent::ReadyForReplication()
 {
 	Super::ReadyForReplication();
diff --git a/Engine/Plugins/Runtime/GameplayAbilities/Source/GameplayAbilities/Public/Abilities/GameplayAbility.h b/Engine/Plugins/Runtime/GameplayAbilities/Source/GameplayAbilities/Public/Abilities/GameplayAbility.h
--- a/Engine/Plugins/Runtime/GameplayAbilities/Source/GameplayAbilities/Public/Abilities/GameplayAbility.h
+++ b/Engine/Plugins/Runtime/GameplayAbilities/Source/GameplayAbilities/Public/Abilities/GameplayAbility.h
@@ -929,6 +929,10 @@ protected:
 	UE_DEPRECATED(5.4, "This is unsafe. Do not use.")
 	UPROPERTY(BlueprintReadOnly, Category = Ability, meta=(DeprecatedProperty, DeprecationMessage="This is unsafe. Do not use."))
 	bool bMarkPendingKillOnAbilityEnd;
+
+#if UE_WITH_REMOTE_OBJECT_HANDLE
+	virtual void PostMigrate(const struct FUObjectMigrationContext& MigrationContext) override;
+#endif
 };
 
 #undef UE_API
diff --git a/Engine/Plugins/Runtime/GameplayAbilities/Source/GameplayAbilities/Public/AbilitySystemComponent.h b/Engine/Plugins/Runtime/GameplayAbilities/Source/GameplayAbilities/Public/AbilitySystemComponent.h
--- a/Engine/Plugins/Runtime/GameplayAbilities/Source/GameplayAbilities/Public/AbilitySystemComponent.h
+++ b/Engine/Plugins/Runtime/GameplayAbilities/Source/GameplayAbilities/Public/AbilitySystemComponent.h
@@ -1647,6 +1647,7 @@ public:
 	UE_API virtual void OnUnregister() override;
 	UE_API virtual void ReadyForReplication() override;
 	UE_API virtual void BeginPlay() override;
+	UE_API virtual void Serialize(FArchive& Ar) override;
 
 protected:
 	UE_API virtual void GetLifetimeReplicatedProps(TArray<FLifetimeProperty>& OutLifetimeProps) const override;
