[GAS] Add virtual function to UE ASC to override logic around CMC->FlushServerMoves() for Locally Predicted Abilities before activation

#rnx

diff --git a/Engine/Plugins/Runtime/GameplayAbilities/Source/GameplayAbilities/Private/AbilitySystemComponent_Abilities.cpp b/Engine/Plugins/Runtime/GameplayAbilities/Source/GameplayAbilities/Private/AbilitySystemComponent_Abilities.cpp
--- a/Engine/Plugins/Runtime/GameplayAbilities/Source/GameplayAbilities/Private/AbilitySystemComponent_Abilities.cpp
+++ b/Engine/Plugins/Runtime/GameplayAbilities/Source/GameplayAbilities/Private/AbilitySystemComponent_Abilities.cpp
@@ -1895,20 +1895,7 @@ PRAGMA_ENABLE_DEPRECATION_WARNINGS
 	}
 	else if (Ability->GetNetExecutionPolicy() == EGameplayAbilityNetExecutionPolicy::LocalPredicted)
 	{
-		// Flush server moves that occurred before this ability activation so that the server receives the RPCs in the correct order
-		// Necessary to prevent abilities that trigger animation root motion or impact movement from causing network corrections
-		if (!ActorInfo->IsNetAuthority())
-		{
-			ACharacter* AvatarCharacter = Cast<ACharacter>(ActorInfo->AvatarActor.Get());
-			if (AvatarCharacter)
-			{
-				UCharacterMovementComponent* AvatarCharMoveComp = Cast<UCharacterMovementComponent>(AvatarCharacter->GetMovementComponent());
-				if (AvatarCharMoveComp)
-				{
-					AvatarCharMoveComp->FlushServerMoves();
-				}
-			}
-		}
+		PrePredictionActivation(Ability);
 
 		// This execution is now officially EGameplayAbilityActivationMode:Predicting and has a PredictionKey
 		FScopedPredictionWindow ScopedPredictionWindow(this, true);
@@ -1977,6 +1964,26 @@ PRAGMA_ENABLE_DEPRECATION_WARNINGS
 	return true;
 }
 
+void UAbilitySystemComponent::PrePredictionActivation(const UGameplayAbility* Ability)
+{
+	// Flush server moves that occurred before this ability activation so that the server receives the RPCs in the correct order
+	// Necessary to prevent abilities that trigger animation root motion or impact movement from causing network corrections
+	if (Ability && Ability->GetNetExecutionPolicy() == EGameplayAbilityNetExecutionPolicy::LocalPredicted)
+    {
+		const FGameplayAbilityActorInfo* ActorInfo = AbilityActorInfo.Get();
+		if (ActorInfo && !ActorInfo->IsNetAuthority())
+		{
+			if (ACharacter* AvatarCharacter = Cast<ACharacter>(ActorInfo->AvatarActor.Get()))
+			{
+				if (UCharacterMovementComponent* AvatarCharMoveComp = Cast<UCharacterMovementComponent>(AvatarCharacter->GetMovementComponent()))
+				{
+					AvatarCharMoveComp->FlushServerMoves();
+				}
+			}
+		}
+	}
+}
+
 void UAbilitySystemComponent::ServerTryActivateAbility_Implementation(FGameplayAbilitySpecHandle Handle, bool InputPressed, FPredictionKey PredictionKey)
 {
 	InternalServerTryActivateAbility(Handle, InputPressed, PredictionKey, nullptr);
diff --git a/Engine/Plugins/Runtime/GameplayAbilities/Source/GameplayAbilities/Public/AbilitySystemComponent.h b/Engine/Plugins/Runtime/GameplayAbilities/Source/GameplayAbilities/Public/AbilitySystemComponent.h
--- a/Engine/Plugins/Runtime/GameplayAbilities/Source/GameplayAbilities/Public/AbilitySystemComponent.h
+++ b/Engine/Plugins/Runtime/GameplayAbilities/Source/GameplayAbilities/Public/AbilitySystemComponent.h
@@ -1987,6 +1987,11 @@ protected:
 
 private:
 	UE_API float PlayMontageInternal(UGameplayAbility* AnimatingAbility, FGameplayAbilityActivationInfo ActivationInfo, UAnimMontage* Montage, float InPlayRate, FName StartSectionName, float StartTimeSeconds, TFunctionRef<void(FGameplayAbilityRepAnimMontage&)> MutateRepAnimMontageFunction);
+
+protected:
+	/** Called before activation for LocalPredicted Ability to execute CharacterMovementComponent->FlushServerMoves 
+	 * When overriding make sure to include Super::PrePredictionActivation(Ability) */
+	UE_API virtual void PrePredictionActivation(const UGameplayAbility* Ability);
 	
 };
 
