Fix GCs that were skipping on NetSerialize because of missing owner not using the replicated location

diff --git a/Engine/Plugins/Runtime/GameplayAbilities/Source/GameplayAbilities/Private/GameplayCueInterface.cpp b/Engine/Plugins/Runtime/GameplayAbilities/Source/GameplayAbilities/Private/GameplayCueInterface.cpp
--- a/Engine/Plugins/Runtime/GameplayAbilities/Source/GameplayAbilities/Private/GameplayCueInterface.cpp
+++ b/Engine/Plugins/Runtime/GameplayAbilities/Source/GameplayAbilities/Private/GameplayCueInterface.cpp
@@ -504,14 +504,20 @@ void FMinimalGameplayCueReplicationProxy::SetOwner(UAbilitySystemComponent* ASC)
 	Owner = ASC;
 	if (Owner && ReplicatedTags.Num() > 0)
 	{
+		const bool bLocationsValid = ensure(ReplicatedLocations.Num() >= ReplicatedTags.Num());
+
 		// Invoke events in case we skipped them during ::NetSerialize
 		FGameplayCueParameters Parameters;
 		InitGameplayCueParametersFunc(Parameters, Owner);
 
-		for (FGameplayTag& Tag : ReplicatedTags)
+		for (int32 idx=0; idx < ReplicatedTags.Num(); ++idx)
 		{
-			Owner->SetTagMapCount(Tag, 1);
-			Owner->InvokeGameplayCueEvent(Tag, EGameplayCueEvent::WhileActive, Parameters);
+			if (bLocationsValid)
+			{
+				Parameters.Location = ReplicatedLocations[idx];
+			}
+			Owner->SetTagMapCount(ReplicatedTags[idx], 1);
+			Owner->InvokeGameplayCueEvent(ReplicatedTags[idx], EGameplayCueEvent::WhileActive, Parameters);
 		}
 	}
 }
