Make sure to send the generated location if the replicated location is invalid

diff --git a/Engine/Plugins/Runtime/GameplayAbilities/Source/GameplayAbilities/Private/GameplayCueInterface.cpp b/Engine/Plugins/Runtime/GameplayAbilities/Source/GameplayAbilities/Private/GameplayCueInterface.cpp
--- a/Engine/Plugins/Runtime/GameplayAbilities/Source/GameplayAbilities/Private/GameplayCueInterface.cpp
+++ b/Engine/Plugins/Runtime/GameplayAbilities/Source/GameplayAbilities/Private/GameplayCueInterface.cpp
@@ -510,11 +510,13 @@ void FMinimalGameplayCueReplicationProxy::SetOwner(UAbilitySystemComponent* ASC)
 		FGameplayCueParameters Parameters;
 		InitGameplayCueParametersFunc(Parameters, Owner);
 
+		const FVector OriginalLocationParameter = Parameters.Location;
+
 		for (int32 idx=0; idx < ReplicatedTags.Num(); ++idx)
 		{
 			if (bLocationsValid)
 			{
-				Parameters.Location = ReplicatedLocations[idx];
+				Parameters.Location = ReplicatedLocations[idx].IsZero() ? OriginalLocationParameter : ReplicatedLocations[idx];
 			}
 			Owner->SetTagMapCount(ReplicatedTags[idx], 1);
 			Owner->InvokeGameplayCueEvent(ReplicatedTags[idx], EGameplayCueEvent::WhileActive, Parameters);
