Removing Compound multiply fix CVar

diff --git a/Engine/Plugins/Runtime/GameplayAbilities/Source/GameplayAbilities/Private/GameplayEffectTypes.cpp b/Engine/Plugins/Runtime/GameplayAbilities/Source/GameplayAbilities/Private/GameplayEffectTypes.cpp
--- a/Engine/Plugins/Runtime/GameplayAbilities/Source/GameplayAbilities/Private/GameplayEffectTypes.cpp
+++ b/Engine/Plugins/Runtime/GameplayAbilities/Source/GameplayAbilities/Private/GameplayEffectTypes.cpp
@@ -44,13 +44,6 @@ static FAutoConsoleVariableRef CVarWarnIfTryingToReplicateNotSupportedActorRefer
 }
 #endif
 
-static bool bUseMultiplyCompoundFix = true;
-static FAutoConsoleVariableRef CVarUseMultiplyCompoundFix(
-	TEXT("AbilitySystem.UseMultiplyCompundFix"),
-	bUseMultiplyCompoundFix,
-	TEXT("True/False - Add proper logic for compound multiplication in stackable GEs."),
-	ECVF_Default);
-
 TAutoConsoleVariable<bool> CVarReplicateTagCountContainerWithIris(
 		TEXT("AbilitySystem.Fix.ReplicateTagCountContainerWithIris"),
 		false,
@@ -141,16 +134,9 @@ float GameplayEffectUtilities::ComputeStackedModifierMagnitude(float BaseCompute
 	if (ModOp != EGameplayModOp::Override)
 	{
 		StackMag -= OperationBias;
-		if (bUseMultiplyCompoundFix)
+		if (ModOp == EGameplayModOp::MultiplyCompound)
 		{
-			if (ModOp == EGameplayModOp::MultiplyCompound)
-			{
-				StackMag = FMath::Pow(StackMag, StackCount);
-			}
-			else
-			{
-				StackMag *= StackCount;
-			}
+			StackMag = FMath::Pow(StackMag, StackCount);
 		}
 		else
 		{
