GE stacking rules are followed on the initial application of the stack.
GE predictive stacks treats 0 stack limit more as expected from users.

diff --git a/Engine/Plugins/Runtime/GameplayAbilities/Source/GameplayAbilities/Private/GameplayEffect.cpp b/Engine/Plugins/Runtime/GameplayAbilities/Source/GameplayAbilities/Private/GameplayEffect.cpp
--- a/Engine/Plugins/Runtime/GameplayAbilities/Source/GameplayAbilities/Private/GameplayEffect.cpp
+++ b/Engine/Plugins/Runtime/GameplayAbilities/Source/GameplayAbilities/Private/GameplayEffect.cpp
@@ -4136,7 +4136,7 @@ FActiveGameplayEffect* FActiveGameplayEffectsContainer::ApplyGameplayEffectSpec(
 	bool bSetPeriodTimer = true;
 
 	int32 StartingStackCount = 0;
-	int32 NewStackCount = 0;
+	int32 NewStackCount = Spec.GetStackCount();
 
 	// If we're applying a stack and the GE is configured to extend duration, 
 	// we will remember how much duration was left of the existing application.
@@ -4172,16 +4172,14 @@ FActiveGameplayEffect* FActiveGameplayEffectsContainer::ApplyGameplayEffectSpec(
 		// This is now the global "being applied GE"
 		UAbilitySystemGlobals::Get().SetCurrentAppliedGE(&ExistingSpec);
 		
-		// How to apply multiple stacks at once? What if we trigger an overflow which can reject the application?
-		// We still want to apply the stacks that didn't push us over, but we also want to call HandleActiveGameplayEffectStackOverflow.
-		
+		const int32 ExistingStackLimit = ExistingSpec.Def->StackLimitCount;
 		NewStackCount = ExistingSpec.GetStackCount() + Spec.GetStackCount();
 
 		// Legacy behavior: call HandleActiveGameplayEffectStackOverflow only if we are ALREADY at the limit.
-		const bool bStackCountAtLimit = ExistingSpec.GetStackCount() == ExistingSpec.Def->StackLimitCount;
+		const bool bStackCountAtLimit = ExistingSpec.GetStackCount() == ExistingStackLimit;
 
 		// In this path, we call HandleActiveGameplayEffectStackOverflow if we have any overflowing stacks. 
-		const bool bWillOverflowStacks = NewStackCount > ExistingSpec.Def->StackLimitCount && HasActiveGameplayEffectOverflowBehavior(EActiveGameplayEffectOverflowBehavior::HandleOverflowingStacksAlways);
+		const bool bWillOverflowStacks = ExistingStackLimit > 0 && NewStackCount > ExistingSpec.Def->StackLimitCount && HasActiveGameplayEffectOverflowBehavior(EActiveGameplayEffectOverflowBehavior::HandleOverflowingStacksAlways);
 
 		// We call HandleActiveGameplayEffectStackOverflow if we want to apply the remaining stacks up to the stack limit or for legacy behavior
 		if (bStackCountAtLimit || bWillOverflowStacks)
@@ -4193,9 +4191,9 @@ FActiveGameplayEffect* FActiveGameplayEffectsContainer::ApplyGameplayEffectSpec(
 			}
 		}
 		
-		if (ExistingSpec.Def->StackLimitCount > 0)
+		if (ExistingStackLimit > 0)
 		{
-			NewStackCount = FMath::Min(NewStackCount, ExistingSpec.Def->StackLimitCount);
+			NewStackCount = FMath::Min(NewStackCount, ExistingStackLimit);
 		}
 
 		// Need to unregister callbacks because the source aggregators could potentially be different with the new application. They will be
@@ -4252,6 +4250,8 @@ PRAGMA_ENABLE_DEPRECATION_WARNINGS
 	else
 	{
 		FActiveGameplayEffectHandle NewHandle = FActiveGameplayEffectHandle::GenerateNewHandle(Owner);
+		FActiveGameplayEffect* PreviousPendingEffect = nullptr;
+		bool bQueuedEffect = false;
 
 		if (GameplayEffects_Internal.GetSlack() <= 0)
 		{
@@ -4265,6 +4265,7 @@ PRAGMA_ENABLE_DEPRECATION_WARNINGS
 			 */
 
 			check(PendingGameplayEffectNext);
+			PreviousPendingEffect = *PendingGameplayEffectNext;
 			const FActiveGameplayEffect* PreviousPendingNext = (*PendingGameplayEffectNext) ? (*PendingGameplayEffectNext)->PendingNext : nullptr;
 
 			if (*PendingGameplayEffectNext == nullptr)
@@ -4287,6 +4288,7 @@ PRAGMA_ENABLE_DEPRECATION_WARNINGS
 
 			// The next pending GameplayEffect goes to where our PendingNext points
 			PendingGameplayEffectNext = &AppliedActiveGE->PendingNext;
+			bQueuedEffect = true;
 		}
 		else
 		{
@@ -4294,6 +4296,24 @@ PRAGMA_ENABLE_DEPRECATION_WARNINGS
 			// [#3] If you change this, please change #1-3!!!
 			AppliedActiveGE = new(GameplayEffects_Internal) FActiveGameplayEffect(NewHandle, Spec, GetWorldTime(), GetServerWorldTime(), InPredictionKey);
 		}
+
+		// Account for applying multiple stacks in our initial application.
+		if (!HandleActiveGameplayEffectStackOverflow(*AppliedActiveGE, Spec, Spec, InPredictionKey))
+		{
+			UE_VLOG_UELOG(Owner, LogGameplayEffects, Log, TEXT("Application of %s denied (StackLimit)"), *Spec.ToSimpleString());
+			if (bQueuedEffect)
+			{
+				*PendingGameplayEffectNext = PreviousPendingEffect;
+			}
+			return nullptr;
+		}
+		
+		if (Spec.Def->StackLimitCount > 0)
+		{
+			NewStackCount = FMath::Min(NewStackCount, Spec.Def->StackLimitCount);
+		}
+		
+		AppliedActiveGE->Spec.SetStackCount(NewStackCount);
 	}
 
 	// This is now the global "being applied GE"
@@ -4456,7 +4476,7 @@ PRAGMA_ENABLE_DEPRECATION_WARNINGS
 
 			// @note: This should be the number of stacks that would have been removed if the application is rejected but we would have been relying on the FPredictionKey reject delegate.
 			// However the reject delegate is not guaranteed to be called with 'bad' network conditions and we would end up with too many stacks on the client.
-			const int32 StacksAppliedWithinThisApplication = FMath::Min(Spec.GetStackCount(), Spec.Def->StackLimitCount - StartingStackCount);
+			const int32 StacksAppliedWithinThisApplication = NewStackCount - StartingStackCount;
 			UE_LOG(LogGameplayEffects, Log, TEXT("This Gameplay Effect [%s] has applied [%d] stacks when it had to stack."), *AppliedActiveGE->GetDebugString(), StacksAppliedWithinThisApplication);
 		}
 	}
@@ -4473,7 +4493,7 @@ PRAGMA_ENABLE_DEPRECATION_WARNINGS
 		const bool bInvokeGameplayCueEvents = (Spec.Def->DurationPolicy != EGameplayEffectDurationType::Instant);
 		InternalOnActiveGameplayEffectAdded(*AppliedActiveGE, bInvokeGameplayCueEvents);
 	}
-	
+
 	return AppliedActiveGE;
 }
 
