[Dynamic Ability Triggers] Allow for definition of dynamic ability triggers on the FGameplayAbilitySpec.

#tests
- Ensured that ability assets still opened containing their expected triggers.
- Added breakpoints in 'RegisterAbilityTriggers' and 'UnregisterAbilityTriggers' to ensure that registration still occurs.
- Executed a few triggered abilities to ensure that abilities still trigger correctly.

diff --git a/Engine/Plugins/Runtime/GameplayAbilities/Source/GameplayAbilities/Private/Abilities/GameplayAbility.cpp b/Engine/Plugins/Runtime/GameplayAbilities/Source/GameplayAbilities/Private/Abilities/GameplayAbility.cpp
--- a/Engine/Plugins/Runtime/GameplayAbilities/Source/GameplayAbilities/Private/Abilities/GameplayAbility.cpp
+++ b/Engine/Plugins/Runtime/GameplayAbilities/Source/GameplayAbilities/Private/Abilities/GameplayAbility.cpp
@@ -1932,7 +1932,13 @@ bool UGameplayAbility::IsTriggered() const
 {
 	// Assume that if there is triggered data, then we are triggered. 
 	// If we need to support abilities that can be both, this will need to be expanded.
-	return AbilityTriggers.Num() > 0;
+	if (AbilityTriggers.IsEmpty())
+	{
+		const FGameplayAbilitySpec* AbilitySpec = GetCurrentAbilitySpec();
+		return AbilitySpec != nullptr && !AbilitySpec->DynamicAbilityTriggers.IsEmpty();
+	}
+
+	return true;
 }
 
 bool UGameplayAbility::IsPredictingClient() const
diff --git a/Engine/Plugins/Runtime/GameplayAbilities/Source/GameplayAbilities/Private/AbilitySystemComponent_Abilities.cpp b/Engine/Plugins/Runtime/GameplayAbilities/Source/GameplayAbilities/Private/AbilitySystemComponent_Abilities.cpp
--- a/Engine/Plugins/Runtime/GameplayAbilities/Source/GameplayAbilities/Private/AbilitySystemComponent_Abilities.cpp
+++ b/Engine/Plugins/Runtime/GameplayAbilities/Source/GameplayAbilities/Private/AbilitySystemComponent_Abilities.cpp
@@ -553,11 +553,28 @@ void UAbilitySystemComponent::OnGiveAbility(FGameplayAbilitySpec& Spec)
 		}
 	}
 
-	for (const FAbilityTriggerData& TriggerData : Spec.Ability->AbilityTriggers)
+	RegisterAbilityTriggers(Spec, Spec.Ability->AbilityTriggers);
+	RegisterAbilityTriggers(Spec, Spec.DynamicAbilityTriggers);
+
+	// If there's already a primary instance, it should be the one to receive the OnGiveAbility call
+	UGameplayAbility* PrimaryInstance = bInstancedPerActor ? Spec.GetPrimaryInstance() : nullptr;
+	if (PrimaryInstance)
+	{
+		PrimaryInstance->OnGiveAbility(AbilityActorInfo.Get(), Spec);
+	}
+	else
+	{
+		Spec.Ability->OnGiveAbility(AbilityActorInfo.Get(), Spec);
+	}
+}
+
+void UAbilitySystemComponent::RegisterAbilityTriggers(const FGameplayAbilitySpec& Spec, TConstArrayView<FAbilityTriggerData> AbilityTriggers)
+{
+	for (const FAbilityTriggerData& TriggerData : AbilityTriggers)
 	{
 		FGameplayTag EventTag = TriggerData.TriggerTag;
 
-		auto& TriggeredAbilityMap = (TriggerData.TriggerSource == EGameplayAbilityTriggerSource::GameplayEvent) ? GameplayEventTriggeredAbilities : OwnedTagTriggeredAbilities;
+		TMap<FGameplayTag, TArray<FGameplayAbilitySpecHandle>>& TriggeredAbilityMap = (TriggerData.TriggerSource == EGameplayAbilityTriggerSource::GameplayEvent) ? GameplayEventTriggeredAbilities : OwnedTagTriggeredAbilities;
 
 		if (TriggeredAbilityMap.Contains(EventTag))
 		{
@@ -581,17 +598,6 @@ void UAbilitySystemComponent::OnGiveAbility(FGameplayAbilitySpec& Spec)
 			}
 		}
 	}
-
-	// If there's already a primary instance, it should be the one to receive the OnGiveAbility call
-	UGameplayAbility* PrimaryInstance = bInstancedPerActor ? Spec.GetPrimaryInstance() : nullptr;
-	if (PrimaryInstance)
-	{
-		PrimaryInstance->OnGiveAbility(AbilityActorInfo.Get(), Spec);
-	}
-	else
-	{
-		Spec.Ability->OnGiveAbility(AbilityActorInfo.Get(), Spec);
-	}
 }
 
 void UAbilitySystemComponent::OnRemoveAbility(FGameplayAbilitySpec& Spec)
@@ -606,22 +612,8 @@ void UAbilitySystemComponent::OnRemoveAbility(FGameplayAbilitySpec& Spec)
 	UE_LOG(LogAbilitySystem, Log, TEXT("%s: Removing Ability [%s] %s Level: %d"), *GetNameSafe(GetOwner()), *Spec.Handle.ToString(), *GetNameSafe(Spec.Ability), Spec.Level);
 	UE_VLOG(GetOwner(), VLogAbilitySystem, Log, TEXT("Removing Ability [%s] %s Level: %d"), *Spec.Handle.ToString(), *GetNameSafe(Spec.Ability), Spec.Level);
 
-	for (const FAbilityTriggerData& TriggerData : Spec.Ability->AbilityTriggers)
-	{
-		FGameplayTag EventTag = TriggerData.TriggerTag;
-
-		auto& TriggeredAbilityMap = (TriggerData.TriggerSource == EGameplayAbilityTriggerSource::GameplayEvent) ? GameplayEventTriggeredAbilities : OwnedTagTriggeredAbilities;
-
-		if (ensureMsgf(TriggeredAbilityMap.Contains(EventTag), 
-			TEXT("%s::%s not found in TriggeredAbilityMap while removing, TriggerSource: %d"), *Spec.Ability->GetName(), *EventTag.ToString(), (int32)TriggerData.TriggerSource))
-		{
-			TriggeredAbilityMap[EventTag].Remove(Spec.Handle);
-			if (TriggeredAbilityMap[EventTag].Num() == 0)
-			{
-				TriggeredAbilityMap.Remove(EventTag);
-			}
-		}
-	}
+	UnregisterAbilityTriggers(Spec, Spec.Ability->AbilityTriggers);
+	UnregisterAbilityTriggers(Spec, Spec.DynamicAbilityTriggers);
 
 	TArray<UGameplayAbility*> Instances = Spec.GetAbilityInstances();
 
@@ -724,6 +716,26 @@ PRAGMA_ENABLE_DEPRECATION_WARNINGS
 	Spec.NonReplicatedInstances.Empty();
 }
 
+void UAbilitySystemComponent::UnregisterAbilityTriggers(const FGameplayAbilitySpec& Spec, TConstArrayView<FAbilityTriggerData> AbilityTriggers)
+{
+	for (const FAbilityTriggerData& TriggerData : AbilityTriggers)
+	{
+		FGameplayTag EventTag = TriggerData.TriggerTag;
+
+		TMap<FGameplayTag, TArray<FGameplayAbilitySpecHandle>>& TriggeredAbilityMap = (TriggerData.TriggerSource == EGameplayAbilityTriggerSource::GameplayEvent) ? GameplayEventTriggeredAbilities : OwnedTagTriggeredAbilities;
+
+		if (ensureMsgf(TriggeredAbilityMap.Contains(EventTag),
+			TEXT("%s::%s not found in TriggeredAbilityMap while removing, TriggerSource: %d"), *Spec.Ability->GetName(), *EventTag.ToString(), (int32)TriggerData.TriggerSource))
+		{
+			TriggeredAbilityMap[EventTag].Remove(Spec.Handle);
+			if (TriggeredAbilityMap[EventTag].Num() == 0)
+			{
+				TriggeredAbilityMap.Remove(EventTag);
+			}
+		}
+	}
+}
+
 void UAbilitySystemComponent::CheckForClearedAbilities()
 {
 	for (auto& Triggered : GameplayEventTriggeredAbilities)
@@ -2635,6 +2647,7 @@ void UAbilitySystemComponent::MonitoredTagChanged(const FGameplayTag Tag, int32
 			if (Spec->Ability)
 			{
 				TArray<FAbilityTriggerData> AbilityTriggers = Spec->Ability->AbilityTriggers;
+				AbilityTriggers.Append(Spec->DynamicAbilityTriggers);
 				for (const FAbilityTriggerData& TriggerData : AbilityTriggers)
 				{
 					FGameplayTag EventTag = TriggerData.TriggerTag;
diff --git a/Engine/Plugins/Runtime/GameplayAbilities/Source/GameplayAbilities/Public/Abilities/GameplayAbility.h b/Engine/Plugins/Runtime/GameplayAbilities/Source/GameplayAbilities/Public/Abilities/GameplayAbility.h
--- a/Engine/Plugins/Runtime/GameplayAbilities/Source/GameplayAbilities/Public/Abilities/GameplayAbility.h
+++ b/Engine/Plugins/Runtime/GameplayAbilities/Source/GameplayAbilities/Public/Abilities/GameplayAbility.h
@@ -86,25 +86,6 @@ DECLARE_MULTICAST_DELEGATE_OneParam(FOnGameplayAbilityStateEnded, FName);
 /** Used to delay execution until we leave a critical section */
 DECLARE_DELEGATE(FPostLockDelegate);
 
-/** Structure that defines how an ability will be triggered by external events */
-USTRUCT()
-struct FAbilityTriggerData
-{
-	GENERATED_USTRUCT_BODY()
-
-	FAbilityTriggerData() 
-	: TriggerSource(EGameplayAbilityTriggerSource::GameplayEvent)
-	{}
-
-	/** The tag to respond to */
-	UPROPERTY(EditAnywhere, Category=TriggerData, meta=(Categories="TriggerTagCategory"))
-	FGameplayTag TriggerTag;
-
-	/** The type of trigger to respond to */
-	UPROPERTY(EditAnywhere, Category=TriggerData)
-	TEnumAsByte<EGameplayAbilityTriggerSource::Type> TriggerSource;
-};
-
 /** Abilities define custom gameplay logic that can be activated by players or external game logic */
 UCLASS(Blueprintable, MinimalAPI)
 class UGameplayAbility : public UObject, public IGameplayTaskOwnerInterface
diff --git a/Engine/Plugins/Runtime/GameplayAbilities/Source/GameplayAbilities/Public/Abilities/GameplayAbilityTriggerType.h b/Engine/Plugins/Runtime/GameplayAbilities/Source/GameplayAbilities/Public/Abilities/GameplayAbilityTriggerType.h
new file mode 100644
--- /dev/null
+++ b/Engine/Plugins/Runtime/GameplayAbilities/Source/GameplayAbilities/Public/Abilities/GameplayAbilityTriggerType.h
@@ -0,0 +1,42 @@
+// Copyright Epic Games, Inc. All Rights Reserved.
+
+#pragma once
+
+#include "GameplayTagContainer.h"
+#include "GameplayAbilityTriggerType.generated.h"
+
+UENUM(BlueprintType)
+namespace EGameplayAbilityTriggerSource
+{
+	/**	Defines what type of trigger will activate the ability, paired to a tag */
+	enum Type : int
+	{
+		/** Triggered by an external gameplay event. The ability will receive a GameplayEvent payload. */
+		GameplayEvent      UMETA(DisplayName = "On Gameplay Event"),
+		
+		/** Triggered when the owner gains the specified tag. Will not cancel when the tag is removed. */
+		OwnedTagAdded      UMETA(DisplayName = "When Tag Is Added"),
+		
+		/** Triggered when the owner has the specified tag. The ability will be canceled if the tag is later removed. */
+		OwnedTagPresent    UMETA(DisplayName = "When Tag Is Present"),
+	};
+}
+
+/** Structure that defines how an ability will be triggered by external events */
+USTRUCT()
+struct FAbilityTriggerData
+{
+	GENERATED_USTRUCT_BODY()
+
+	FAbilityTriggerData() 
+	: TriggerSource(EGameplayAbilityTriggerSource::GameplayEvent)
+	{}
+
+	/** The tag to respond to */
+	UPROPERTY(EditAnywhere, Category=TriggerData, meta=(Categories="TriggerTagCategory"))
+	FGameplayTag TriggerTag;
+
+	/** The type of trigger to respond to */
+	UPROPERTY(EditAnywhere, Category=TriggerData)
+	TEnumAsByte<EGameplayAbilityTriggerSource::Type> TriggerSource;
+};
diff --git a/Engine/Plugins/Runtime/GameplayAbilities/Source/GameplayAbilities/Public/Abilities/GameplayAbilityTypes.h b/Engine/Plugins/Runtime/GameplayAbilities/Source/GameplayAbilities/Public/Abilities/GameplayAbilityTypes.h
--- a/Engine/Plugins/Runtime/GameplayAbilities/Source/GameplayAbilities/Public/Abilities/GameplayAbilityTypes.h
+++ b/Engine/Plugins/Runtime/GameplayAbilities/Source/GameplayAbilities/Public/Abilities/GameplayAbilityTypes.h
@@ -11,6 +11,7 @@
 #include "GameplayEffectTypes.h"
 #include "GameplayPrediction.h"
 #include "GameplayAbilitySpec.h"
+#include "Abilities/GameplayAbilityTriggerType.h"
 #include "Abilities/GameplayAbilityTargetTypes.h"
 
 #include "GameplayAbilityTypes.generated.h"
@@ -108,24 +109,6 @@ namespace EGameplayAbilityReplicationPolicy
 	};
 }
 
-UENUM(BlueprintType)
-namespace EGameplayAbilityTriggerSource
-{
-	/**	Defines what type of trigger will activate the ability, paired to a tag */
-	enum Type : int
-	{
-		/** Triggered by an external gameplay event. The ability will receive a GameplayEvent payload. */
-		GameplayEvent      UMETA(DisplayName = "On Gameplay Event"),
-		
-		/** Triggered when the owner gains the specified tag. Will not cancel when the tag is removed. */
-		OwnedTagAdded      UMETA(DisplayName = "When Tag Is Added"),
-		
-		/** Triggered when the owner has the specified tag. The ability will be canceled if the tag is later removed. */
-		OwnedTagPresent    UMETA(DisplayName = "When Tag Is Present"),
-	};
-}
-
-
 /**
  *	FGameplayAbilityActorInfo
  *
diff --git a/Engine/Plugins/Runtime/GameplayAbilities/Source/GameplayAbilities/Public/AbilitySystemComponent.h b/Engine/Plugins/Runtime/GameplayAbilities/Source/GameplayAbilities/Public/AbilitySystemComponent.h
--- a/Engine/Plugins/Runtime/GameplayAbilities/Source/GameplayAbilities/Public/AbilitySystemComponent.h
+++ b/Engine/Plugins/Runtime/GameplayAbilities/Source/GameplayAbilities/Public/AbilitySystemComponent.h
@@ -1932,6 +1932,9 @@ private:
 	UPROPERTY(Replicated)
 	FMinimalReplicationTagCountMap MinimalReplicationTags;
 
+	void RegisterAbilityTriggers(const FGameplayAbilitySpec& Spec, TConstArrayView<FAbilityTriggerData> AbilityTriggers);
+	void UnregisterAbilityTriggers(const FGameplayAbilitySpec& Spec, TConstArrayView<FAbilityTriggerData> AbilityTriggers);
+
 private:
 
 	/** List of attribute sets */
diff --git a/Engine/Plugins/Runtime/GameplayAbilities/Source/GameplayAbilities/Public/GameplayAbilitySpec.h b/Engine/Plugins/Runtime/GameplayAbilities/Source/GameplayAbilities/Public/GameplayAbilitySpec.h
--- a/Engine/Plugins/Runtime/GameplayAbilities/Source/GameplayAbilities/Public/GameplayAbilitySpec.h
+++ b/Engine/Plugins/Runtime/GameplayAbilities/Source/GameplayAbilities/Public/GameplayAbilitySpec.h
@@ -8,6 +8,7 @@
 #include "Templates/SubclassOf.h"
 #include "Engine/NetSerialization.h"
 #include "AttributeSet.h"
+#include "Abilities/GameplayAbilityTriggerType.h"
 #include "GameplayAbilitySpecHandle.h"
 #include "GameplayEffectTypes.h"
 #include "GameplayPrediction.h"
@@ -248,6 +249,10 @@ PRAGMA_DISABLE_DEPRECATION_WARNINGS
 	const FGameplayTagContainer& GetDynamicSpecSourceTags() const { return DynamicAbilityTags; }
 PRAGMA_ENABLE_DEPRECATION_WARNINGS
 
+	/** Additional Ability Triggers that can be used in order to trigger the ability associated with this spec. */
+	UPROPERTY()
+	TArray<FAbilityTriggerData> DynamicAbilityTriggers;
+
 	/** Non replicating instances of this ability. */
 	UPROPERTY(NotReplicated)
 	TArray<TObjectPtr<UGameplayAbility>> NonReplicatedInstances;
