Submit a workaround fix for infinite loop at PIE shutdown in ~FActiveGameplayEffectsContainer() destructor, when FActiveGameplayEffect.PendingNext points to itself.
NOTE: [REDACTED_HANDLE] will be investigating further to determine if this is related to recent changes

#rnx

diff --git a/Engine/Plugins/Runtime/GameplayAbilities/Source/GameplayAbilities/Private/GameplayEffect.cpp b/Engine/Plugins/Runtime/GameplayAbilities/Source/GameplayAbilities/Private/GameplayEffect.cpp
--- a/Engine/Plugins/Runtime/GameplayAbilities/Source/GameplayAbilities/Private/GameplayEffect.cpp
+++ b/Engine/Plugins/Runtime/GameplayAbilities/Source/GameplayAbilities/Private/GameplayEffect.cpp
@@ -2979,7 +2979,7 @@ FActiveGameplayEffectsContainer::~FActiveGameplayEffectsContainer()
 	{
 		while (PendingGameplayEffectHead)
 		{
-			FActiveGameplayEffect* Next = PendingGameplayEffectHead->PendingNext;
+			FActiveGameplayEffect* Next = (PendingGameplayEffectHead->PendingNext != PendingGameplayEffectHead) ? PendingGameplayEffectHead->PendingNext : nullptr;
 			delete PendingGameplayEffectHead;
 			PendingGameplayEffectHead = Next;
 		}
