Added a templated static method to get attribute sets from the AbilitySystemComponent

diff --git a/Engine/Plugins/Runtime/GameplayAbilities/Source/GameplayAbilities/Public/AbilitySystemComponent.h b/Engine/Plugins/Runtime/GameplayAbilities/Source/GameplayAbilities/Public/AbilitySystemComponent.h
--- a/Engine/Plugins/Runtime/GameplayAbilities/Source/GameplayAbilities/Public/AbilitySystemComponent.h
+++ b/Engine/Plugins/Runtime/GameplayAbilities/Source/GameplayAbilities/Public/AbilitySystemComponent.h
@@ -2,6 +2,7 @@
 
 #pragma once
 
+#include "AbilitySystemBlueprintLibrary.h"
 #include "CoreMinimal.h"
 #include "UObject/ObjectMacros.h"
 #include "Templates/SubclassOf.h"
@@ -136,6 +137,18 @@ class UAbilitySystemComponent : public UGameplayTasksComponent, public IGameplay
 		return (T*)GetAttributeSubobjectChecked(T::StaticClass());
 	}
 
+	/** Finds existing AttributeSet for the given actor */
+	template <class T >
+	static const T* GetSetOnActor(AActor* InActor)
+	{
+		if (UAbilitySystemComponent* ASC = UAbilitySystemBlueprintLibrary::GetAbilitySystemComponent(InActor))
+		{
+			return ASC->GetSet<T>();
+		}
+
+		return nullptr;
+	}
+
 	/** Adds a new AttributeSet (initialized to default values) */
 	template <class T >
 	const T*  AddSet()
