Fixed instance where adding a new GE with the same number of stacks as the stack limit cause it not to be applied.

diff --git a/Engine/Plugins/Runtime/GameplayAbilities/Source/GameplayAbilities/Private/GameplayEffect.cpp b/Engine/Plugins/Runtime/GameplayAbilities/Source/GameplayAbilities/Private/GameplayEffect.cpp
--- a/Engine/Plugins/Runtime/GameplayAbilities/Source/GameplayAbilities/Private/GameplayEffect.cpp
+++ b/Engine/Plugins/Runtime/GameplayAbilities/Source/GameplayAbilities/Private/GameplayEffect.cpp
@@ -2973,13 +2973,13 @@ FActiveGameplayEffectsContainer::~FActiveGameplayEffectsContainer()
 {
 	using namespace UE::GameplayEffect;
 
-	FActiveGameplayEffect* PendingGameplayEffect = PendingGameplayEffectHead;
+	ensureAlwaysMsgf(!PendingGameplayEffectHead || PendingGameplayEffectHead != PendingGameplayEffectHead->PendingNext, TEXT("Pending Gameplay Effect Head is pointing to the same location as Pending Next this shouldn't happen'"));
 
 	if (HasActiveGameplayEffectFix(EActiveGameplayEffectFix::CleanupAllPendingActiveGEs))
 	{
 		while (PendingGameplayEffectHead)
 		{
-			FActiveGameplayEffect* Next = (PendingGameplayEffectHead->PendingNext != PendingGameplayEffectHead) ? PendingGameplayEffectHead->PendingNext : nullptr;
+			FActiveGameplayEffect* Next = PendingGameplayEffectHead->PendingNext;
 			delete PendingGameplayEffectHead;
 			PendingGameplayEffectHead = Next;
 		}
@@ -4297,8 +4297,11 @@ PRAGMA_ENABLE_DEPRECATION_WARNINGS
 			AppliedActiveGE = new(GameplayEffects_Internal) FActiveGameplayEffect(NewHandle, Spec, GetWorldTime(), GetServerWorldTime(), InPredictionKey);
 		}
 
+		// Make a copy of the spec with a count of 0 so that the handle overflow logic is accurate.
+		FGameplayEffectSpec SpecCopy = Spec;
+		SpecCopy.SetStackCount(0);
 		// Account for applying multiple stacks in our initial application.
-		if (!HandleActiveGameplayEffectStackOverflow(*AppliedActiveGE, Spec, Spec, InPredictionKey))
+		if (!HandleActiveGameplayEffectStackOverflow(*AppliedActiveGE, SpecCopy, Spec, InPredictionKey))
 		{
 			UE_VLOG_UELOG(Owner, LogGameplayEffects, Log, TEXT("Application of %s denied (StackLimit)"), *Spec.ToSimpleString());
 			if (bQueuedEffect)
