- Added object name to non instanced ability ensures and updated the string to include the function name, as sometimes __func__ could get the wrong function name.

diff --git a/Engine/Plugins/Runtime/GameplayAbilities/Source/GameplayAbilities/Private/Abilities/GameplayAbility.cpp b/Engine/Plugins/Runtime/GameplayAbilities/Source/GameplayAbilities/Private/Abilities/GameplayAbility.cpp
--- a/Engine/Plugins/Runtime/GameplayAbilities/Source/GameplayAbilities/Private/Abilities/GameplayAbility.cpp
+++ b/Engine/Plugins/Runtime/GameplayAbilities/Source/GameplayAbilities/Private/Abilities/GameplayAbility.cpp
@@ -1232,7 +1232,7 @@ FGameplayAbilityActorInfo UGameplayAbility::GetActorInfo() const
 AActor* UGameplayAbility::GetOwningActorFromActorInfo() const
 {
 	if (!ensureMsgf(UE::AbilitySystem::Private::CVarAllowNonInstancedGAsValue || IsInstantiated(),
-		TEXT("%hs called on the CDO.  NonInstanced abilities are deprecated, thus we always expect this to be called on an instanced object."), __func__))
+		TEXT("GetOwningActorFromActorInfo() called on the CDO for %s.  NonInstanced abilities are deprecated, thus we always expect this to be called on an instanced object."), *GetNameSafe(this)))
 	{
 		return nullptr;
 	}
@@ -1247,7 +1247,7 @@ AActor* UGameplayAbility::GetOwningActorFromActorInfo() const
 AActor* UGameplayAbility::GetAvatarActorFromActorInfo() const
 {
 	ensureMsgf(UE::AbilitySystem::Private::CVarAllowNonInstancedGAsValue || IsInstantiated(),
-		TEXT("%hs called on the CDO.  NonInstanced abilities are deprecated, thus we always expect this to be called on an instanced object."), __func__);
+		TEXT("GetAvatarActorFromActorInfo() called on the CDO for %s.  NonInstanced abilities are deprecated, thus we always expect this to be called on an instanced object."), *GetNameSafe(this));
 
 	if (!ensure(CurrentActorInfo))
 	{
@@ -1259,7 +1259,7 @@ AActor* UGameplayAbility::GetAvatarActorFromActorInfo() const
 USkeletalMeshComponent* UGameplayAbility::GetOwningComponentFromActorInfo() const
 {
 	ensureMsgf(UE::AbilitySystem::Private::CVarAllowNonInstancedGAsValue || IsInstantiated(),
-		TEXT("%hs called on the CDO.  NonInstanced abilities are deprecated, thus we always expect this to be called on an instanced object."), __func__);
+		TEXT("GetOwningComponentFromActorInfo() called on the CDO for %s.  NonInstanced abilities are deprecated, thus we always expect this to be called on an instanced object."), *GetNameSafe(this));
 
 	if (!ensure(CurrentActorInfo))
 	{
@@ -1272,7 +1272,7 @@ USkeletalMeshComponent* UGameplayAbility::GetOwningComponentFromActorInfo() cons
 UAbilitySystemComponent* UGameplayAbility::GetAbilitySystemComponentFromActorInfo() const
 {
 	ensureMsgf(UE::AbilitySystem::Private::CVarAllowNonInstancedGAsValue || IsInstantiated(),
-		TEXT("%hs called on the CDO.  NonInstanced abilities are deprecated, thus we always expect this to be called on an instanced object."), __func__);
+		TEXT("GetAbilitySystemComponentFromActorInfo() called on the CDO for %s.  NonInstanced abilities are deprecated, thus we always expect this to be called on an instanced object."), *GetNameSafe(this));
 
 	if (!ensure(CurrentActorInfo))
 	{
@@ -1311,7 +1311,7 @@ void UGameplayAbility::SetAssetTags(const FGameplayTagContainer& InAssetTags)
 const FGameplayAbilityActorInfo* UGameplayAbility::GetCurrentActorInfo() const
 {
 	if (!ensureMsgf(UE::AbilitySystem::Private::CVarAllowNonInstancedGAsValue || IsInstantiated(),
-		TEXT("%hs called on the CDO.  NonInstanced abilities are deprecated, thus we always expect this to be called on an instanced object."), __func__))
+		TEXT("GetCurrentActorInfo() called on the CDO for %s.  NonInstanced abilities are deprecated, thus we always expect this to be called on an instanced object."), *GetNameSafe(this)))
 	{
 		return nullptr;
 	}
@@ -1322,7 +1322,7 @@ const FGameplayAbilityActorInfo* UGameplayAbility::GetCurrentActorInfo() const
 FGameplayAbilityActivationInfo UGameplayAbility::GetCurrentActivationInfo() const
 {
 	if (!ensureMsgf(UE::AbilitySystem::Private::CVarAllowNonInstancedGAsValue || IsInstantiated(),
-		TEXT("%hs called on the CDO.  NonInstanced abilities are deprecated, thus we always expect this to be called on an instanced object."), __func__))
+		TEXT("GetCurrentActivationInfo() called on the CDO for %s.  NonInstanced abilities are deprecated, thus we always expect this to be called on an instanced object."), *GetNameSafe(this)))
 	{
 		return FGameplayAbilityActivationInfo{};
 	}
@@ -1333,7 +1333,7 @@ FGameplayAbilityActivationInfo UGameplayAbility::GetCurrentActivationInfo() cons
 FGameplayAbilitySpecHandle UGameplayAbility::GetCurrentAbilitySpecHandle() const
 {
 	if (!ensureMsgf(UE::AbilitySystem::Private::CVarAllowNonInstancedGAsValue || IsInstantiated(),
-		TEXT("%hs called on the CDO.  NonInstanced abilities are deprecated, thus we always expect this to be called on an instanced object."), __func__))
+		TEXT("GetCurrentAbilitySpecHandle() called on the CDO for %s.  NonInstanced abilities are deprecated, thus we always expect this to be called on an instanced object."), *GetNameSafe(this)))
 	{
 		return FGameplayAbilitySpecHandle{};
 	}
@@ -1507,7 +1507,7 @@ void UGameplayAbility::MontageStop(float OverrideBlendOutTime)
 void UGameplayAbility::SetCurrentMontage(class UAnimMontage* InCurrentMontage)
 {
 	if (ensureMsgf(UE::AbilitySystem::Private::CVarAllowNonInstancedGAsValue || IsInstantiated(),
-		TEXT("%hs called on the CDO.  NonInstanced abilities are deprecated, thus we always expect this to be called on an instanced object."), __func__))
+		TEXT("SetCurrentMontage() called on the CDO for %s.  NonInstanced abilities are deprecated, thus we always expect this to be called on an instanced object."), *GetNameSafe(this)))
 	{
 		CurrentMontage = InCurrentMontage;
 	}
@@ -1516,7 +1516,7 @@ void UGameplayAbility::SetCurrentMontage(class UAnimMontage* InCurrentMontage)
 UAnimMontage* UGameplayAbility::GetCurrentMontage() const
 {
 	ensureMsgf(UE::AbilitySystem::Private::CVarAllowNonInstancedGAsValue || IsInstantiated(),
-		TEXT("%hs called on the CDO.  NonInstanced abilities are deprecated, thus we always expect this to be called on an instanced object."), __func__);
+		TEXT("GetCurrentMontage() called on the CDO for %s.  NonInstanced abilities are deprecated, thus we always expect this to be called on an instanced object."), *GetNameSafe(this));
 	return CurrentMontage;
 }
 
