Fix ScalableFloats to correctly reset cached curve when serializing or copying. This fixes issues when used in inherited blueprints

diff --git a/Engine/Plugins/Runtime/GameplayAbilities/Source/GameplayAbilities/Private/ScalableFloat.cpp b/Engine/Plugins/Runtime/GameplayAbilities/Source/GameplayAbilities/Private/ScalableFloat.cpp
--- a/Engine/Plugins/Runtime/GameplayAbilities/Source/GameplayAbilities/Private/ScalableFloat.cpp
+++ b/Engine/Plugins/Runtime/GameplayAbilities/Source/GameplayAbilities/Private/ScalableFloat.cpp
@@ -97,7 +97,7 @@ bool FScalableFloat::EvaluateCurveAtLevel(float& OutValue, const FRealCurve*& Ou
 		const int32 GlobalCachedCurveID = UCurveTable::GetGlobalCachedCurveID();
 		if (LocalCachedCurveID != GlobalCachedCurveID)
 		{
-			CachedCurve = nullptr;
+			ClearCachedCurve();
 		}
 
 		OutCurve = CachedCurve;
@@ -213,8 +213,7 @@ void FScalableFloat::SetValue(float NewValue)
 	Value = NewValue;
 	Curve.CurveTable = nullptr;
 	Curve.RowName = NAME_None;
-	CachedCurve = nullptr;
-	LocalCachedCurveID = INDEX_NONE;
+	ClearCachedCurve();
 }
 
 void FScalableFloat::SetScalingValue(float InCoeffecient, FName InRowName, UCurveTable * InTable)
@@ -222,8 +221,7 @@ void FScalableFloat::SetScalingValue(float InCoeffecient, FName InRowName, UCurv
 	Value = InCoeffecient;
 	Curve.RowName = InRowName;
 	Curve.CurveTable = InTable;
-	CachedCurve = nullptr;
-	LocalCachedCurveID = INDEX_NONE;
+	ClearCachedCurve();
 }
 
 FString FScalableFloat::ToSimpleString() const
@@ -295,6 +293,14 @@ bool FScalableFloat::SerializeFromMismatchedTag(const FPropertyTag& Tag, FStruct
 	return false;
 }
 
+void FScalableFloat::PostSerialize(const FArchive& Ar)
+{
+	if (Ar.IsLoading())
+	{
+		ClearCachedCurve();
+	}
+}
+
 bool FScalableFloat::operator==(const FScalableFloat& Other) const
 {
 	return ((Other.Curve == Curve) && (Other.RegistryType == RegistryType) && (Other.Value == Value));
@@ -310,8 +316,9 @@ FScalableFloat& FScalableFloat::operator=(const FScalableFloat& Src)
 	Value = Src.Value;
 	Curve = Src.Curve;
 	RegistryType = Src.RegistryType;
-	LocalCachedCurveID = Src.LocalCachedCurveID;
-	CachedCurve = Src.CachedCurve;
+	
+	// Clear cache here instead of copying it over to ensure we never have out of date cached values
+	ClearCachedCurve();
 
 	return *this;
 }
diff --git a/Engine/Plugins/Runtime/GameplayAbilities/Source/GameplayAbilities/Public/ScalableFloat.h b/Engine/Plugins/Runtime/GameplayAbilities/Source/GameplayAbilities/Public/ScalableFloat.h
--- a/Engine/Plugins/Runtime/GameplayAbilities/Source/GameplayAbilities/Public/ScalableFloat.h
+++ b/Engine/Plugins/Runtime/GameplayAbilities/Source/GameplayAbilities/Public/ScalableFloat.h
@@ -38,6 +38,7 @@ public:
 	float	Value;
 
 private:
+	/** Value of GetGlobalCachedCurveID when CachedCurve was set */
 	mutable int32 LocalCachedCurveID;
 
 public:
@@ -107,11 +108,19 @@ public:
 
 	/** Used to upgrade a float or int8/int16/int32 property into an FScalableFloat */
 	UE_API bool SerializeFromMismatchedTag(const FPropertyTag& Tag, FStructuredArchive::FSlot Slot);
+	UE_API void PostSerialize(const FArchive& Ar);
 
 private:
-
-	// Cached direct pointer to the RealCurve we should evaluate
+	/** Cached direct pointer to the RealCurve we should evaluate */
 	mutable const FRealCurve* CachedCurve;
+
+	/** Clear cached data */
+	void ClearCachedCurve() const
+	{
+		LocalCachedCurveID = INDEX_NONE;
+		CachedCurve = nullptr;
+	}
+
 };
 
 template<>
@@ -121,6 +130,7 @@ struct TStructOpsTypeTraits<FScalableFloat>
 	enum
 	{
 		WithStructuredSerializeFromMismatchedTag = true,
+		WithPostSerialize = true,
 	};
 };
 
