Move ActorInfo null check earlier to actually check before we use it at all

#tests Loaded test level and verified that gameplay effects were running properly

diff --git a/Engine/Plugins/Runtime/GameplayAbilities/Source/GameplayAbilities/Private/Abilities/GameplayAbility.cpp b/Engine/Plugins/Runtime/GameplayAbilities/Source/GameplayAbilities/Private/Abilities/GameplayAbility.cpp
--- a/Engine/Plugins/Runtime/GameplayAbilities/Source/GameplayAbilities/Private/Abilities/GameplayAbility.cpp
+++ b/Engine/Plugins/Runtime/GameplayAbilities/Source/GameplayAbilities/Private/Abilities/GameplayAbility.cpp
@@ -2064,13 +2064,13 @@ FActiveGameplayEffectHandle UGameplayAbility::ApplyGameplayEffectSpecToOwner(con
 	// This batches all created cues together
 	FScopedGameplayCueSendContext GameplayCueSendContext;
 
-	if (SpecHandle.IsValid() && (HasAuthorityOrPredictionKey(ActorInfo, &ActivationInfo)))
+	if (SpecHandle.IsValid() && ActorInfo && (HasAuthorityOrPredictionKey(ActorInfo, &ActivationInfo)))
 	{
 		UAbilitySystemComponent* const AbilitySystemComponent = ActorInfo->AbilitySystemComponent.Get();
 		FActiveGameplayEffectHandle ReturnHandle = AbilitySystemComponent->ApplyGameplayEffectSpecToSelf(*SpecHandle.Data.Get(), AbilitySystemComponent->GetPredictionKeyForNewAction());
 
 		// Force a net update for predicted effects so that short durations are replicated correctly
-		if (ReturnHandle.IsValid() && ActorInfo && ActorInfo->OwnerActor.IsValid() && ActorInfo->OwnerActor->HasAuthority())
+		if (ReturnHandle.IsValid() && ActorInfo->OwnerActor.IsValid() && ActorInfo->OwnerActor->HasAuthority())
 		{
 			AbilitySystemComponent->GetOwner()->ForceNetUpdate();
 		}
