Access GameplayEffect Stacking Type from external module.
#14025
#Synced-CL 48156723

diff --git a/Engine/Plugins/Runtime/GameplayAbilities/Source/GameplayAbilities/Public/GameplayEffect.h b/Engine/Plugins/Runtime/GameplayAbilities/Source/GameplayAbilities/Public/GameplayEffect.h
--- a/Engine/Plugins/Runtime/GameplayAbilities/Source/GameplayAbilities/Public/GameplayEffect.h
+++ b/Engine/Plugins/Runtime/GameplayAbilities/Source/GameplayAbilities/Public/GameplayEffect.h
@@ -2425,7 +2425,7 @@ public:
 	bool bFactorInStackCount;
 
 	/** return the stacking type defined from the data asset */
-	EGameplayEffectStackingType GetStackingType() const;
+	UE_API EGameplayEffectStackingType GetStackingType() const;
 
 #if WITH_EDITOR
 	/** Set stacking type */
