Cleaned up comments referencing removed Pre/PostAttributeModify functions.
#14019
#Synced-CL 48156723

diff --git a/Engine/Plugins/Runtime/GameplayAbilities/Source/GameplayAbilities/Private/AbilitySystemTestAttributeSet.cpp b/Engine/Plugins/Runtime/GameplayAbilities/Source/GameplayAbilities/Private/AbilitySystemTestAttributeSet.cpp
--- a/Engine/Plugins/Runtime/GameplayAbilities/Source/GameplayAbilities/Private/AbilitySystemTestAttributeSet.cpp
+++ b/Engine/Plugins/Runtime/GameplayAbilities/Source/GameplayAbilities/Private/AbilitySystemTestAttributeSet.cpp
@@ -81,7 +81,6 @@ bool UAbilitySystemTestAttributeSet::PreGameplayEffectExecute(struct FGameplayEf
 			}
 		}
 
-		// At this point, the Magnitude of the applied damage may have been modified by us. We still do the translation to Health in UAbilitySystemTestAttributeSet::PostAttributeModify.
 	}
 
 #endif
diff --git a/Engine/Plugins/Runtime/GameplayAbilities/Source/GameplayAbilities/Private/GameplayEffect.cpp b/Engine/Plugins/Runtime/GameplayAbilities/Source/GameplayAbilities/Private/GameplayEffect.cpp
--- a/Engine/Plugins/Runtime/GameplayAbilities/Source/GameplayAbilities/Private/GameplayEffect.cpp
+++ b/Engine/Plugins/Runtime/GameplayAbilities/Source/GameplayAbilities/Private/GameplayEffect.cpp
@@ -4047,7 +4047,6 @@ bool FActiveGameplayEffectsContainer::InternalExecuteMod(FGameplayEffectSpec& Sp
 
 		/**
 		 *  This should apply 'gamewide' rules. Such as clamping Health to MaxHealth or granting +3 health for every point of strength, etc
-		 *	PreAttributeModify can return false to 'throw out' this modification.
 		 */
 		if (AttributeSet->PreGameplayEffectExecute(ExecuteData))
 		{
diff --git a/Engine/Plugins/Runtime/GameplayAbilities/Source/GameplayAbilities/Private/Tests/GameplayEffectTests.cpp b/Engine/Plugins/Runtime/GameplayAbilities/Source/GameplayAbilities/Private/Tests/GameplayEffectTests.cpp
--- a/Engine/Plugins/Runtime/GameplayAbilities/Source/GameplayAbilities/Private/Tests/GameplayEffectTests.cpp
+++ b/Engine/Plugins/Runtime/GameplayAbilities/Source/GameplayAbilities/Private/Tests/GameplayEffectTests.cpp
@@ -147,7 +147,7 @@ public: // the tests
 		const float DamageValue = 5.f;
 		const float StartingHealth = DestComponent->GetSet<UAbilitySystemTestAttributeSet>()->Health;
 
-		// This is the same as GameplayEffectsTest_InstantDamage but modifies the Damage attribute and confirms it is remapped to -Health by UAbilitySystemTestAttributeSet::PostAttributeModify
+		// This is the same as GameplayEffectsTest_InstantDamage but modifies the Damage attribute and confirms it is remapped to -Health
 		{
 			CONSTRUCT_CLASS(UGameplayEffect, BaseDmgEffect);
 			AddModifier(BaseDmgEffect, GET_FIELD_CHECKED(UAbilitySystemTestAttributeSet, Damage), EGameplayModOp::Additive, FScalableFloat(DamageValue));
diff --git a/Engine/Plugins/Runtime/GameplayAbilities/Source/GameplayAbilities/Public/AttributeSet.h b/Engine/Plugins/Runtime/GameplayAbilities/Source/GameplayAbilities/Public/AttributeSet.h
--- a/Engine/Plugins/Runtime/GameplayAbilities/Source/GameplayAbilities/Public/AttributeSet.h
+++ b/Engine/Plugins/Runtime/GameplayAbilities/Source/GameplayAbilities/Public/AttributeSet.h
@@ -211,7 +211,7 @@ public:
 	 */
 
 	/**
-	 *	Called just before any modification happens to an attribute. This is lower level than PreAttributeModify/PostAttribute modify.
+	 *	Called just before any modification happens to an attribute.
 	 *	There is no additional context provided here since anything can trigger this. Executed effects, duration based effects, effects being removed, immunity being applied, stacking rules changing, etc.
 	 *	This function is meant to enforce things like "Health = Clamp(Health, 0, MaxHealth)" and NOT things like "trigger this extra thing if damage is applied, etc".
 	 *	
