When disabled the GameplayTagCountContainer should not replicate it's initial state either, this is especially important for actors with low net replication rates.
#Synced-CL 48357934

diff --git a/Engine/Plugins/Runtime/GameplayAbilities/Source/GameplayAbilities/Public/Serialization/GameplayTagCountContainerNetSerializer.cpp b/Engine/Plugins/Runtime/GameplayAbilities/Source/GameplayAbilities/Public/Serialization/GameplayTagCountContainerNetSerializer.cpp
--- a/Engine/Plugins/Runtime/GameplayAbilities/Source/GameplayAbilities/Public/Serialization/GameplayTagCountContainerNetSerializer.cpp
+++ b/Engine/Plugins/Runtime/GameplayAbilities/Source/GameplayAbilities/Public/Serialization/GameplayTagCountContainerNetSerializer.cpp
@@ -128,6 +128,12 @@ void FGameplayTagCountContainerNetSerializer::DeserializeDelta(FNetSerialization
 
 void FGameplayTagCountContainerNetSerializer::Quantize(FNetSerializationContext& Context, const FNetQuantizeArgs& Args)
 {
+	// this is to stop initial replication overriding things incorrectly when we're not supposed to be using the iris serializer.
+	if (!CVarReplicateTagCountContainerWithIris.GetValueOnAnyThread())
+	{
+		return;
+	}
+	
 	const SourceType& Source = *reinterpret_cast<const SourceType*>(Args.Source);
 
 	FNetGameplayTagCountContainerStateForNetSerialize SourceCopy;
@@ -141,6 +147,12 @@ void FGameplayTagCountContainerNetSerializer::Quantize(FNetSerializationContext&
 
 void FGameplayTagCountContainerNetSerializer::Dequantize(FNetSerializationContext& Context, const FNetDequantizeArgs& Args)
 {
+	// this is to stop initial replication overriding things incorrectly when we're not supposed to be using the iris serializer.
+	if (!CVarReplicateTagCountContainerWithIris.GetValueOnAnyThread())
+	{
+		return;
+	}
+	
 	SourceType& TargetValue = *reinterpret_cast<SourceType*>(Args.Target);
 	
 	FNetGameplayTagCountContainerStateForNetSerialize SourceCopy;
