Unshelved from pending changelist '48681436':

Replicated GE tag application fix.

#swarm https://p4-swarm.epicgames.net/changes/48681436

diff --git a/Engine/Plugins/Runtime/GameplayAbilities/Source/GameplayAbilities/Private/AbilitySystemComponent.cpp b/Engine/Plugins/Runtime/GameplayAbilities/Source/GameplayAbilities/Private/AbilitySystemComponent.cpp
--- a/Engine/Plugins/Runtime/GameplayAbilities/Source/GameplayAbilities/Private/AbilitySystemComponent.cpp
+++ b/Engine/Plugins/Runtime/GameplayAbilities/Source/GameplayAbilities/Private/AbilitySystemComponent.cpp
@@ -332,13 +332,13 @@ void UAbilitySystemComponent::CacheIsNetSimulated()
 	UpdateActiveGameplayEffectsReplicationCondition();
 }
 
-void UAbilitySystemComponent::InhibitActiveGameplayEffect(FActiveGameplayEffectHandle ActiveGEHandle, bool bInhibit, bool bInvokeGameplayCueEvents)
+void UAbilitySystemComponent::InhibitActiveGameplayEffect(FActiveGameplayEffectHandle ActiveGEHandle, bool bInhibit, bool bInvokePredictedEffects)
 {
-	FActiveGameplayEffectHandle ContinuationHandle = SetActiveGameplayEffectInhibit(MoveTemp(ActiveGEHandle), bInhibit, bInvokeGameplayCueEvents);
+	FActiveGameplayEffectHandle ContinuationHandle = SetActiveGameplayEffectInhibit(MoveTemp(ActiveGEHandle), bInhibit, bInvokePredictedEffects);
 	ensureMsgf(ContinuationHandle.IsValid(), TEXT("InhibitActiveGameplayEffect invalidated the incoming ActiveGEHandle. Update your code to SetActiveGameplayEffectInhibit so it's clear the incoming handle can be invalidated."));
 }
 
-FActiveGameplayEffectHandle UAbilitySystemComponent::SetActiveGameplayEffectInhibit(FActiveGameplayEffectHandle&& ActiveGEHandle, bool bInhibit, bool bInvokeGameplayCueEvents)
+FActiveGameplayEffectHandle UAbilitySystemComponent::SetActiveGameplayEffectInhibit(FActiveGameplayEffectHandle&& ActiveGEHandle, bool bInhibit, bool bInvokePredictedEffects)
 {
 	FActiveGameplayEffect* ActiveGE = ActiveGameplayEffects.GetActiveGameplayEffect(ActiveGEHandle);
 	if (!ActiveGE)
@@ -360,11 +360,11 @@ FActiveGameplayEffectHandle UAbilitySystemComponent::SetActiveGameplayEffectInhi
 		if (bInhibit)
 		{
 			// Remove our ActiveGameplayEffects modifiers with our Attribute Aggregators
-			ActiveGameplayEffects.RemoveActiveGameplayEffectGrantedTagsAndModifiers(*ActiveGE, bInvokeGameplayCueEvents);
+			ActiveGameplayEffects.RemoveActiveGameplayEffectGrantedTagsAndModifiers(*ActiveGE, bInvokePredictedEffects);
 		}
 		else
 		{
-			ActiveGameplayEffects.AddActiveGameplayEffectGrantedTagsAndModifiers(*ActiveGE, bInvokeGameplayCueEvents);
+			ActiveGameplayEffects.AddActiveGameplayEffectGrantedTagsAndModifiers(*ActiveGE, bInvokePredictedEffects);
 		}
 
 		// The act of executing anything on the ActiveGE can invalidate it.  So we need to recheck if we can continue to execute the callbacks.
diff --git a/Engine/Plugins/Runtime/GameplayAbilities/Source/GameplayAbilities/Private/GameplayEffect.cpp b/Engine/Plugins/Runtime/GameplayAbilities/Source/GameplayAbilities/Private/GameplayEffect.cpp
--- a/Engine/Plugins/Runtime/GameplayAbilities/Source/GameplayAbilities/Private/GameplayEffect.cpp
+++ b/Engine/Plugins/Runtime/GameplayAbilities/Source/GameplayAbilities/Private/GameplayEffect.cpp
@@ -108,6 +108,11 @@ namespace UE::GameplayEffect
 	FAutoConsoleVariableRef CVarReturnAssetTagsViaGetOwnedGameplayTags{ TEXT("AbilitySystem.Fix.GameplayEffects.ReturnAssetTagsViaGetOwnedGameplayTags"), ReturnAssetTagsViaGetOwnedGameplayTags ,
 		TEXT("0 = Return the asset tags as per the function name in get owned gameplay tags. 1 =  Use legacy behavior. If legacy is desired, change code to GetGrantedTags. 2 = Legacy behavior but ensure to catch bad use. 3. Legacy behavior but always ensure to catch multiple issues."), ECVF_Default };
 
+	TAutoConsoleVariable<bool> CVarFixPredictedTagApplication(
+		TEXT("AbilitySystem.Fix.PredictedTagApplication"),
+		true,
+		TEXT("Default: True.  When enabled Gameplay Effect tag application usesthe same logic as Gameplay Cues so that application and removal works correctly for predicted effects."),
+		ECVF_Default);
 #if WITH_EDITOR
 	namespace EditorOnly
 	{
@@ -2751,7 +2756,7 @@ FActiveGameplayEffect& FActiveGameplayEffect::operator=(const FActiveGameplayEff
 	return *this;
 }
 
-void FActiveGameplayEffect::CheckOngoingTagRequirements(const FGameplayTagContainer& OwnerTags, FActiveGameplayEffectsContainer& OwningContainer, bool bInvokeGameplayCueEvents)
+void FActiveGameplayEffect::CheckOngoingTagRequirements(const FGameplayTagContainer& OwnerTags, FActiveGameplayEffectsContainer& OwningContainer, bool bInvokePredictedEffects)
 {
 
 }
@@ -2876,9 +2881,9 @@ void FActiveGameplayEffect::PostReplicatedAdd(const struct FActiveGameplayEffect
 		bPendingRepOnActiveGC,
 		bPendingRepWhileActiveGC);
 
-	// Do stuff for adding GEs (add mods, tags, *invoke callbacks*).  But do NOT invoke the GameplayCues as we don't know if this GE ends up inhibited or not (thus the bPendingRepOnActiveGC variables).
-	constexpr bool bInvokeGameplayCueEvents = false;
-	const_cast<FActiveGameplayEffectsContainer&>(InArray).InternalOnActiveGameplayEffectAdded(*this, bInvokeGameplayCueEvents);	// Const cast is ok. It is there to prevent mutation of the GameplayEffects array, which this wont do.
+	// Do stuff for adding GEs (add mods, tags, *invoke callbacks*).  But do NOT invoke the GameplayCues or tags as we don't know if this GE ends up inhibited or not (thus the bPendingRepOnActiveGC variables).
+	constexpr bool bInvokePredictedEffects = false;
+	const_cast<FActiveGameplayEffectsContainer&>(InArray).InternalOnActiveGameplayEffectAdded(*this, bInvokePredictedEffects);	// Const cast is ok. It is there to prevent mutation of the GameplayEffects array, which this wont do.
 
 	// Make sure no calls caused a realloc of InArray while holding this pointer
 	ensureMsgf(InArray.GameplayEffects_Internal.IsValidIndex(PreAddIdx) && &InArray.GameplayEffects_Internal[PreAddIdx] == this, TEXT("%hs had realloc of the array holding the item pointed to by 'this'. This will eventually crash"), __func__);
@@ -4468,8 +4473,15 @@ PRAGMA_ENABLE_DEPRECATION_WARNINGS
 		{
 			// Once replicated state has caught up to this prediction key, we must remove this gameplay effect.
 			constexpr int32 RemoveAllStacks = -1;
-			InPredictionKey.NewCaughtUpDelegate().BindUObject(Owner, &UAbilitySystemComponent::OnCaughtUpActiveGameplayEffect, AppliedActiveGE->Handle, RemoveAllStacks);
-			InPredictionKey.NewRejectedDelegate().BindUObject(Owner, &UAbilitySystemComponent::OnRejectedActiveGameplayEffect, AppliedActiveGE->Handle, RemoveAllStacks);
+			if (!CVarFixPredictedTagApplication.GetValueOnAnyThread())
+			{
+				InPredictionKey.NewCaughtUpDelegate().BindUObject(Owner, &UAbilitySystemComponent::OnCaughtUpActiveGameplayEffect, AppliedActiveGE->Handle, RemoveAllStacks);
+				InPredictionKey.NewRejectedDelegate().BindUObject(Owner, &UAbilitySystemComponent::OnRejectedActiveGameplayEffect, AppliedActiveGE->Handle, RemoveAllStacks);
+			}
+			else
+			{
+				InPredictionKey.NewRejectOrCaughtUpDelegate(FPredictionKeyEvent::CreateUObject(Owner, &UAbilitySystemComponent::RemoveActiveGameplayEffect_AllowClientRemoval, AppliedActiveGE->Handle, RemoveAllStacks));
+			}
 		}
 		else // We have an existing Stackable GE
 		{
@@ -4492,15 +4504,15 @@ PRAGMA_ENABLE_DEPRECATION_WARNINGS
 	else
 	{
 		// Since we are applying it locally (and possibly predictively) invoke the cues.  Unless it's an Instant Cue, in which case we're not invoking the OnActive/WhileActive cues.
-		const bool bInvokeGameplayCueEvents = (Spec.Def->DurationPolicy != EGameplayEffectDurationType::Instant);
-		InternalOnActiveGameplayEffectAdded(*AppliedActiveGE, bInvokeGameplayCueEvents);
+		const bool bInvokePredictedEffects = (Spec.Def->DurationPolicy != EGameplayEffectDurationType::Instant);
+		InternalOnActiveGameplayEffectAdded(*AppliedActiveGE, bInvokePredictedEffects);
 	}
 
 	return AppliedActiveGE;
 }
 
 /** This is called anytime a new ActiveGameplayEffect is added, on both client and server in all cases */
-void FActiveGameplayEffectsContainer::InternalOnActiveGameplayEffectAdded(FActiveGameplayEffect& Effect, const bool bInvokeGameplayCueEvents)
+void FActiveGameplayEffectsContainer::InternalOnActiveGameplayEffectAdded(FActiveGameplayEffect& Effect, const bool bInvokePredictedEffects)
 {
 	SCOPE_CYCLE_COUNTER(STAT_OnActiveGameplayEffectAdded);
 
@@ -4525,11 +4537,11 @@ void FActiveGameplayEffectsContainer::InternalOnActiveGameplayEffectAdded(FActiv
 	FActiveGameplayEffectHandle EffectHandle = Effect.Handle;
 	if (IsValid(Owner))
 	{
-		Owner->SetActiveGameplayEffectInhibit(MoveTemp(EffectHandle), !bActive, bInvokeGameplayCueEvents);
+		Owner->SetActiveGameplayEffectInhibit(MoveTemp(EffectHandle), !bActive, bInvokePredictedEffects);
 	}
 }
 
-void FActiveGameplayEffectsContainer::AddActiveGameplayEffectGrantedTagsAndModifiers(FActiveGameplayEffect& Effect, bool bInvokeGameplayCueEvents)
+void FActiveGameplayEffectsContainer::AddActiveGameplayEffectGrantedTagsAndModifiers(FActiveGameplayEffect& Effect, bool bInvokePredictedEffects)
 {
 	if (Effect.Spec.Def == nullptr)
 	{
@@ -4595,8 +4607,9 @@ void FActiveGameplayEffectsContainer::AddActiveGameplayEffectGrantedTagsAndModif
 	}
 
 	const bool bIrisWithoutTagCountSerializer = UE::Net::ShouldUseIrisReplication() && !CVarReplicateTagCountContainerWithIris.GetValueOnAnyThread();
+	const bool bShouldApplyTags = UE::GameplayEffect::CVarFixPredictedTagApplication.GetValueOnAnyThread() ? bInvokePredictedEffects : (!Effect.bPostPredictObject || bIrisWithoutTagCountSerializer);
 	// Local application should be responsible for applying tags
-	if (!Effect.bPostPredictObject || bIrisWithoutTagCountSerializer)
+	if (bShouldApplyTags)
 	{
 		// Update our owner with the tags this GameplayEffect grants them
 		Owner->UpdateTagMap(Effect.Spec.Def->GetGrantedTags(), 1, bIrisWithoutTagCountSerializer ? EGameplayTagReplicationState::TagOnly : EGameplayTagReplicationState::CountToOwner);
@@ -4636,7 +4649,7 @@ PRAGMA_ENABLE_DEPRECATION_WARNINGS
 		for (const FGameplayEffectCue& Cue : Effect.Spec.Def->GameplayCues)
 		{
 			// If we use Minimal/Mixed Replication, then this path will AddCue and broadcast the RPC that calls EGameplayCueEvent::OnActive (and WhileActive)
-			// Note: This is going to ignore bInvokeGameplayCueEvents and invoke them anyway (with a bunch of caveats e.g. you're the server and ignoring them)
+			// Note: This is going to ignore bInvokePredictedEffects and invoke them anyway (with a bunch of caveats e.g. you're the server and ignoring them)
 			if (ShouldUseMinimalReplication()) // This can only be true on Authority
 			{
 				for (const FGameplayTag& CueTag : Cue.GameplayCueTags)
@@ -4651,7 +4664,7 @@ PRAGMA_ENABLE_DEPRECATION_WARNINGS
 				// Do a pseudo-AddGameplayCue (but don't add to ActiveGameplayCues so it doesn't replicate in addition to the AGE we're replicating).
 				Owner->UpdateTagMap(Cue.GameplayCueTags, 1, EGameplayTagReplicationState::None);
 
-				if (bInvokeGameplayCueEvents)
+				if (bInvokePredictedEffects)
 				{
 					Owner->InvokeGameplayCueEvent(Effect.Spec, EGameplayCueEvent::OnActive);
 					Owner->InvokeGameplayCueEvent(Effect.Spec, EGameplayCueEvent::WhileActive);
@@ -4844,7 +4857,7 @@ bool FActiveGameplayEffectsContainer::InternalRemoveActiveGameplayEffect(int32 I
 }
 
 /** Called by client and server: This does cleanup that has to happen whether the effect is being removed locally or due to replication */
-void FActiveGameplayEffectsContainer::InternalOnActiveGameplayEffectRemoved(FActiveGameplayEffect& Effect, bool bInvokeGameplayCueEvents, const FGameplayEffectRemovalInfo& GameplayEffectRemovalInfo)
+void FActiveGameplayEffectsContainer::InternalOnActiveGameplayEffectRemoved(FActiveGameplayEffect& Effect, bool bInvokePredictedEffects, const FGameplayEffectRemovalInfo& GameplayEffectRemovalInfo)
 {
 	SCOPE_CYCLE_COUNTER(STAT_OnActiveGameplayEffectRemoved);
 	SCOPE_CYCLE_UOBJECT(EffectDef, Effect.Spec.Def);
@@ -4857,7 +4870,7 @@ void FActiveGameplayEffectsContainer::InternalOnActiveGameplayEffectRemoved(FAct
 		// Only Need to update tags and modifiers if the gameplay effect is active.
 		if (!Effect.bIsInhibited)
 		{
-			RemoveActiveGameplayEffectGrantedTagsAndModifiers(Effect, bInvokeGameplayCueEvents, GameplayEffectRemovalInfo.bPredictionRejected);
+			RemoveActiveGameplayEffectGrantedTagsAndModifiers(Effect, bInvokePredictedEffects, GameplayEffectRemovalInfo.bPredictionRejected);
 		}
 
 		RemoveCustomMagnitudeExternalDependencies(Effect);
@@ -4872,7 +4885,7 @@ void FActiveGameplayEffectsContainer::InternalOnActiveGameplayEffectRemoved(FAct
 	OnActiveGameplayEffectRemovedDelegate.Broadcast(Effect);
 }
 
-void FActiveGameplayEffectsContainer::RemoveActiveGameplayEffectGrantedTagsAndModifiers(const FActiveGameplayEffect& Effect, bool bInvokeGameplayCueEvents, bool bPredictionRejected)
+void FActiveGameplayEffectsContainer::RemoveActiveGameplayEffectGrantedTagsAndModifiers(const FActiveGameplayEffect& Effect, bool bInvokePredictedEffects, bool bPredictionRejected)
 {
 	// Update AttributeAggregators: remove mods from this ActiveGE Handle
 	if (Effect.Spec.GetPeriod() <= UGameplayEffect::NO_PERIOD)
@@ -4890,8 +4903,9 @@ void FActiveGameplayEffectsContainer::RemoveActiveGameplayEffectGrantedTagsAndMo
 	}
 
 	const bool bIrisWithoutTagCountSerializer = UE::Net::ShouldUseIrisReplication() && !CVarReplicateTagCountContainerWithIris.GetValueOnAnyThread();
-	// We're no longer predicting for remove unless it was rejected, so let the server replication handle the tags.
-	if (IsNetAuthority() || bPredictionRejected || bIrisWithoutTagCountSerializer)
+	const bool bShouldApplyTags = UE::GameplayEffect::CVarFixPredictedTagApplication.GetValueOnAnyThread() ? bInvokePredictedEffects : (IsNetAuthority() || bPredictionRejected || bIrisWithoutTagCountSerializer);
+	// Local application should be responsible for applying tags
+	if (bShouldApplyTags)
 	{
 		// Update gameplaytag count and broadcast delegate if we are at 0
 		Owner->UpdateTagMap(Effect.Spec.Def->GetGrantedTags(), -1, bIrisWithoutTagCountSerializer ? EGameplayTagReplicationState::TagOnly : EGameplayTagReplicationState::CountToOwner);
@@ -4950,7 +4964,7 @@ PRAGMA_ENABLE_DEPRECATION_WARNINGS
 				// Perform pseudo-RemoveCue (without affecting ActiveGameplayCues, as we were not inserted there - see AddActiveGameplayEffectGrantedTagsAndModifiers)
 				Owner->UpdateTagMap(Cue.GameplayCueTags, -1, EGameplayTagReplicationState::None);
 
-				if (bInvokeGameplayCueEvents)
+				if (bInvokePredictedEffects)
 				{
 					Owner->InvokeGameplayCueEvent(Effect.Spec, EGameplayCueEvent::Removed);
 				}
diff --git a/Engine/Plugins/Runtime/GameplayAbilities/Source/GameplayAbilities/Private/GameplayEffectTypes.cpp b/Engine/Plugins/Runtime/GameplayAbilities/Source/GameplayAbilities/Private/GameplayEffectTypes.cpp
--- a/Engine/Plugins/Runtime/GameplayAbilities/Source/GameplayAbilities/Private/GameplayEffectTypes.cpp
+++ b/Engine/Plugins/Runtime/GameplayAbilities/Source/GameplayAbilities/Private/GameplayEffectTypes.cpp
@@ -47,7 +47,7 @@ static FAutoConsoleVariableRef CVarWarnIfTryingToReplicateNotSupportedActorRefer
 TAutoConsoleVariable<bool> CVarReplicateTagCountContainerWithIris(
 		TEXT("AbilitySystem.Fix.ReplicateTagCountContainerWithIris"),
 		false,
-		TEXT("Default: True.  When true will use Iris serialization for Gameplay Tag Count Container enabling tag rep state promotion with iris."),
+		TEXT("Default: False.  When true, will use Iris serialization for Gameplay Tag Count Container enabling tag rep state promotion with iris."),
 		ECVF_Default);
 
 
diff --git a/Engine/Plugins/Runtime/GameplayAbilities/Source/GameplayAbilities/Public/AbilitySystemComponent.h b/Engine/Plugins/Runtime/GameplayAbilities/Source/GameplayAbilities/Public/AbilitySystemComponent.h
--- a/Engine/Plugins/Runtime/GameplayAbilities/Source/GameplayAbilities/Public/AbilitySystemComponent.h
+++ b/Engine/Plugins/Runtime/GameplayAbilities/Source/GameplayAbilities/Public/AbilitySystemComponent.h
@@ -439,14 +439,14 @@ class UAbilitySystemComponent : public UGameplayTasksComponent, public IGameplay
 
 	/** Inhibit an active gameplay effect so that it is disabled, but not removed */
 	UE_DEPRECATED(5.4, "Use SetActiveGameplayEffectInhibit with a MoveTemp(ActiveGEHandle) so it's clear the Handle is no longer valid. Check (then use) the returned FActiveGameplayEffectHandle to continue your operation.")
-	UE_API virtual void InhibitActiveGameplayEffect(FActiveGameplayEffectHandle ActiveGEHandle, bool bInhibit, bool bInvokeGameplayCueEvents);
+	UE_API virtual void InhibitActiveGameplayEffect(FActiveGameplayEffectHandle ActiveGEHandle, bool bInhibit, bool bInvokePredictedEffects);
 
 	/**
 	 * (Un-)Inhibit an Active Gameplay Effect so it may be disabled (and perform some disabling actions, such as removing tags).
 	 * An inhibited Active Gameplay Effect will lay dormant for re-enabling on some condition (usually tags).  When it's uninhibited, it will reapply some of its functionality.
 	 * NOTE:  The passed-in ActiveGEHandle is no longer valid.  Use the returned FActiveGameplayEffectHandle to determine if the Active GE Handle is still active.
 	 */
-	UE_API virtual FActiveGameplayEffectHandle SetActiveGameplayEffectInhibit(FActiveGameplayEffectHandle&& ActiveGEHandle, bool bInhibit, bool bInvokeGameplayCueEvents);
+	UE_API virtual FActiveGameplayEffectHandle SetActiveGameplayEffectInhibit(FActiveGameplayEffectHandle&& ActiveGEHandle, bool bInhibit, bool bInvokePredictedEffects);
 
 	/**
 	 * Raw accessor to ask the magnitude of a gameplay effect, not necessarily always correct. How should outside code (UI, etc) ask things like 'how much is this gameplay effect modifying my damage by'
diff --git a/Engine/Plugins/Runtime/GameplayAbilities/Source/GameplayAbilities/Public/GameplayEffect.h b/Engine/Plugins/Runtime/GameplayAbilities/Source/GameplayAbilities/Public/GameplayEffect.h
--- a/Engine/Plugins/Runtime/GameplayAbilities/Source/GameplayAbilities/Public/GameplayEffect.h
+++ b/Engine/Plugins/Runtime/GameplayAbilities/Source/GameplayAbilities/Public/GameplayEffect.h
@@ -1372,7 +1372,7 @@ struct FActiveGameplayEffect : public FFastArraySerializerItem
 
 	/** This was the core function that turns the ActiveGE 'on' or 'off.  That function can be carried out by UGameplayEffectComponents, @see UTargetTagRequirementsGameplayEffectComponent */
 	UE_DEPRECATED(5.3, "CheckOngoingTagRequirements has been deprecated in favor of UTargetTagRequirementsGameplayEffectComponent")
-	UE_API void CheckOngoingTagRequirements(const FGameplayTagContainer& OwnerTags, struct FActiveGameplayEffectsContainer& OwningContainer, bool bInvokeGameplayCueEvents = false);
+	UE_API void CheckOngoingTagRequirements(const FGameplayTagContainer& OwnerTags, struct FActiveGameplayEffectsContainer& OwningContainer, bool bInvokePredictedEffects = false);
 	
 	/** Method to check if this effect should remove because the owner tags pass the RemovalTagRequirements requirement check */
 	UE_DEPRECATED(5.3, "CheckRemovalTagRequirements has been deprecated in favor of UTargetTagRequirementsGameplayEffectComponent")
@@ -1930,11 +1930,11 @@ private:
 	UE_API bool InternalRemoveActiveGameplayEffect(int32 Idx, int32 StacksToRemove, bool bPrematureRemoval, bool bPredictionRejected = false);
 	
 	/** Called both in server side creation and replication creation/deletion */
-	UE_API void InternalOnActiveGameplayEffectAdded(FActiveGameplayEffect& Effect, bool bInvokeGameplayCueEvents);
-	UE_API void InternalOnActiveGameplayEffectRemoved(FActiveGameplayEffect& Effect, bool bInvokeGameplayCueEvents, const FGameplayEffectRemovalInfo& GameplayEffectRemovalInfo);
+	UE_API void InternalOnActiveGameplayEffectAdded(FActiveGameplayEffect& Effect, bool bInvokePredictedEffects);
+	UE_API void InternalOnActiveGameplayEffectRemoved(FActiveGameplayEffect& Effect, bool bInvokePredictedEffects, const FGameplayEffectRemovalInfo& GameplayEffectRemovalInfo);
 
-	UE_API void RemoveActiveGameplayEffectGrantedTagsAndModifiers(const FActiveGameplayEffect& Effect, bool bInvokeGameplayCueEvents, bool bPredictionRejected = false);
-	UE_API void AddActiveGameplayEffectGrantedTagsAndModifiers(FActiveGameplayEffect& Effect, bool bInvokeGameplayCueEvents);
+	UE_API void RemoveActiveGameplayEffectGrantedTagsAndModifiers(const FActiveGameplayEffect& Effect, bool bInvokePredictedEffects, bool bPredictionRejected = false);
+	UE_API void AddActiveGameplayEffectGrantedTagsAndModifiers(FActiveGameplayEffect& Effect, bool bInvokePredictedEffects);
 
 	/** Updates tag dependency map when a GameplayEffect is removed */
 	UE_API void RemoveActiveEffectTagDependency(const FGameplayTagContainer& Tags, FActiveGameplayEffectHandle Handle);
