[GAS] Missed a case where we disable Iris GameplayTagCountNetSerializer and need to apply the predicted tags in that case.

diff --git a/Engine/Plugins/Runtime/GameplayAbilities/Source/GameplayAbilities/Private/GameplayEffect.cpp b/Engine/Plugins/Runtime/GameplayAbilities/Source/GameplayAbilities/Private/GameplayEffect.cpp
--- a/Engine/Plugins/Runtime/GameplayAbilities/Source/GameplayAbilities/Private/GameplayEffect.cpp
+++ b/Engine/Plugins/Runtime/GameplayAbilities/Source/GameplayAbilities/Private/GameplayEffect.cpp
@@ -4606,8 +4606,11 @@ void FActiveGameplayEffectsContainer::AddActiveGameplayEffectGrantedTagsAndModif
 		}
 	}
 
+	// Generic Replication will replicate Tags w/ Counts correctly.  So we should only apply the tags if we're predicting, otherwise we will get the tags of a Server-Initiated GE through GameplayTagCountContainer replication.
+	// Iris Replication is more complicated. CVarReplicateTagCountContainerWithIris indicates that we do not have working Tag Count replication through GameplayTagCountContainerNetSerializer. Instead, we replicate using the legacy path of separated containers (Minimal/Loose).
+	// The complex bShouldApplyTags case is we should still apply the tags of a Server-Initiated GE because we won't get those Tags through Iris GameplayTagCountContainerNetSerializer replication.  We should assume CVarFixPredictedTagApplication is true.
 	const bool bIrisWithoutTagCountSerializer = UE::Net::ShouldUseIrisReplication() && !CVarReplicateTagCountContainerWithIris.GetValueOnAnyThread();
-	const bool bShouldApplyTags = UE::GameplayEffect::CVarFixPredictedTagApplication.GetValueOnAnyThread() ? bInvokePredictedEffects : (!Effect.bPostPredictObject || bIrisWithoutTagCountSerializer);
+	const bool bShouldApplyTags = UE::GameplayEffect::CVarFixPredictedTagApplication.GetValueOnAnyThread() ? (bInvokePredictedEffects || bIrisWithoutTagCountSerializer) : (!Effect.bPostPredictObject || bIrisWithoutTagCountSerializer);
 	// Local application should be responsible for applying tags
 	if (bShouldApplyTags)
 	{
@@ -4902,8 +4905,11 @@ void FActiveGameplayEffectsContainer::RemoveActiveGameplayEffectGrantedTagsAndMo
 		}
 	}
 
+	// Generic Replication will replicate Tags w/ Counts correctly.  So we should only apply the tags if we're predicting, otherwise we will get the tags of a Server-Initiated GE through GameplayTagCountContainer replication.
+	// Iris Replication is more complicated. CVarReplicateTagCountContainerWithIris indicates that we do not have working Tag Count replication through GameplayTagCountContainerNetSerializer. Instead, we replicate using the legacy path of separated containers (Minimal/Loose).
+	// The complex bShouldApplyTags case is we should still apply the tags of a Server-Initiated GE because we won't get those Tags through Iris GameplayTagCountContainerNetSerializer replication.  We should assume CVarFixPredictedTagApplication is true.
 	const bool bIrisWithoutTagCountSerializer = UE::Net::ShouldUseIrisReplication() && !CVarReplicateTagCountContainerWithIris.GetValueOnAnyThread();
-	const bool bShouldApplyTags = UE::GameplayEffect::CVarFixPredictedTagApplication.GetValueOnAnyThread() ? bInvokePredictedEffects : (IsNetAuthority() || bPredictionRejected || bIrisWithoutTagCountSerializer);
+	const bool bShouldApplyTags = UE::GameplayEffect::CVarFixPredictedTagApplication.GetValueOnAnyThread() ? (bInvokePredictedEffects || bIrisWithoutTagCountSerializer) : (IsNetAuthority() || bPredictionRejected || bIrisWithoutTagCountSerializer);
 	// Local application should be responsible for applying tags
 	if (bShouldApplyTags)
 	{
