[GAS] Quick fix for GEComponents accessing other UObjects during PostInitProperties while AsyncLoading is enabled.
#Synced-CL 48766133

diff --git a/Engine/Plugins/Runtime/GameplayAbilities/Source/GameplayAbilities/Private/GameplayEffectComponents/AssetTagsGameplayEffectComponent.cpp b/Engine/Plugins/Runtime/GameplayAbilities/Source/GameplayAbilities/Private/GameplayEffectComponents/AssetTagsGameplayEffectComponent.cpp
--- a/Engine/Plugins/Runtime/GameplayAbilities/Source/GameplayAbilities/Private/GameplayEffectComponents/AssetTagsGameplayEffectComponent.cpp
+++ b/Engine/Plugins/Runtime/GameplayAbilities/Source/GameplayAbilities/Private/GameplayEffectComponents/AssetTagsGameplayEffectComponent.cpp
@@ -12,17 +12,34 @@ void UAssetTagsGameplayEffectComponent::PostInitProperties()
 {
 	Super::PostInitProperties();
 
+#if WITH_EDITORONLY_DATA
+	EditorFriendlyName = LOCTEXT("AssetTagsOnGE", "Asset Tags (on Gameplay Effect)").ToString();
+#endif
+
+	// This code is meant to catch the case of newly constructed Components (not loaded Components).
+	// With asynchronous loading, we are not guaranteed the Parent object is loaded at this point.
+	// During loading, we'll call OnGameplayEffectChanged from UGameplayEffect::PostLoad().
+	if (IsInAsyncLoadingThread() || IsInParallelLoadingThread())
+	{
+		return;
+	}
+
 	// Try to find the parent and update the inherited tags.  We do this 'early' because this is the only function
 	// we can use for creation of the object (and this runs post-constructor).
 	const UAssetTagsGameplayEffectComponent* Parent = FindParentComponent(*this);
 	if (Parent == this)
 	{
-		UE_LOGFMT(LogAbilitySystem, Warning, "UAssetTagsGameplayEffectComponent got self for parent compontnt when trying to update inherited tag container. {Component}", FSoftObjectPath(this));
+		UE_LOGFMT(LogAbilitySystem, Warning, "UAssetTagsGameplayEffectComponent got self for parent component when trying to update inherited tag container. {Component}", FSoftObjectPath(this));
 	}
 	else 
 	{
 		InheritableAssetTags.UpdateInheritedTagProperties(Parent ? &Parent->InheritableAssetTags : nullptr);
 	}
+}
+
+void UAssetTagsGameplayEffectComponent::OnGameplayEffectChanged()
+{
+	Super::OnGameplayEffectChanged();
 
 #if WITH_EDITORONLY_DATA
 	if (GetOwner())
@@ -34,11 +51,7 @@ void UAssetTagsGameplayEffectComponent::PostInitProperties()
 		EditorFriendlyName = LOCTEXT("AssetTagsOnGE", "Asset Tags (on Gameplay Effect)").ToString();
 	}
 #endif // WITH_EDITORONLY_DATA
-}
 
-void UAssetTagsGameplayEffectComponent::OnGameplayEffectChanged()
-{
-	Super::OnGameplayEffectChanged();
 	SetAndApplyAssetTagChanges(InheritableAssetTags);
 }
 
diff --git a/Engine/Plugins/Runtime/GameplayAbilities/Source/GameplayAbilities/Private/GameplayEffectComponents/BlockAbilityTagsGameplayEffectComponent.cpp b/Engine/Plugins/Runtime/GameplayAbilities/Source/GameplayAbilities/Private/GameplayEffectComponents/BlockAbilityTagsGameplayEffectComponent.cpp
--- a/Engine/Plugins/Runtime/GameplayAbilities/Source/GameplayAbilities/Private/GameplayEffectComponents/BlockAbilityTagsGameplayEffectComponent.cpp
+++ b/Engine/Plugins/Runtime/GameplayAbilities/Source/GameplayAbilities/Private/GameplayEffectComponents/BlockAbilityTagsGameplayEffectComponent.cpp
@@ -11,14 +11,22 @@ void UBlockAbilityTagsGameplayEffectComponent::PostInitProperties()
 {
 	Super::PostInitProperties();
 
+#if WITH_EDITORONLY_DATA
+	EditorFriendlyName = TEXT("Block Abilities w/ Tags");
+#endif
+
+	// This code is meant to catch the case of newly constructed Components (not loaded Components).
+	// With asynchronous loading, we are not guaranteed the Parent object is loaded at this point.
+	// During loading, we'll call OnGameplayEffectChanged from UGameplayEffect::PostLoad().
+	if (IsInAsyncLoadingThread() || IsInParallelLoadingThread())
+	{
+		return;
+	}
+
 	// Try to find the parent and update the inherited tags.  We do this 'early' because this is the only function
 	// we can use for creation of the object (and this runs post-constructor).
 	const UBlockAbilityTagsGameplayEffectComponent* Parent = FindParentComponent(*this);
 	InheritableBlockedAbilityTagsContainer.UpdateInheritedTagProperties(Parent ? &Parent->InheritableBlockedAbilityTagsContainer : nullptr);
-
-#if WITH_EDITORONLY_DATA
-	EditorFriendlyName = TEXT("Block Abilities w/ Tags");
-#endif
 }
 
 void UBlockAbilityTagsGameplayEffectComponent::OnGameplayEffectChanged()
diff --git a/Engine/Plugins/Runtime/GameplayAbilities/Source/GameplayAbilities/Private/GameplayEffectComponents/CancelAbilityTagsGameplayEffectComponent.cpp b/Engine/Plugins/Runtime/GameplayAbilities/Source/GameplayAbilities/Private/GameplayEffectComponents/CancelAbilityTagsGameplayEffectComponent.cpp
--- a/Engine/Plugins/Runtime/GameplayAbilities/Source/GameplayAbilities/Private/GameplayEffectComponents/CancelAbilityTagsGameplayEffectComponent.cpp
+++ b/Engine/Plugins/Runtime/GameplayAbilities/Source/GameplayAbilities/Private/GameplayEffectComponents/CancelAbilityTagsGameplayEffectComponent.cpp
@@ -13,15 +13,23 @@ void UCancelAbilityTagsGameplayEffectComponent::PostInitProperties()
 {
 	Super::PostInitProperties();
 
+#if WITH_EDITORONLY_DATA
+	EditorFriendlyName = TEXT("Cancel Abilities With Tags");
+#endif
+
+	// This code is meant to catch the case of newly constructed Components (not loaded Components).
+	// With asynchronous loading, we are not guaranteed the Parent object is loaded at this point.
+	// During loading, we'll call OnGameplayEffectChanged from UGameplayEffect::PostLoad().
+	if (IsInAsyncLoadingThread() || IsInParallelLoadingThread())
+	{
+		return;
+	}
+
 	// Try to find the parent and update the inherited tags.  We do this 'early' because this is the only function
 	// we can use for creation of the object (and this runs post-constructor).
 	const UCancelAbilityTagsGameplayEffectComponent* Parent = FindParentComponent(*this);
 	InheritableCancelAbilitiesWithTagsContainer.UpdateInheritedTagProperties(Parent ? &Parent->InheritableCancelAbilitiesWithTagsContainer : nullptr);
 	InheritableCancelAbilitiesWithoutTagsContainer.UpdateInheritedTagProperties(Parent ? &Parent->InheritableCancelAbilitiesWithoutTagsContainer : nullptr);
-
-#if WITH_EDITORONLY_DATA
-	EditorFriendlyName = TEXT("Cancel Abilities With Tags");
-#endif
 }
 
 void UCancelAbilityTagsGameplayEffectComponent::OnGameplayEffectApplied(FActiveGameplayEffectsContainer& ActiveGEContainer, FGameplayEffectSpec& GESpec, FPredictionKey& PredictionKey) const
diff --git a/Engine/Plugins/Runtime/GameplayAbilities/Source/GameplayAbilities/Private/GameplayEffectComponents/TargetTagsGameplayEffectComponent.cpp b/Engine/Plugins/Runtime/GameplayAbilities/Source/GameplayAbilities/Private/GameplayEffectComponents/TargetTagsGameplayEffectComponent.cpp
--- a/Engine/Plugins/Runtime/GameplayAbilities/Source/GameplayAbilities/Private/GameplayEffectComponents/TargetTagsGameplayEffectComponent.cpp
+++ b/Engine/Plugins/Runtime/GameplayAbilities/Source/GameplayAbilities/Private/GameplayEffectComponents/TargetTagsGameplayEffectComponent.cpp
@@ -11,14 +11,22 @@ void UTargetTagsGameplayEffectComponent::PostInitProperties()
 {
 	Super::PostInitProperties();
 
+#if WITH_EDITORONLY_DATA
+	EditorFriendlyName = TEXT("Target Tags (Granted to Actor)");
+#endif // WITH_EDITORONLY_DATA
+
+	// This code is meant to catch the case of newly constructed Components (not loaded Components).
+	// With asynchronous loading, we are not guaranteed the Parent object is loaded at this point.
+	// During loading, we'll call OnGameplayEffectChanged from UGameplayEffect::PostLoad().
+	if (IsInAsyncLoadingThread() || IsInParallelLoadingThread())
+	{
+		return;
+	}
+
 	// Try to find the parent and update the inherited tags.  We do this 'early' because this is the only function
 	// we can use for creation of the object (and this runs post-constructor).
 	const UTargetTagsGameplayEffectComponent* Parent = FindParentComponent(*this);
 	InheritableGrantedTagsContainer.UpdateInheritedTagProperties(Parent ? &Parent->InheritableGrantedTagsContainer : nullptr);
-
-#if WITH_EDITORONLY_DATA
-	EditorFriendlyName = TEXT("Target Tags (Granted to Actor)");
-#endif // WITH_EDITORONLY_DATA
 }
 
 void UTargetTagsGameplayEffectComponent::OnGameplayEffectChanged()
