Avoid unnecessary memory allocations in Gameplay Effect captures
#rnx
#Synced-CL 49215206

diff --git a/Engine/Plugins/Runtime/GameplayAbilities/Source/GameplayAbilities/Private/GameplayEffectAggregator.cpp b/Engine/Plugins/Runtime/GameplayAbilities/Source/GameplayAbilities/Private/GameplayEffectAggregator.cpp
--- a/Engine/Plugins/Runtime/GameplayAbilities/Source/GameplayAbilities/Private/GameplayEffectAggregator.cpp
+++ b/Engine/Plugins/Runtime/GameplayAbilities/Source/GameplayAbilities/Private/GameplayEffectAggregator.cpp
@@ -656,7 +656,7 @@ void FAggregatorRef::TakeSnapshotOf(const FAggregatorRef& RefToSnapshot)
 	{
 		FAggregator* SrcData = RefToSnapshot.Data.Get();
 
-		Data = TSharedPtr<FAggregator>(new FAggregator());
+		Data = MakeShared<FAggregator>();
 		Data->TakeSnapshotOf(*SrcData);
 	}
 	else
diff --git a/Engine/Plugins/Runtime/GameplayAbilities/Source/GameplayAbilities/Public/GameplayEffectAttributeCaptureDefinition.h b/Engine/Plugins/Runtime/GameplayAbilities/Source/GameplayAbilities/Public/GameplayEffectAttributeCaptureDefinition.h
--- a/Engine/Plugins/Runtime/GameplayAbilities/Source/GameplayAbilities/Public/GameplayEffectAttributeCaptureDefinition.h
+++ b/Engine/Plugins/Runtime/GameplayAbilities/Source/GameplayAbilities/Public/GameplayEffectAttributeCaptureDefinition.h
@@ -31,7 +31,7 @@ struct FGameplayEffectAttributeCaptureDefinition
 	}
 
 	FGameplayEffectAttributeCaptureDefinition(FGameplayAttribute InAttribute, EGameplayEffectAttributeCaptureSource InSource, bool InSnapshot)
-		: AttributeToCapture(InAttribute), AttributeSource(InSource), bSnapshot(InSnapshot)
+		: AttributeToCapture(MoveTemp(InAttribute)), AttributeSource(InSource), bSnapshot(InSnapshot)
 	{
 
 	}
