[GAS] FActiveGameplayEffectHandle now holds a reference to its Owning AbilitySystemComponent rather than using a Global Map.  The FActiveGameplayEffectHandle constructor taking an int32 Handle value has been deprecated; its use has been replaced with FActiveGameplayEffectHandle::GetInstantExecutedHandle.  FGameplayEffectRemovalInfo now contains a pointer to the ActiveGE's owning ASC.

#Synced-CL 49222664

diff --git a/Engine/Plugins/Runtime/GameplayAbilities/Source/GameplayAbilities/Private/AbilitySystemBlueprintLibrary.cpp b/Engine/Plugins/Runtime/GameplayAbilities/Source/GameplayAbilities/Private/AbilitySystemBlueprintLibrary.cpp
--- a/Engine/Plugins/Runtime/GameplayAbilities/Source/GameplayAbilities/Private/AbilitySystemBlueprintLibrary.cpp
+++ b/Engine/Plugins/Runtime/GameplayAbilities/Source/GameplayAbilities/Private/AbilitySystemBlueprintLibrary.cpp
@@ -1220,7 +1220,8 @@ bool UAbilitySystemBlueprintLibrary::IsActiveGameplayEffectHandleValid(FActiveGa
 
 bool UAbilitySystemBlueprintLibrary::IsActiveGameplayEffectHandleActive(FActiveGameplayEffectHandle Handle)
 {
-	return Handle.IsValid() && Handle.GetOwningAbilitySystemComponent() && Handle.GetOwningAbilitySystemComponent()->GetActiveGameplayEffect(Handle) != nullptr;
+	// GetOwningAbilitySystemComponent will only return a valid pointer if the effect is still active
+	return Handle.IsValid() && Handle.GetOwningAbilitySystemComponent();
 }
 
 UAbilitySystemComponent* UAbilitySystemBlueprintLibrary::GetAbilitySystemComponentFromActiveGameplayEffectHandle(FActiveGameplayEffectHandle Handle)
diff --git a/Engine/Plugins/Runtime/GameplayAbilities/Source/GameplayAbilities/Private/AbilitySystemComponent.cpp b/Engine/Plugins/Runtime/GameplayAbilities/Source/GameplayAbilities/Private/AbilitySystemComponent.cpp
--- a/Engine/Plugins/Runtime/GameplayAbilities/Source/GameplayAbilities/Private/AbilitySystemComponent.cpp
+++ b/Engine/Plugins/Runtime/GameplayAbilities/Source/GameplayAbilities/Private/AbilitySystemComponent.cpp
@@ -1040,9 +1040,8 @@ FActiveGameplayEffectHandle UAbilitySystemComponent::ApplyGameplayEffectSpecToSe
 	bool bTreatAsInfiniteDuration = GetOwnerRole() != ROLE_Authority && PredictionKey.IsLocalClientKey() && Spec.Def->DurationPolicy == EGameplayEffectDurationType::Instant;
 
 	// Make sure we create our copy of the spec in the right place
-	// We initialize the FActiveGameplayEffectHandle here with INDEX_NONE to handle the case of instant GE
-	// Initializing it like this will set the bPassedFiltersAndWasExecuted on the FActiveGameplayEffectHandle to true so we can know that we applied a GE
-	FActiveGameplayEffectHandle	MyHandle(INDEX_NONE);
+	// We default to returning an FActiveGameplayEffectHandle that denotes the GE was successfully executed instantly
+	FActiveGameplayEffectHandle MyHandle = FActiveGameplayEffectHandle::GetInstantExecutedHandle();
 	bool bInvokeGameplayCueApplied = Spec.Def->DurationPolicy != EGameplayEffectDurationType::Instant; // Cache this now before possibly modifying predictive instant effect to infinite duration effect.
 	bool bFoundExistingStackableGE = false;
 
diff --git a/Engine/Plugins/Runtime/GameplayAbilities/Source/GameplayAbilities/Private/AbilitySystemComponent_Abilities.cpp b/Engine/Plugins/Runtime/GameplayAbilities/Source/GameplayAbilities/Private/AbilitySystemComponent_Abilities.cpp
--- a/Engine/Plugins/Runtime/GameplayAbilities/Source/GameplayAbilities/Private/AbilitySystemComponent_Abilities.cpp
+++ b/Engine/Plugins/Runtime/GameplayAbilities/Source/GameplayAbilities/Private/AbilitySystemComponent_Abilities.cpp
@@ -705,7 +705,8 @@ PRAGMA_ENABLE_DEPRECATION_WARNINGS
 	{
 		if (UAbilitySystemComponent* SourceASC = Spec.GameplayEffectHandle.GetOwningAbilitySystemComponent())
 		{
-			if (FActiveGameplayEffect* SourceActiveGE = SourceASC->ActiveGameplayEffects.GetActiveGameplayEffect(Spec.GameplayEffectHandle))
+			FActiveGameplayEffect* SourceActiveGE = SourceASC->ActiveGameplayEffects.GetActiveGameplayEffect(Spec.GameplayEffectHandle);
+			if (ensureMsgf(SourceActiveGE, TEXT("FActiveGameplayEffectHandle::GetOwningAbilitySystemComponent returned an ASC which did not contain the handle")))
 			{
 				SourceActiveGE->GrantedAbilityHandles.Remove(Spec.Handle);
 			}
diff --git a/Engine/Plugins/Runtime/GameplayAbilities/Source/GameplayAbilities/Private/AbilitySystemGlobals.cpp b/Engine/Plugins/Runtime/GameplayAbilities/Source/GameplayAbilities/Private/AbilitySystemGlobals.cpp
--- a/Engine/Plugins/Runtime/GameplayAbilities/Source/GameplayAbilities/Private/AbilitySystemGlobals.cpp
+++ b/Engine/Plugins/Runtime/GameplayAbilities/Source/GameplayAbilities/Private/AbilitySystemGlobals.cpp
@@ -694,7 +694,6 @@ void UAbilitySystemGlobals::OnPreBeginPIE(const bool bIsSimulatingInEditor)
 void UAbilitySystemGlobals::ResetCachedData()
 {
 	IGameplayCueInterface::ClearTagToFunctionMap();
-	FActiveGameplayEffectHandle::ResetGlobalHandleMap();
 }
 
 void UAbilitySystemGlobals::HandlePreLoadMap(const FWorldContext& WorldContext, const FString& MapName)
diff --git a/Engine/Plugins/Runtime/GameplayAbilities/Source/GameplayAbilities/Private/ActiveGameplayEffectHandle.cpp b/Engine/Plugins/Runtime/GameplayAbilities/Source/GameplayAbilities/Private/ActiveGameplayEffectHandle.cpp
--- a/Engine/Plugins/Runtime/GameplayAbilities/Source/GameplayAbilities/Private/ActiveGameplayEffectHandle.cpp
+++ b/Engine/Plugins/Runtime/GameplayAbilities/Source/GameplayAbilities/Private/ActiveGameplayEffectHandle.cpp
@@ -5,40 +5,51 @@
 
 #include UE_INLINE_GENERATED_CPP_BY_NAME(ActiveGameplayEffectHandle)
 
-namespace GlobalActiveGameplayEffectHandles
+namespace
 {
-	static TMap<FActiveGameplayEffectHandle, TWeakObjectPtr<UAbilitySystemComponent>>	Map;
+	static int32 GetNextActiveGEHandleID()
+	{
+		static int32 GHandleID = 0;
+		int32 NewHandleID = ++GHandleID;
+		if (GHandleID < 1)
+		{
+			GHandleID = 1;
+		}
+
+		return NewHandleID;
+	}
 }
 
-void FActiveGameplayEffectHandle::ResetGlobalHandleMap()
+FActiveGameplayEffectHandle::FActiveGameplayEffectHandle(UAbilitySystemComponent* InOwningASC)
+	: WeakOwningASC(InOwningASC)
+	, Handle(GetNextActiveGEHandleID())
+	, bPassedFiltersAndWasExecuted(true)
 {
-	GlobalActiveGameplayEffectHandles::Map.Reset();
 }
 
-FActiveGameplayEffectHandle FActiveGameplayEffectHandle::GenerateNewHandle(UAbilitySystemComponent* OwningComponent)
+FActiveGameplayEffectHandle FActiveGameplayEffectHandle::GetInstantExecutedHandle()
 {
-	static int32 GHandleID=0;
-	FActiveGameplayEffectHandle NewHandle(GHandleID++);
+	FActiveGameplayEffectHandle InstantGEHandle;
+	InstantGEHandle.bPassedFiltersAndWasExecuted = true;
 
-	TWeakObjectPtr<UAbilitySystemComponent> WeakPtr(OwningComponent);
-
-	GlobalActiveGameplayEffectHandles::Map.Add(NewHandle, WeakPtr);
+	return InstantGEHandle;
+}
 
-	return NewHandle;
+FActiveGameplayEffectHandle FActiveGameplayEffectHandle::GenerateNewHandle(UAbilitySystemComponent* OwningComponent)
+{
+	return FActiveGameplayEffectHandle(OwningComponent);
 }
 
 UAbilitySystemComponent* FActiveGameplayEffectHandle::GetOwningAbilitySystemComponent() const
 {
-	TWeakObjectPtr<UAbilitySystemComponent>* Ptr = GlobalActiveGameplayEffectHandles::Map.Find(*this);
-	if (Ptr)
+	if (UAbilitySystemComponent* OwningASC = WeakOwningASC.Get())
 	{
-		return Ptr->Get();
+		// Check if this ActiveGE is still valid.  Note: This will fail if PendingRemove is true
+		if (const FActiveGameplayEffect* ActiveGE = OwningASC->GetActiveGameplayEffect(*this))
+		{
+			return OwningASC;
+		}
 	}
 
 	return nullptr;
 }
-
-void FActiveGameplayEffectHandle::RemoveFromGlobalMap()
-{
-	GlobalActiveGameplayEffectHandles::Map.Remove(*this);
-}
diff --git a/Engine/Plugins/Runtime/GameplayAbilities/Source/GameplayAbilities/Private/GameplayEffect.cpp b/Engine/Plugins/Runtime/GameplayAbilities/Source/GameplayAbilities/Private/GameplayEffect.cpp
--- a/Engine/Plugins/Runtime/GameplayAbilities/Source/GameplayAbilities/Private/GameplayEffect.cpp
+++ b/Engine/Plugins/Runtime/GameplayAbilities/Source/GameplayAbilities/Private/GameplayEffect.cpp
@@ -2782,6 +2782,7 @@ void FActiveGameplayEffect::PreReplicatedRemove(const struct FActiveGameplayEffe
 	FGameplayEffectRemovalInfo GameplayEffectRemovalInfo;
 	GameplayEffectRemovalInfo.ActiveEffect = this;
 	GameplayEffectRemovalInfo.StackCount = ClientCachedStackCount;
+	GameplayEffectRemovalInfo.OwningASC = InArray.Owner;
 	//Check duration to set bPrematureRemoval as req.
 	if (DurationHandle.IsValid())
 	{
@@ -4769,6 +4770,7 @@ bool FActiveGameplayEffectsContainer::InternalRemoveActiveGameplayEffect(int32 I
 		GameplayEffectRemovalInfo.bPrematureRemoval = bPrematureRemoval;
 		GameplayEffectRemovalInfo.bPredictionRejected = bPredictionRejected;
 		GameplayEffectRemovalInfo.EffectContext = Effect.Spec.GetEffectContext();
+		GameplayEffectRemovalInfo.OwningASC = Owner;
 
 		if (StacksToRemove > 0 && Effect.Spec.GetStackCount() > StacksToRemove)
 		{
@@ -4833,9 +4835,6 @@ bool FActiveGameplayEffectsContainer::InternalRemoveActiveGameplayEffect(int32 I
 		{
 			// Not locked, so do the removal right away.
 
-			// Remove this handle from the global map
-			Effect.Handle.RemoveFromGlobalMap();
-
 			// If we are not scope locked, then there is no way this idx should be referring to something on the pending add list.
 			// It is possible to remove a GE that is pending add, but it would happen while the scope lock is still in effect, resulting
 			// in a pending remove being set.
@@ -6424,9 +6423,6 @@ void FActiveGameplayEffectsContainer::DecrementLock()
 			{
 				UE_LOG(LogGameplayEffects, Verbose, TEXT("%s: Finish PendingRemove: %s. Auth: %d"), *GetNameSafe(Owner->GetOwnerActor()), *Effect.GetDebugString(), IsNetAuthority());
 				
-				// Remove this handle from the global map
-				Effect.Handle.RemoveFromGlobalMap();
-
 				GameplayEffects_Internal.RemoveAtSwap(idx, EAllowShrinking::No);
 				ModifiedArray = true;
 				PendingRemoves--;
diff --git a/Engine/Plugins/Runtime/GameplayAbilities/Source/GameplayAbilities/Private/GameplayEffectAggregator.cpp b/Engine/Plugins/Runtime/GameplayAbilities/Source/GameplayAbilities/Private/GameplayEffectAggregator.cpp
--- a/Engine/Plugins/Runtime/GameplayAbilities/Source/GameplayAbilities/Private/GameplayEffectAggregator.cpp
+++ b/Engine/Plugins/Runtime/GameplayAbilities/Source/GameplayAbilities/Private/GameplayEffectAggregator.cpp
@@ -636,7 +636,7 @@ void FAggregator::BroadcastOnDirty()
 	TArray<FActiveGameplayEffectHandle> DependantsLocalCopy = Dependents;
 	Dependents.Empty(); // We will add valid handles back as we process the local list
 
-	for (FActiveGameplayEffectHandle Handle : DependantsLocalCopy)
+	for (FActiveGameplayEffectHandle& Handle : DependantsLocalCopy)
 	{
 		UAbilitySystemComponent* ASC = Handle.GetOwningAbilitySystemComponent();
 		if (ASC)
diff --git a/Engine/Plugins/Runtime/GameplayAbilities/Source/GameplayAbilities/Private/GameplayEffectComponents/AbilitiesGameplayEffectComponent.cpp b/Engine/Plugins/Runtime/GameplayAbilities/Source/GameplayAbilities/Private/GameplayEffectComponents/AbilitiesGameplayEffectComponent.cpp
--- a/Engine/Plugins/Runtime/GameplayAbilities/Source/GameplayAbilities/Private/GameplayEffectComponents/AbilitiesGameplayEffectComponent.cpp
+++ b/Engine/Plugins/Runtime/GameplayAbilities/Source/GameplayAbilities/Private/GameplayEffectComponents/AbilitiesGameplayEffectComponent.cpp
@@ -31,42 +31,45 @@ void UAbilitiesGameplayEffectComponent::AddGrantedAbilityConfig(const FGameplayA
 
 // Reminder: this is happening on the authority only
 void UAbilitiesGameplayEffectComponent::OnInhibitionChanged(FActiveGameplayEffectHandle ActiveGEHandle, bool bIsInhibited) const
-{
-	if (bIsInhibited)
-	{
-		RemoveAbilities(ActiveGEHandle);
-	}
-	else
-	{
-		GrantAbilities(ActiveGEHandle);
-	}
-}
-
-// Reminder: this is happening on the authority only
-void UAbilitiesGameplayEffectComponent::GrantAbilities(FActiveGameplayEffectHandle ActiveGEHandle) const
 {
 	UAbilitySystemComponent* ASC = ActiveGEHandle.GetOwningAbilitySystemComponent();
 	if (!ensure(ASC))
 	{
-		UE_LOG(LogGameplayEffects, Error, TEXT("%s was passed an ActiveGEHandle %s which did not have a valid associated AbilitySystemComponent"), ANSI_TO_TCHAR(__func__), *ActiveGEHandle.ToString());
+		UE_LOG(LogGameplayEffects, Error, TEXT("%hs was passed an ActiveGEHandle %s which did not have a valid associated AbilitySystemComponent (is it PendingRemove?)"), __FUNCTION__, *ActiveGEHandle.ToString());
 		return;
 	}
 
-	if (ASC->bSuppressGrantAbility)
+	const FActiveGameplayEffect* ActiveGE = ASC->GetActiveGameplayEffect(ActiveGEHandle);
+	if (!ensure(ActiveGE))
 	{
-		UE_LOG(LogGameplayEffects, Warning, TEXT("%s suppressed by %s bSuppressGrantAbility"), *GetName(), *ASC->GetName());
+		UE_LOG(LogGameplayEffects, Error, TEXT("ActiveGEHandle %s did not corresponds to Active Gameplay Effect on %s. This is unexpected since GetOwningAbilitySystemComponent succeeded."), *ActiveGEHandle.ToString(), *ASC->GetName());
 		return;
 	}
 
-	const FActiveGameplayEffect* ActiveGE = ASC->GetActiveGameplayEffect(ActiveGEHandle);
-	if (!ActiveGE)
+	if (bIsInhibited)
 	{
-		UE_LOG(LogGameplayEffects, Warning, TEXT("ActiveGEHandle %s did not corresponds to Active Gameplay Effect on %s. This could potentially happen if you remove the GE during the application of other GE's"), *ActiveGEHandle.ToString(), *ASC->GetName());
+		RemoveAbilities(*ASC, *ActiveGE);
+	}
+	else
+	{
+		GrantAbilities(*ASC, *ActiveGE);
+	}
+}
+
+// Reminder: this is happening on the authority only and any call to an ASC member function can invalidate ActiveGE.
+void UAbilitiesGameplayEffectComponent::GrantAbilities(UAbilitySystemComponent& ASC, const FActiveGameplayEffect& ActiveGE) const
+{
+	if (ASC.bSuppressGrantAbility)
+	{
+		UE_LOG(LogGameplayEffects, Log, TEXT("%s suppressed by %s bSuppressGrantAbility"), *GetName(), *ASC.GetName());
 		return;
 	}
-	const FGameplayEffectSpec& ActiveGESpec = ActiveGE->Spec;
 
-	const TArray<FGameplayAbilitySpec>& AllAbilities = ASC->GetActivatableAbilities();
+	// Make sure these are copies because ASC.GiveAbility could invalidate ActiveGE
+	const FGameplayEffectSpec ActiveGESpec = ActiveGE.Spec;
+	const FActiveGameplayEffectHandle ActiveGEHandle = ActiveGE.Handle;
+
+	const TArray<FGameplayAbilitySpec>& AllAbilities = ASC.GetActivatableAbilities();
 	for (const FGameplayAbilitySpecConfig& AbilityConfig : GrantAbilityConfigs)
 	{
 		// Check that we're configured
@@ -78,13 +81,13 @@ void UAbilitiesGameplayEffectComponent::GrantAbilities(FActiveGameplayEffectHand
 
 		// Only do this if we haven't assigned the ability yet! This prevents cases where stacking GEs
 		// would regrant the ability every time the stack was applied
-		const bool bAlreadyGrantedAbility = AllAbilities.ContainsByPredicate([ASC, AbilityCDO, &ActiveGEHandle](FGameplayAbilitySpec& Spec) { return Spec.Ability == AbilityCDO && Spec.GameplayEffectHandle == ActiveGEHandle; });
+		const bool bAlreadyGrantedAbility = AllAbilities.ContainsByPredicate([AbilityCDO, &ActiveGEHandle](FGameplayAbilitySpec& Spec) { return Spec.Ability == AbilityCDO && Spec.GameplayEffectHandle == ActiveGEHandle; });
 		if (bAlreadyGrantedAbility)
 		{
 			continue;
 		}
 
-		const FString ContextString = FString::Printf(TEXT("%s for %s from %s"), ANSI_TO_TCHAR(__func__), *AbilityCDO->GetName(), *GetNameSafe(ActiveGESpec.Def));
+		const FString ContextString = FString::Printf(TEXT("%hs for %s from %s"), __FUNCTION__, *AbilityCDO->GetName(), *GetNameSafe(ActiveGESpec.Def));
 		const int32 Level = static_cast<int32>(AbilityConfig.LevelScalableFloat.GetValueAtLevel(ActiveGESpec.GetLevel(), &ContextString));
 
 		// Now grant that ability to the owning actor
@@ -92,22 +95,16 @@ void UAbilitiesGameplayEffectComponent::GrantAbilities(FActiveGameplayEffectHand
 		AbilitySpec.SetByCallerTagMagnitudes = ActiveGESpec.SetByCallerTagMagnitudes;
 		AbilitySpec.GameplayEffectHandle = ActiveGEHandle;
 
-		ASC->GiveAbility(AbilitySpec);
+		ASC.GiveAbility(AbilitySpec);
 	}
 }
 
-// Reminder: this is happening on the authority only and the ActiveGEHandle can be considered 'inactive' by this time.
-void UAbilitiesGameplayEffectComponent::RemoveAbilities(FActiveGameplayEffectHandle ActiveGEHandle) const
+// Reminder: this is happening on the authority only and the ActiveGE can be considered 'inactive' (pending remove) by this time.
+// Once you do anything on the ASC, ActiveGE can become invalid -- so cache any values you need before any ASC calls.
+void UAbilitiesGameplayEffectComponent::RemoveAbilities(UAbilitySystemComponent& ASC, const FActiveGameplayEffect& ActiveGE) const
 {
-	UAbilitySystemComponent* ASC = ActiveGEHandle.GetOwningAbilitySystemComponent();
-	if (!ensure(ASC))
-	{
-		UE_LOG(LogGameplayEffects, Error, TEXT("%s was passed an ActiveGEHandle %s which did not have a valid associated AbilitySystemComponent"), ANSI_TO_TCHAR(__func__), *ActiveGEHandle.ToString());
-		return;
-	}
-
-	FScopedAbilityListLock ScopedAbilityListLock(*ASC);
-	const TArray<const FGameplayAbilitySpec*> GrantedAbilities = ASC->FindAbilitySpecsFromGEHandle(ScopedAbilityListLock, ActiveGEHandle, EConsiderPending::All);
+	FScopedAbilityListLock ScopedAbilityListLock(ASC);
+	const TArray<const FGameplayAbilitySpec*> GrantedAbilities = ASC.FindAbilitySpecsFromGEHandle(ScopedAbilityListLock, ActiveGE.Handle, EConsiderPending::All);
 	for (const FGameplayAbilitySpecConfig& AbilityConfig : GrantAbilityConfigs)
 	{
 		// Check that we're configured
@@ -129,12 +126,12 @@ void UAbilitiesGameplayEffectComponent::RemoveAbilities(FActiveGameplayEffectHan
 		{
 			case EGameplayEffectGrantedAbilityRemovePolicy::CancelAbilityImmediately:
 			{
-				ASC->ClearAbility(AbilitySpecDef.Handle);
+				ASC.ClearAbility(AbilitySpecDef.Handle);
 				break;
 			}
 			case EGameplayEffectGrantedAbilityRemovePolicy::RemoveAbilityOnEnd:
 			{
-				ASC->SetRemoveAbilityOnEnd(AbilitySpecDef.Handle);
+				ASC.SetRemoveAbilityOnEnd(AbilitySpecDef.Handle);
 				break;
 			}
 			default:
@@ -162,11 +159,15 @@ void UAbilitiesGameplayEffectComponent::OnActiveGameplayEffectRemoved(const FGam
 	const FActiveGameplayEffect* ActiveGE = RemovalInfo.ActiveEffect;
 	if (!ensure(ActiveGE))
 	{
-		UE_LOG(LogGameplayEffects, Error, TEXT("FGameplayEffectRemovalInfo::ActiveEffect was not populated in %s"), ANSI_TO_TCHAR(__func__));
+		UE_LOG(LogGameplayEffects, Error, TEXT("FGameplayEffectRemovalInfo::ActiveEffect was not populated in %hs"), __FUNCTION__);
 		return;
 	}
 
-	RemoveAbilities(ActiveGE->Handle);
+	UAbilitySystemComponent* OwningASC = RemovalInfo.OwningASC;
+	if (ensure(OwningASC))
+	{
+		RemoveAbilities(*OwningASC, *ActiveGE);
+	}
 }
 
 #if WITH_EDITOR
diff --git a/Engine/Plugins/Runtime/GameplayAbilities/Source/GameplayAbilities/Private/GameplayEffectComponents/ImmunityGameplayEffectComponent.cpp b/Engine/Plugins/Runtime/GameplayAbilities/Source/GameplayAbilities/Private/GameplayEffectComponents/ImmunityGameplayEffectComponent.cpp
--- a/Engine/Plugins/Runtime/GameplayAbilities/Source/GameplayAbilities/Private/GameplayEffectComponents/ImmunityGameplayEffectComponent.cpp
+++ b/Engine/Plugins/Runtime/GameplayAbilities/Source/GameplayAbilities/Private/GameplayEffectComponents/ImmunityGameplayEffectComponent.cpp
@@ -50,11 +50,6 @@ bool UImmunityGameplayEffectComponent::AllowGameplayEffectApplication(const FAct
 	SCOPE_CYCLE_COUNTER(STAT_HasApplicationImmunityToSpec)
 
 	const UAbilitySystemComponent* ASC = ActiveGEContainer.Owner;
-	if (ASC != ImmunityActiveGEHandle.GetOwningAbilitySystemComponent())
-	{
-		ensureMsgf(false, TEXT("Something went wrong where an ActiveGameplayEffect jumped AbilitySystemComponents"));
-		return false;
-	}
 
 	// ActiveGE can be null (IsPendingRemove) but we could have not yet received OnEffectRemoved. This can happen when applying more effects during its removal.
 	const FActiveGameplayEffect* ActiveGE = ASC->GetActiveGameplayEffect(ImmunityActiveGEHandle);
@@ -63,6 +58,8 @@ bool UImmunityGameplayEffectComponent::AllowGameplayEffectApplication(const FAct
 		return true;
 	}
 
+	ensureMsgf(ASC == ImmunityActiveGEHandle.GetOwningAbilitySystemComponent(), TEXT("%hs: FActiveGameplayEffectHandle lost its Owning ASC"), __FUNCTION__);
+
 	for (const FGameplayEffectQuery& ImmunityQuery : ImmunityQueries)
 	{
 		if (!ImmunityQuery.IsEmpty() && ImmunityQuery.Matches(GESpecToConsider))
diff --git a/Engine/Plugins/Runtime/GameplayAbilities/Source/GameplayAbilities/Public/ActiveGameplayEffectHandle.h b/Engine/Plugins/Runtime/GameplayAbilities/Source/GameplayAbilities/Public/ActiveGameplayEffectHandle.h
--- a/Engine/Plugins/Runtime/GameplayAbilities/Source/GameplayAbilities/Public/ActiveGameplayEffectHandle.h
+++ b/Engine/Plugins/Runtime/GameplayAbilities/Source/GameplayAbilities/Public/ActiveGameplayEffectHandle.h
@@ -16,21 +16,38 @@ class UAbilitySystemComponent;
 USTRUCT(BlueprintType)
 struct FActiveGameplayEffectHandle
 {
-	GENERATED_USTRUCT_BODY()
-
-	FActiveGameplayEffectHandle()
-		: Handle(INDEX_NONE),
-		bPassedFiltersAndWasExecuted(false)
-	{
-	}
-
+	GENERATED_BODY()
+
+	/** Construct an invalid handle typically used to indicate failure of a Gameplay Effect to execute */
+	FActiveGameplayEffectHandle() = default;
+
+	/**
+	 * This constructor has been deprecated since it leaks internal implementation, leaves the Owning ASC undefined,
+	 * and it's not always clear what the intent of the caller is.  It's been replaced by two factory functions:
+	 * 1. GenerateNewHandle which will give you a valid, properly constructed Handle
+	 * 2. GetInstantExecutedHandle which will give you a special handle, deemed to be invalid since it already executed instantly.
+	 */
+	UE_DEPRECATED(5.8, "Use GenerateNewHandle (for properly constructed, valid handle) or GetInstantExecutedHandle (see comments)")
 	FActiveGameplayEffectHandle(int32 InHandle)
 		: Handle(InHandle),
 		bPassedFiltersAndWasExecuted(true)
 	{
 	}
 
-	/** True if this is tracking an active ongoing gameplay effect */
+	/** Generates a new handle, it will be set to successfully applied. It is up to the caller to associate it with an FActiveGameplayEffect. */
+	static UE_API FActiveGameplayEffectHandle GenerateNewHandle(UAbilitySystemComponent* OwningComponent);
+
+	/** Get a handle that represents an Instant (non-duration) Gameplay Effect. We use this as a sentinel value for FActiveGameplayEffectHandle
+	 * returned from functions where the underlying GE has already completed execution (instantly) and thus cannot be queried further.
+	 * 
+	 * Such a handle has the following rules:
+	 * IsValid() unintuitively returns false, as it does not represent an Active Gameplay Effect.
+	 * WasSuccessfullyApplied() returns true, as it is assumed this represents the execution of the Instant GE.
+	 * GetOwningAbilitySystemComponent() returns an undefined value and should not be relied upon.
+	 */
+	UE_API static FActiveGameplayEffectHandle GetInstantExecutedHandle();
+
+	/** True if this is/was tracking an active (non-instant) gameplay effect. Will continue to return true if the GE was removed. */
 	bool IsValid() const
 	{
 		return Handle != INDEX_NONE;
@@ -42,17 +59,15 @@ struct FActiveGameplayEffectHandle
 		return bPassedFiltersAndWasExecuted;
 	}
 
-	/** Creates a new handle, will be set to successfully applied */
-	static UE_API FActiveGameplayEffectHandle GenerateNewHandle(UAbilitySystemComponent* OwningComponent);
-
-	/** Resets the map that supports GetOwningAbilitySystemComponent */
-	static UE_API void ResetGlobalHandleMap();
-
-	/** Returns the ability system component that created this handle */
+	/** Returns the ability system component that currently has the ActiveGameplayEffect this handle represents.
+	 * Note: If the GameplayEffect has been removed, this will return nullptr regardless of where the ActiveGE was applied. */
 	UE_API UAbilitySystemComponent* GetOwningAbilitySystemComponent() const;
 
-	/** Remove this from the GetOwningAbilitySystemComponent map */
-	UE_API void RemoveFromGlobalMap();
+	UE_DEPRECATED(5.8, "We no longer use a global TMap to get the Owning ASC. You can remove this function.")
+	static void ResetGlobalHandleMap() {}
+
+	UE_DEPRECATED(5.8, "We no longer use a global TMap to get the Owning ASC. You can remove this function.")
+	void RemoveFromGlobalMap() {}
 
 	bool operator==(const FActiveGameplayEffectHandle& Other) const
 	{
@@ -74,18 +89,23 @@ struct FActiveGameplayEffectHandle
 		return FString::Printf(TEXT("%d"), Handle);
 	}
 
+	/** Reset the handle to a default (invalid) state */
 	void Invalidate()
 	{
-		Handle = INDEX_NONE;
+		*this = FActiveGameplayEffectHandle{};
 	}
 
 private:
+	UE_API FActiveGameplayEffectHandle(UAbilitySystemComponent* InOwningASC);
+
+	UPROPERTY()
+	TWeakObjectPtr<UAbilitySystemComponent> WeakOwningASC = {};
 
 	UPROPERTY()
-	int32 Handle;
+	int32 Handle = INDEX_NONE;
 
 	UPROPERTY()
-	bool bPassedFiltersAndWasExecuted;
+	bool bPassedFiltersAndWasExecuted = false;
 };
 
 #undef UE_API
diff --git a/Engine/Plugins/Runtime/GameplayAbilities/Source/GameplayAbilities/Public/GameplayEffectComponents/AbilitiesGameplayEffectComponent.h b/Engine/Plugins/Runtime/GameplayAbilities/Source/GameplayAbilities/Public/GameplayEffectComponents/AbilitiesGameplayEffectComponent.h
--- a/Engine/Plugins/Runtime/GameplayAbilities/Source/GameplayAbilities/Public/GameplayEffectComponents/AbilitiesGameplayEffectComponent.h
+++ b/Engine/Plugins/Runtime/GameplayAbilities/Source/GameplayAbilities/Public/GameplayEffectComponents/AbilitiesGameplayEffectComponent.h
@@ -60,11 +60,17 @@ protected:
 	/** This allows us to be notified when the owning GameplayEffect has had its inhibition changed (which can happen on the initial application). */
 	UE_API virtual void OnInhibitionChanged(FActiveGameplayEffectHandle ActiveGEHandle, bool bIsInhibited) const;
 
+	UE_DEPRECATED(5.8, "Use the version that takes an FActiveGameplayEffect (because they may be PendingRemove by now, so the Handle won't resolve)")
+	virtual void GrantAbilities(FActiveGameplayEffectHandle ActiveGEHandle) const {}
+
+	UE_DEPRECATED(5.8, "Use the version that takes an FActiveGameplayEffect (because they may be PendingRemove by now, so the Handle won't resolve)")
+	virtual void RemoveAbilities(FActiveGameplayEffectHandle ActiveGEHandle) const {}
+
 	/** We should grant the configured Gameplay Abilities */
-	UE_API virtual void GrantAbilities(FActiveGameplayEffectHandle ActiveGEHandle) const;
+	UE_API virtual void GrantAbilities(UAbilitySystemComponent& OwningASC, const FActiveGameplayEffect& ActiveGE) const;
 
 	/** We should remove the configured Gameplay Abilities */
-	UE_API virtual void RemoveAbilities(FActiveGameplayEffectHandle ActiveGEHandle) const;
+	UE_API virtual void RemoveAbilities(UAbilitySystemComponent& OwningASC, const FActiveGameplayEffect& ActiveGE) const;
 
 private:
 	/** We must undo all effects when removed */
diff --git a/Engine/Plugins/Runtime/GameplayAbilities/Source/GameplayAbilities/Public/GameplayEffectTypes.h b/Engine/Plugins/Runtime/GameplayAbilities/Source/GameplayAbilities/Public/GameplayEffectTypes.h
--- a/Engine/Plugins/Runtime/GameplayAbilities/Source/GameplayAbilities/Public/GameplayEffectTypes.h
+++ b/Engine/Plugins/Runtime/GameplayAbilities/Source/GameplayAbilities/Public/GameplayEffectTypes.h
@@ -827,6 +827,10 @@ struct FGameplayEffectRemovalInfo
 	UPROPERTY(VisibleAnywhere, BlueprintReadOnly, Category="Removal")
 	FGameplayEffectContextHandle EffectContext;
 	
+	/** Which AbilitySystemComponent this Gameplay Effect was applied to */
+	UPROPERTY(Transient)
+	TObjectPtr<UAbilitySystemComponent> OwningASC = {};
+
 	/** The Effect being Removed */
 	const FActiveGameplayEffect* ActiveEffect = nullptr;
 };
