[GameplayAbility] Don't allow stale CurrentEventData to linger between activations.
- Without this check, if an Ability is activated first with Event Data and a second time without it; the second activation will have Event Data information from the first activation.

#Synced-CL 49326051

diff --git a/Engine/Plugins/Runtime/GameplayAbilities/Source/GameplayAbilities/Private/Abilities/GameplayAbility.cpp b/Engine/Plugins/Runtime/GameplayAbilities/Source/GameplayAbilities/Private/Abilities/GameplayAbility.cpp
--- a/Engine/Plugins/Runtime/GameplayAbilities/Source/GameplayAbilities/Private/Abilities/GameplayAbility.cpp
+++ b/Engine/Plugins/Runtime/GameplayAbilities/Source/GameplayAbilities/Private/Abilities/GameplayAbility.cpp
@@ -977,9 +977,16 @@ void UGameplayAbility::PreActivate(const FGameplayAbilitySpecHandle Handle, cons
 	// We could set off a chain that results in calling functions on this ability that rely on the current info being set.
 	SetCurrentInfo(Handle, ActorInfo, ActivationInfo);
 
-	if (TriggerEventData && IsInstantiated())
+	if (IsInstantiated())
 	{
-		CurrentEventData = *TriggerEventData;
+		if (TriggerEventData)
+		{
+			CurrentEventData = *TriggerEventData;
+		}
+		else
+		{
+			CurrentEventData = FGameplayEventData();
+		}
 	}
 
 	Comp->HandleChangeAbilityCanBeCanceled(GetAssetTags(), this, true);
