SimulatedOnly Replication state for tags, for use with GEs so that the GE can manage the tags for the owner and still send tags to simulated proxies.

#Synced-CL 49597063

diff --git a/Engine/Plugins/Runtime/GameplayAbilities/Source/GameplayAbilities/Private/AbilitySystemComponent.cpp b/Engine/Plugins/Runtime/GameplayAbilities/Source/GameplayAbilities/Private/AbilitySystemComponent.cpp
--- a/Engine/Plugins/Runtime/GameplayAbilities/Source/GameplayAbilities/Private/AbilitySystemComponent.cpp
+++ b/Engine/Plugins/Runtime/GameplayAbilities/Source/GameplayAbilities/Private/AbilitySystemComponent.cpp
@@ -700,16 +700,16 @@ void UAbilitySystemComponent::UpdateTagMapSingle_Internal(const FGameplayTag& Ta
 		ENetMode NetMode = World ? World->GetNetMode() : NM_Standalone;
 		if (NetMode == NM_DedicatedServer || NetMode == NM_ListenServer)
 		{
-			if (TagRepState == EGameplayTagReplicationState::TagOnly)
-			{
-				int32& TagCount = GetMinimalReplicationTags_Mutable().TagMap.FindOrAdd(Tag);
-				GetMinimalReplicationTags_Mutable().SetTagCount(Tag, FMath::Max(0, TagCount + CountDelta));
-			}
-			else if (TagRepState > EGameplayTagReplicationState::TagOnly)
+			if (TagRepState > EGameplayTagReplicationState::TagOnly)
 			{
 				int32& TagCount = GetReplicatedLooseTags_Mutable().TagMap.FindOrAdd(Tag);
 				GetReplicatedLooseTags_Mutable().SetTagCount(Tag, FMath::Max(0, TagCount + CountDelta));
 			}
+			else if (TagRepState > EGameplayTagReplicationState::None)
+			{
+				int32& TagCount = GetMinimalReplicationTags_Mutable().TagMap.FindOrAdd(Tag);
+				GetMinimalReplicationTags_Mutable().SetTagCount(Tag, FMath::Max(0, TagCount + CountDelta));
+			}
 		}
 	}
 	
@@ -771,22 +771,22 @@ void UAbilitySystemComponent::UpdateTagMap_Internal(const FGameplayTagContainer&
 		ENetMode NetMode = World ? World->GetNetMode() : NM_Standalone;
 		if (NetMode == NM_DedicatedServer || NetMode == NM_ListenServer)
 		{
-			if (TagRepState == EGameplayTagReplicationState::TagOnly)
+			if (TagRepState > EGameplayTagReplicationState::TagOnly)
 			{
 				for (TArray<FGameplayTag>::TConstIterator TagIt = Container.CreateConstIterator(); TagIt; ++TagIt)
 				{
 					const FGameplayTag& Tag = *TagIt;
-					int32& TagCount = GetMinimalReplicationTags_Mutable().TagMap.FindOrAdd(Tag);
-					GetMinimalReplicationTags_Mutable().SetTagCount(Tag, FMath::Max(0, TagCount + CountDelta));
+					int32& TagCount = GetReplicatedLooseTags_Mutable().TagMap.FindOrAdd(Tag);
+					GetReplicatedLooseTags_Mutable().SetTagCount(Tag, FMath::Max(0, TagCount + CountDelta));
 				}
 			}
-			else if (TagRepState > EGameplayTagReplicationState::TagOnly)
+			else if (TagRepState > EGameplayTagReplicationState::None)
 			{
 				for (TArray<FGameplayTag>::TConstIterator TagIt = Container.CreateConstIterator(); TagIt; ++TagIt)
 				{
 					const FGameplayTag& Tag = *TagIt;
-					int32& TagCount = GetReplicatedLooseTags_Mutable().TagMap.FindOrAdd(Tag);
-					GetReplicatedLooseTags_Mutable().SetTagCount(Tag, FMath::Max(0, TagCount + CountDelta));
+					int32& TagCount = GetMinimalReplicationTags_Mutable().TagMap.FindOrAdd(Tag);
+					GetMinimalReplicationTags_Mutable().SetTagCount(Tag, FMath::Max(0, TagCount + CountDelta));
 				}
 			}
 		}
diff --git a/Engine/Plugins/Runtime/GameplayAbilities/Source/GameplayAbilities/Private/GameplayEffect.cpp b/Engine/Plugins/Runtime/GameplayAbilities/Source/GameplayAbilities/Private/GameplayEffect.cpp
--- a/Engine/Plugins/Runtime/GameplayAbilities/Source/GameplayAbilities/Private/GameplayEffect.cpp
+++ b/Engine/Plugins/Runtime/GameplayAbilities/Source/GameplayAbilities/Private/GameplayEffect.cpp
@@ -57,6 +57,7 @@ const float UGameplayEffect::INVALID_LEVEL = FGameplayEffectConstants::INVALID_L
 DECLARE_CYCLE_STAT(TEXT("MakeQuery"), STAT_MakeGameplayEffectQuery, STATGROUP_AbilitySystem);
 
 extern TAutoConsoleVariable<bool> CVarReplicateTagCountContainerWithIris;
+extern TAutoConsoleVariable<bool> CVarGESimOnlyTags;
 
 namespace UE::GameplayEffect
 {
@@ -4606,18 +4607,29 @@ void FActiveGameplayEffectsContainer::AddActiveGameplayEffectGrantedTagsAndModif
 			TimerManager.SetTimer(Effect.PeriodHandle, Delegate, Effect.Spec.GetPeriod(), true);
 		}
 	}
-
-	// Generic Replication will replicate Tags w/ Counts correctly.  So we should only apply the tags if we're predicting, otherwise we will get the tags of a Server-Initiated GE through GameplayTagCountContainer replication.
-	// Iris Replication is more complicated. CVarReplicateTagCountContainerWithIris indicates that we do not have working Tag Count replication through GameplayTagCountContainerNetSerializer. Instead, we replicate using the legacy path of separated containers (Minimal/Loose).
-	// The complex bShouldApplyTags case is we should still apply the tags of a Server-Initiated GE because we won't get those Tags through Iris GameplayTagCountContainerNetSerializer replication.  We should assume CVarFixPredictedTagApplication is true.
-	const bool bIrisWithoutTagCountSerializer = UE::Net::ShouldUseIrisReplication() && !CVarReplicateTagCountContainerWithIris.GetValueOnAnyThread();
-	const bool bShouldApplyTags = UE::GameplayEffect::CVarFixPredictedTagApplication.GetValueOnAnyThread() ? (bInvokePredictedEffects || bIrisWithoutTagCountSerializer) : (!Effect.bPostPredictObject || bIrisWithoutTagCountSerializer);
-	// Local application should be responsible for applying tags
-	if (bShouldApplyTags)
+	
+	if (CVarGESimOnlyTags.GetValueOnAnyThread())
 	{
+		// With Minimal Replication replicate tags to simulated proxies and the GE will manage the server and owning client.
+		// Without Minimal Replication GE's are replicated fully to simulated proxies and will manage the tags on all clients.
 		// Update our owner with the tags this GameplayEffect grants them
-		Owner->UpdateTagMap(Effect.Spec.Def->GetGrantedTags(), 1, bIrisWithoutTagCountSerializer ? EGameplayTagReplicationState::TagOnly : EGameplayTagReplicationState::CountToOwner);
-		Owner->UpdateTagMap(Effect.Spec.DynamicGrantedTags, 1, bIrisWithoutTagCountSerializer ? EGameplayTagReplicationState::TagOnly : EGameplayTagReplicationState::CountToOwner);
+		Owner->UpdateTagMap(Effect.Spec.Def->GetGrantedTags(), 1, ShouldUseMinimalReplication() ? EGameplayTagReplicationState::SimulatedTagOnly : EGameplayTagReplicationState::None);
+		Owner->UpdateTagMap(Effect.Spec.DynamicGrantedTags, 1, ShouldUseMinimalReplication() ? EGameplayTagReplicationState::SimulatedTagOnly : EGameplayTagReplicationState::None);
+	}
+	else
+	{
+		// Generic Replication will replicate Tags w/ Counts correctly.  So we should only apply the tags if we're predicting, otherwise we will get the tags of a Server-Initiated GE through GameplayTagCountContainer replication.
+		// Iris Replication is more complicated. CVarReplicateTagCountContainerWithIris indicates that we do not have working Tag Count replication through GameplayTagCountContainerNetSerializer. Instead, we replicate using the legacy path of separated containers (Minimal/Loose).
+		// The complex bShouldApplyTags case is we should still apply the tags of a Server-Initiated GE because we won't get those Tags through Iris GameplayTagCountContainerNetSerializer replication.  We should assume CVarFixPredictedTagApplication is true.
+		const bool bIrisWithoutTagCountSerializer = UE::Net::ShouldUseIrisReplication() && !CVarReplicateTagCountContainerWithIris.GetValueOnAnyThread();
+		const bool bShouldApplyTags = UE::GameplayEffect::CVarFixPredictedTagApplication.GetValueOnAnyThread() ? (bInvokePredictedEffects || bIrisWithoutTagCountSerializer) : (!Effect.bPostPredictObject || bIrisWithoutTagCountSerializer);
+		// Local application should be responsible for applying tags
+		if (bShouldApplyTags)
+		{
+			// Update our owner with the tags this GameplayEffect grants them
+			Owner->UpdateTagMap(Effect.Spec.Def->GetGrantedTags(), 1, bIrisWithoutTagCountSerializer ? EGameplayTagReplicationState::TagOnly : EGameplayTagReplicationState::CountToOwner);
+			Owner->UpdateTagMap(Effect.Spec.DynamicGrantedTags, 1, bIrisWithoutTagCountSerializer ? EGameplayTagReplicationState::TagOnly : EGameplayTagReplicationState::CountToOwner);
+		}
 	}
 	
 	// Update our owner with the blocked ability tags this GameplayEffect adds to them
@@ -4904,17 +4916,28 @@ void FActiveGameplayEffectsContainer::RemoveActiveGameplayEffectGrantedTagsAndMo
 		}
 	}
 
-	// Generic Replication will replicate Tags w/ Counts correctly.  So we should only apply the tags if we're predicting, otherwise we will get the tags of a Server-Initiated GE through GameplayTagCountContainer replication.
-	// Iris Replication is more complicated. CVarReplicateTagCountContainerWithIris indicates that we do not have working Tag Count replication through GameplayTagCountContainerNetSerializer. Instead, we replicate using the legacy path of separated containers (Minimal/Loose).
-	// The complex bShouldApplyTags case is we should still apply the tags of a Server-Initiated GE because we won't get those Tags through Iris GameplayTagCountContainerNetSerializer replication.  We should assume CVarFixPredictedTagApplication is true.
-	const bool bIrisWithoutTagCountSerializer = UE::Net::ShouldUseIrisReplication() && !CVarReplicateTagCountContainerWithIris.GetValueOnAnyThread();
-	const bool bShouldApplyTags = UE::GameplayEffect::CVarFixPredictedTagApplication.GetValueOnAnyThread() ? (bInvokePredictedEffects || bIrisWithoutTagCountSerializer) : (IsNetAuthority() || bPredictionRejected || bIrisWithoutTagCountSerializer);
-	// Local application should be responsible for applying tags
-	if (bShouldApplyTags)
+	if (CVarGESimOnlyTags.GetValueOnAnyThread())
 	{
-		// Update gameplaytag count and broadcast delegate if we are at 0
-		Owner->UpdateTagMap(Effect.Spec.Def->GetGrantedTags(), -1, bIrisWithoutTagCountSerializer ? EGameplayTagReplicationState::TagOnly : EGameplayTagReplicationState::CountToOwner);
-		Owner->UpdateTagMap(Effect.Spec.DynamicGrantedTags, -1, bIrisWithoutTagCountSerializer ? EGameplayTagReplicationState::TagOnly : EGameplayTagReplicationState::CountToOwner);
+		// With Minimal Replication replicate tags to simulated proxies and the GE will manage the server and owning client.
+		// Without Minimal Replication GE's are replicated fully to simulated proxies and will manage the tags on all clients.
+		// Update our owner with the tags this GameplayEffect grants them
+		Owner->UpdateTagMap(Effect.Spec.Def->GetGrantedTags(), -1, ShouldUseMinimalReplication() ? EGameplayTagReplicationState::SimulatedTagOnly : EGameplayTagReplicationState::None);
+		Owner->UpdateTagMap(Effect.Spec.DynamicGrantedTags, -1, ShouldUseMinimalReplication() ? EGameplayTagReplicationState::SimulatedTagOnly : EGameplayTagReplicationState::None);
+	}
+	else
+	{
+		// Generic Replication will replicate Tags w/ Counts correctly.  So we should only apply the tags if we're predicting, otherwise we will get the tags of a Server-Initiated GE through GameplayTagCountContainer replication.
+		// Iris Replication is more complicated. CVarReplicateTagCountContainerWithIris indicates that we do not have working Tag Count replication through GameplayTagCountContainerNetSerializer. Instead, we replicate using the legacy path of separated containers (Minimal/Loose).
+		// The complex bShouldApplyTags case is we should still apply the tags of a Server-Initiated GE because we won't get those Tags through Iris GameplayTagCountContainerNetSerializer replication.  We should assume CVarFixPredictedTagApplication is true.
+		const bool bIrisWithoutTagCountSerializer = UE::Net::ShouldUseIrisReplication() && !CVarReplicateTagCountContainerWithIris.GetValueOnAnyThread();
+		const bool bShouldApplyTags = UE::GameplayEffect::CVarFixPredictedTagApplication.GetValueOnAnyThread() ? (bInvokePredictedEffects || bIrisWithoutTagCountSerializer) : (IsNetAuthority() || bPredictionRejected || bIrisWithoutTagCountSerializer);
+		// Local application should be responsible for applying tags
+		if (bShouldApplyTags)
+		{
+			// Update gameplaytag count and broadcast delegate if we are at 0
+			Owner->UpdateTagMap(Effect.Spec.Def->GetGrantedTags(), -1, bIrisWithoutTagCountSerializer ? EGameplayTagReplicationState::TagOnly : EGameplayTagReplicationState::CountToOwner);
+			Owner->UpdateTagMap(Effect.Spec.DynamicGrantedTags, -1, bIrisWithoutTagCountSerializer ? EGameplayTagReplicationState::TagOnly : EGameplayTagReplicationState::CountToOwner);
+		}
 	}
 
 	// Update our owner with the blocked ability tags this GameplayEffect adds to them
diff --git a/Engine/Plugins/Runtime/GameplayAbilities/Source/GameplayAbilities/Private/GameplayEffectTypes.cpp b/Engine/Plugins/Runtime/GameplayAbilities/Source/GameplayAbilities/Private/GameplayEffectTypes.cpp
--- a/Engine/Plugins/Runtime/GameplayAbilities/Source/GameplayAbilities/Private/GameplayEffectTypes.cpp
+++ b/Engine/Plugins/Runtime/GameplayAbilities/Source/GameplayAbilities/Private/GameplayEffectTypes.cpp
@@ -47,9 +47,14 @@ static FAutoConsoleVariableRef CVarWarnIfTryingToReplicateNotSupportedActorRefer
 TAutoConsoleVariable<bool> CVarReplicateTagCountContainerWithIris(
 		TEXT("AbilitySystem.Fix.ReplicateTagCountContainerWithIris"),
 		false,
-		TEXT("Default: False.  When true, will use Iris serialization for Gameplay Tag Count Container enabling tag rep state promotion with iris."),
+		TEXT("Default: False.  When true will use Iris serialization for Gameplay Tag Count Container enabling tag rep state promotion with iris."),
 		ECVF_Default);
 
+TAutoConsoleVariable<bool> CVarGESimOnlyTags(
+		TEXT("AbilitySystem.Fix.GESimOnlyTags"),
+		true,
+		TEXT("Default: True.  When active GE's will use the SimOnly tag replication state so that owning client lets the GEs manage their own tags"),
+		ECVF_Default);
 
 FGameplayModEvaluationChannelSettings::FGameplayModEvaluationChannelSettings()
 {
@@ -575,6 +580,12 @@ bool FGameplayTagCountContainer::NetDeltaSerialize(FNetDeltaSerializeInfo & Delt
 				// We don't replicate this tag, so don't save its state
 				continue;
 			}
+			
+			// if the state is SimulatedOnly we want to skip sending this tag to the owner
+			if (CVarGESimOnlyTags.GetValueOnAnyThread() && CountItem.ReplicationState == EGameplayTagReplicationState::SimulatedTagOnly && bIsOwner)
+			{
+				continue;
+			}
 
 			// Record the NewState for the CountItem
 			FNetGameplayTagCountContainerState::FGameplayTagCountStateItem& NewItem = NewState->TagStates.AddDefaulted_GetRef();
@@ -599,7 +610,16 @@ bool FGameplayTagCountContainer::NetDeltaSerialize(FNetDeltaSerializeInfo & Delt
 
 			Removals.RemoveSingleSwap(*OldItem);
 		}
-
+		
+		for (int idx = 0; idx < Removals.Num(); idx++)
+		{
+			// if the state is SimulatedOnly we want to skip sending this tag removal to the owner
+			if (CVarGESimOnlyTags.GetValueOnAnyThread() && Removals[idx].ReplicationState == EGameplayTagReplicationState::SimulatedTagOnly && bIsOwner)
+			{
+				Removals.RemoveAtSwap(idx, false);
+				--idx;
+			}
+		}
 		uint32 NumRemovals = Removals.Num();
 
 		// If there's no changes between the states we can just leave
@@ -1743,4 +1763,3 @@ void FMinimalReplicationTagCountMap::UpdateOwnerTagMap()
 }
 
 #undef LOCTEXT_NAMESPACE
-
diff --git a/Engine/Plugins/Runtime/GameplayAbilities/Source/GameplayAbilities/Public/GameplayEffectTypes.h b/Engine/Plugins/Runtime/GameplayAbilities/Source/GameplayAbilities/Public/GameplayEffectTypes.h
--- a/Engine/Plugins/Runtime/GameplayAbilities/Source/GameplayAbilities/Public/GameplayEffectTypes.h
+++ b/Engine/Plugins/Runtime/GameplayAbilities/Source/GameplayAbilities/Public/GameplayEffectTypes.h
@@ -1050,6 +1050,7 @@ UENUM()
 enum class EGameplayTagReplicationState : uint8
 {
 	None,				// Tag will not replicate
+	SimulatedTagOnly,		// tag will be replicated to simulated proxies, but not the owner. Used for GameplayEffect tags
 	TagOnly,			// Tag will replicate to all without count
 	CountToOwner,		// Tag will replicate to all and Count will replicate to owner, Functions same as TagAndCountToAll with Iris
 	TagAndCountToAll	// Tag and Count will replicate to all
