[GAS] MultiServerProxy fixes when using UE_WITH_REMOTE_OBJECT_HANDLE, the client must #define UE_NET_IRIS_NETREFHANDLE_SERIAL_SIZE=63 otherwise the FRemoteObjectId cannot fit into an FNetRefHandle.

#rnx

#Synced-CL 49835014

diff --git a/Engine/Plugins/Runtime/GameplayAbilities/Source/GameplayAbilities/Private/Serialization/GameplayEffectContextNetSerializer.cpp b/Engine/Plugins/Runtime/GameplayAbilities/Source/GameplayAbilities/Private/Serialization/GameplayEffectContextNetSerializer.cpp
--- a/Engine/Plugins/Runtime/GameplayAbilities/Source/GameplayAbilities/Private/Serialization/GameplayEffectContextNetSerializer.cpp
+++ b/Engine/Plugins/Runtime/GameplayAbilities/Source/GameplayAbilities/Private/Serialization/GameplayEffectContextNetSerializer.cpp
@@ -46,7 +46,7 @@ struct FGameplayEffectContextNetSerializer
 
 	struct FQuantizedType
 	{
-#if UE_WITH_REMOTE_OBJECT_HANDLE
+#if UE_WITH_REMOTE_OBJECT_HANDLE || (UE_NET_IRIS_NETREFHANDLE_SERIAL_SIZE == 63)
 		alignas(16) uint8 EffectContext[192];
 #else
 		alignas(16) uint8 EffectContext[160];
diff --git a/Engine/Plugins/Runtime/GameplayAbilities/Source/GameplayAbilities/Public/Serialization/GameplayEffectContextNetSerializer.h b/Engine/Plugins/Runtime/GameplayAbilities/Source/GameplayAbilities/Public/Serialization/GameplayEffectContextNetSerializer.h
--- a/Engine/Plugins/Runtime/GameplayAbilities/Source/GameplayAbilities/Public/Serialization/GameplayEffectContextNetSerializer.h
+++ b/Engine/Plugins/Runtime/GameplayAbilities/Source/GameplayAbilities/Public/Serialization/GameplayEffectContextNetSerializer.h
@@ -14,7 +14,7 @@ struct FGameplayEffectContextNetSerializerConfig : public FNetSerializerConfig
 namespace UE::Net
 {
 
-#if UE_WITH_REMOTE_OBJECT_HANDLE
+#if UE_WITH_REMOTE_OBJECT_HANDLE || (UE_NET_IRIS_NETREFHANDLE_SERIAL_SIZE == 63)
 constexpr SIZE_T GetGameplayEffectContextNetSerializerSafeQuantizedSize() { return 528; }
 #else
 constexpr SIZE_T GetGameplayEffectContextNetSerializerSafeQuantizedSize() { return 496; }
