Added code to decipher the Tag's EGameplayTagReplicationState on the receiving side based on the data sent.  This is to help reconstruct a better local version of the state for the MultiServerProxy.

Fixed potential crash in FMinimalReplicationTagCountMap::UpdateOwnerTagMap.

#rnx

#Synced-CL 50100874

diff --git a/Engine/Plugins/Runtime/GameplayAbilities/Source/GameplayAbilities/Private/GameplayEffectTypes.cpp b/Engine/Plugins/Runtime/GameplayAbilities/Source/GameplayAbilities/Private/GameplayEffectTypes.cpp
--- a/Engine/Plugins/Runtime/GameplayAbilities/Source/GameplayAbilities/Private/GameplayEffectTypes.cpp
+++ b/Engine/Plugins/Runtime/GameplayAbilities/Source/GameplayAbilities/Private/GameplayEffectTypes.cpp
@@ -528,18 +528,26 @@ public:
 bool FGameplayTagCountContainer::NetDeltaSerialize(FNetDeltaSerializeInfo & DeltaParams)
 {
 	bool bPackCount = false;
-	
-	if (DeltaParams.Writer)
+
+	// Determine if we are replicating as a NetOwner
+	bool bNetOwner = false;
+	if (DeltaParams.Connection)
 	{
-		bool bIsOwner = false;
-		if (AActor* OwnerActor = Cast<UActorComponent>(DeltaParams.Object)->GetOwner())
+		AActor* OwnerActor = nullptr;
+		if (UActorComponent* ParentComponent = Cast<UActorComponent>(DeltaParams.Object))
 		{
-			if (UNetConnection* OwnerNetConnection = OwnerActor->GetNetConnection())
-			{
-				bIsOwner = OwnerNetConnection == DeltaParams.Connection;
-			}
+			OwnerActor = ParentComponent->GetOwner();
 		}
-		
+		else
+		{
+			OwnerActor = Cast<AActor>(DeltaParams.Object);
+		}
+
+		bNetOwner = OwnerActor && (OwnerActor->GetNetConnection() == DeltaParams.Connection);
+	}
+
+	if (DeltaParams.Writer)
+	{		
 		FBitWriter& Writer = *DeltaParams.Writer;
 
 		FNetGameplayTagCountContainerState* OldState = static_cast<FNetGameplayTagCountContainerState*>(DeltaParams.OldState);
@@ -582,7 +590,7 @@ bool FGameplayTagCountContainer::NetDeltaSerialize(FNetDeltaSerializeInfo & Delt
 			}
 			
 			// if the state is SimulatedOnly we want to skip sending this tag to the owner
-			if (CVarGESimOnlyTags.GetValueOnAnyThread() && CountItem.ReplicationState == EGameplayTagReplicationState::SimulatedTagOnly && bIsOwner)
+			if (CVarGESimOnlyTags.GetValueOnAnyThread() && CountItem.ReplicationState == EGameplayTagReplicationState::SimulatedTagOnly && bNetOwner)
 			{
 				continue;
 			}
@@ -592,7 +600,7 @@ bool FGameplayTagCountContainer::NetDeltaSerialize(FNetDeltaSerializeInfo & Delt
 			NewItem.Tag = CountItem.Tag;
 			NewItem.ReplicationState = CountItem.ReplicationState;
 
-			bPackCount = CountItem.ReplicationState >= (bIsOwner ? EGameplayTagReplicationState::CountToOwner : EGameplayTagReplicationState::TagAndCountToAll);
+			bPackCount = CountItem.ReplicationState >= (bNetOwner ? EGameplayTagReplicationState::CountToOwner : EGameplayTagReplicationState::TagAndCountToAll);
 			NewItem.Count = bPackCount ? CountItem.Count : 1;
 			
 			if (!OldItem)
@@ -614,7 +622,7 @@ bool FGameplayTagCountContainer::NetDeltaSerialize(FNetDeltaSerializeInfo & Delt
 		for (int idx = 0; idx < Removals.Num(); idx++)
 		{
 			// if the state is SimulatedOnly we want to skip sending this tag removal to the owner
-			if (CVarGESimOnlyTags.GetValueOnAnyThread() && Removals[idx].ReplicationState == EGameplayTagReplicationState::SimulatedTagOnly && bIsOwner)
+			if (CVarGESimOnlyTags.GetValueOnAnyThread() && Removals[idx].ReplicationState == EGameplayTagReplicationState::SimulatedTagOnly && bNetOwner)
 			{
 				Removals.RemoveAtSwap(idx, EAllowShrinking::No);
 				--idx;
@@ -657,7 +665,7 @@ bool FGameplayTagCountContainer::NetDeltaSerialize(FNetDeltaSerializeInfo & Delt
 		{
 			FNetGameplayTagCountContainerState::FGameplayTagCountStateItem& NewItem = NewState->TagStates[Idx];
 			// only pack count if we're sending to the correct client for the replication state.
-			bPackCount = NewItem.ReplicationState >= (bIsOwner ? EGameplayTagReplicationState::CountToOwner : EGameplayTagReplicationState::TagAndCountToAll);
+			bPackCount = NewItem.ReplicationState >= (bNetOwner ? EGameplayTagReplicationState::CountToOwner : EGameplayTagReplicationState::TagAndCountToAll);
 			NewItem.Tag.NetSerialize(Writer, DeltaParams.Map, bTagSuccess);
 			Writer.SerializeBits(&bPackCount, 1);
 			if (bPackCount)
@@ -700,13 +708,28 @@ bool FGameplayTagCountContainer::NetDeltaSerialize(FNetDeltaSerializeInfo & Delt
 			{
 				NewCount = 1;
 				UpdatedItem.Tag.NetSerialize(Reader, DeltaParams.Map, bTagSuccess);
+
+				// We need to reconstruct the ReplicationState because of the MultiServerProxy which needs to pass this
+				// data along (and thus it needs to know what to Serialize). We start by assuming it's TagOnly.
+				EGameplayTagReplicationState InferredReplicationState = EGameplayTagReplicationState::TagOnly;
+
 				Reader.SerializeBits(&bPackCount, 1);
 				if (bPackCount)
 				{
+					FGameplayTagCountItem* ExistingItem = Items.FindByPredicate(
+						[&UpdatedItem](const FGameplayTagCountItem& ExistingItem)
+						{
+							return ExistingItem.Tag == UpdatedItem.Tag;
+						});
+
+					// If we're receiving a count, then we are at least CountToOwner, and possibly TagAndCountToAll if we ever receive this on a non-owning connection.
+					InferredReplicationState = bNetOwner ? EGameplayTagReplicationState::CountToOwner : EGameplayTagReplicationState::TagAndCountToAll;
+					InferredReplicationState = ExistingItem ? FMath::Max(ExistingItem->ReplicationState, InferredReplicationState) : InferredReplicationState;
+
 					Reader.SerializeIntPacked(NewCount);
 				}
 
-				SetTagCount(UpdatedItem.Tag, NewCount, EGameplayTagReplicationState::None);
+				SetTagCount(UpdatedItem.Tag, NewCount, InferredReplicationState);
 			}
 		}
 	}
@@ -1741,7 +1764,9 @@ void FMinimalReplicationTagCountMap::UpdateOwnerTagMap()
 
 	if (Owner)
 	{
-		for (auto It = TagMap.CreateIterator(); It; ++It)
+		// We need to copy the TagMap since changing a tag count can cause side effects such as adding or removing other tags
+		TMap<FGameplayTag,int32> TagMapCopy = TagMap;
+		for (auto It = TagMapCopy.CreateConstIterator(); It; ++It)
 		{
 			const FGameplayTag Tag = It->Key;
 			const int32 NewValue = It->Value;
diff --git a/Engine/Plugins/Runtime/GameplayAbilities/Source/GameplayAbilities/Public/GameplayEffectTypes.h b/Engine/Plugins/Runtime/GameplayAbilities/Source/GameplayAbilities/Public/GameplayEffectTypes.h
--- a/Engine/Plugins/Runtime/GameplayAbilities/Source/GameplayAbilities/Public/GameplayEffectTypes.h
+++ b/Engine/Plugins/Runtime/GameplayAbilities/Source/GameplayAbilities/Public/GameplayEffectTypes.h
@@ -1079,6 +1079,7 @@ struct FGameplayTagCountItem
 	UPROPERTY()
 	int32 Count;
 
+	UPROPERTY()
 	EGameplayTagReplicationState ReplicationState;
 
 	bool operator == (const FGameplayTagCountItem& Other) const
