Move GC warning about replication to a location where we actually know for sure it will be an issue

#Synced-CL 50152574

diff --git a/Engine/Plugins/Runtime/GameplayAbilities/Source/GameplayAbilities/Private/GameplayCueFunctionLibrary.cpp b/Engine/Plugins/Runtime/GameplayAbilities/Source/GameplayAbilities/Private/GameplayCueFunctionLibrary.cpp
--- a/Engine/Plugins/Runtime/GameplayAbilities/Source/GameplayAbilities/Private/GameplayCueFunctionLibrary.cpp
+++ b/Engine/Plugins/Runtime/GameplayAbilities/Source/GameplayAbilities/Private/GameplayCueFunctionLibrary.cpp
@@ -58,11 +58,6 @@ void UGameplayCueFunctionLibrary::AddGameplayCueOnActor(AActor* Target, const FG
 
 	if (ASC)
 	{
-		if (ASC->ReplicationProxyEnabled && (Parameters != FGameplayCueParameters()))
-		{
-			UE_LOG(LogGameplayCueFunctionLibrary, Warning, TEXT("%hs: Parameters may not replicate when using AbilitySystemComponent with ReplicationProxyEnabled!"), __FUNCTION__);
-		}
-
 		// The actor has an ability system so the event will fire on authority only and will be replicated.
 		if (Target->GetLocalRole() == ROLE_Authority)
 		{
diff --git a/Engine/Plugins/Runtime/GameplayAbilities/Source/GameplayAbilities/Private/GameplayCueInterface.cpp b/Engine/Plugins/Runtime/GameplayAbilities/Source/GameplayAbilities/Private/GameplayCueInterface.cpp
--- a/Engine/Plugins/Runtime/GameplayAbilities/Source/GameplayAbilities/Private/GameplayCueInterface.cpp
+++ b/Engine/Plugins/Runtime/GameplayAbilities/Source/GameplayAbilities/Private/GameplayCueInterface.cpp
@@ -12,6 +12,12 @@
 
 DEFINE_LOG_CATEGORY_STATIC(LogGameplayCues, Log, All);
 
+#if !(UE_BUILD_SHIPPING || UE_BUILD_TEST)
+	#define CANTREPLICATEPARAMS_SEVERITY Warning
+#else
+	#define CANTREPLICATEPARAMS_SEVERITY Verbose
+#endif
+
 namespace GameplayCueInterfacePrivate
 {
 	struct FCueNameAndUFunction
@@ -529,6 +535,48 @@ void FMinimalGameplayCueReplicationProxy::PreReplication(const FActiveGameplayCu
 	if (LastSourceArrayReplicationKey != SourceContainer.ArrayReplicationKey)
 	{
 		LastSourceArrayReplicationKey = SourceContainer.ArrayReplicationKey;
+
+		if (UE_LOG_ACTIVE(LogGameplayCues, CANTREPLICATEPARAMS_SEVERITY))
+		{
+			// Check if any new cue has parameters that we won't be able to replicate and warn the user
+			FGameplayCueParameters ParamsThatCanReplicate;
+			for (const FActiveGameplayCue& Cue : SourceContainer.GameplayCues)
+			{
+				// Look for new tags
+				if (!ReplicatedTags.Contains(Cue.GameplayCueTag))
+				{
+					ParamsThatCanReplicate.bReplicateLocationWhenUsingMinimalRepProxy = Cue.Parameters.bReplicateLocationWhenUsingMinimalRepProxy;
+					ParamsThatCanReplicate.Location = ParamsThatCanReplicate.bReplicateLocationWhenUsingMinimalRepProxy ? Cue.Parameters.Location : FVector::ZeroVector;
+
+					// Those don't replicate
+					ParamsThatCanReplicate.MatchedTagName = Cue.Parameters.MatchedTagName;
+					ParamsThatCanReplicate.OriginalTag = Cue.Parameters.OriginalTag;
+
+					// Effect context is default constructed at some places so do not warn about it
+					ParamsThatCanReplicate.EffectContext = Cue.Parameters.EffectContext;
+
+					if (Cue.Parameters != ParamsThatCanReplicate)
+					{
+						TPropertyValueIterator<const FProperty> CueParams(FGameplayCueParameters::StaticStruct(), &Cue.Parameters);
+						TPropertyValueIterator<const FProperty> ReplicatedParams(FGameplayCueParameters::StaticStruct(), &ParamsThatCanReplicate);
+
+						for (; CueParams && ReplicatedParams; ++CueParams, ++ReplicatedParams)
+						{
+							if (!ensure(CueParams->Key != nullptr))
+							{
+								continue;
+							}
+
+							if (!CueParams->Key->Identical(CueParams->Value, ReplicatedParams->Value))
+							{
+								UE_LOG(LogGameplayCues, CANTREPLICATEPARAMS_SEVERITY, TEXT("%hs: Parameter %s on gameplay cue %s is set but cannot be replicated"), __FUNCTION__, *CueParams->Key->GetName(), *Cue.GameplayCueTag.ToString());
+							}
+						}
+					}
+				}
+			}
+		}
+
 		ReplicatedTags.SetNum(SourceContainer.GameplayCues.Num(), EAllowShrinking::No);
 		ReplicatedLocations.SetNum(SourceContainer.GameplayCues.Num(), EAllowShrinking::No);
 		for (int32 idx=0; idx < SourceContainer.GameplayCues.Num(); ++idx)
@@ -697,3 +745,5 @@ void FMinimalGameplayCueReplicationProxy::RemoveAllCues()
 		Owner->InvokeGameplayCueEvent(GameplayCueTag, EGameplayCueEvent::Removed, Parameters);
 	}
 }
+
+#undef CANTREPLICATEPARAMS_SEVERITY
\ No newline at end of file
