RemoteObjects:  Serialize non-UProperty GameplayAbility fields when migrating objects.

#rnx

#Synced-CL 50269217

diff --git a/Engine/Plugins/Runtime/GameplayAbilities/Source/GameplayAbilities/Private/Abilities/GameplayAbility.cpp b/Engine/Plugins/Runtime/GameplayAbilities/Source/GameplayAbilities/Private/Abilities/GameplayAbility.cpp
--- a/Engine/Plugins/Runtime/GameplayAbilities/Source/GameplayAbilities/Private/Abilities/GameplayAbility.cpp
+++ b/Engine/Plugins/Runtime/GameplayAbilities/Source/GameplayAbilities/Private/Abilities/GameplayAbility.cpp
@@ -2348,6 +2348,21 @@ void UGameplayAbility::RegisterReplicationFragments(UE::Net::FFragmentRegistrati
 	UE::Net::FReplicationFragmentUtil::CreateAndRegisterFragmentsForObject(this, Context, RegistrationFlags);
 }
 
+void UGameplayAbility::Serialize(FArchive& Ar)
+{
+	Super::Serialize(Ar);
+
+	if (Ar.IsMigratingRemoteObjects())
+	{
+		Ar << bHasBlueprintShouldAbilityRespondToEvent;
+		Ar << bHasBlueprintCanUse;
+		Ar << bHasBlueprintActivate;
+		Ar << bHasBlueprintActivateFromEvent;
+		Ar << CurrentSpecHandle;
+		Ar << TrackedGameplayCues;
+	}
+}
+
 #if UE_WITH_REMOTE_OBJECT_HANDLE
 void UGameplayAbility::PostMigrate(const FUObjectMigrationContext& MigrationContext)
 {
diff --git a/Engine/Plugins/Runtime/GameplayAbilities/Source/GameplayAbilities/Public/Abilities/GameplayAbility.h b/Engine/Plugins/Runtime/GameplayAbilities/Source/GameplayAbilities/Public/Abilities/GameplayAbility.h
--- a/Engine/Plugins/Runtime/GameplayAbilities/Source/GameplayAbilities/Public/Abilities/GameplayAbility.h
+++ b/Engine/Plugins/Runtime/GameplayAbilities/Source/GameplayAbilities/Public/Abilities/GameplayAbility.h
@@ -904,6 +904,9 @@ public:
 	UE_DEPRECATED(5.4, "This is unsafe and unnecessary.  It will always return false.")
 	bool IsMarkPendingKillOnAbilityEnd() const { return false; }
 
+	/** We need to serialize non-property fields in certain circumstances */
+	UE_API virtual void Serialize(FArchive& Ar) override;
+
 protected:
 
 	/** Flag that is set by AbilitySystemComponent::OnRemoveAbility to indicate the ability needs to be cleaned up in AbilitySystemComponent::NotifyAbilityEnded */
