Fix potential crash from  UAbilitySystemComponent::ApplyGameplayEffectSpecToSelf when the input Spec's Def is null

#changelist validated

#Synced-CL 50307545

#Synced-CL 50307545

#Synced-CL 50307545

diff --git a/Engine/Plugins/Runtime/GameplayAbilities/Source/GameplayAbilities/Private/AbilitySystemComponent.cpp b/Engine/Plugins/Runtime/GameplayAbilities/Source/GameplayAbilities/Private/AbilitySystemComponent.cpp
--- a/Engine/Plugins/Runtime/GameplayAbilities/Source/GameplayAbilities/Private/AbilitySystemComponent.cpp
+++ b/Engine/Plugins/Runtime/GameplayAbilities/Source/GameplayAbilities/Private/AbilitySystemComponent.cpp
@@ -984,7 +984,12 @@ FActiveGameplayEffectHandle UAbilitySystemComponent::ApplyGameplayEffectSpecToSe
 
 	FScopeCurrentGameplayEffectBeingApplied ScopedGEApplication(&Spec, this);
 
-	const bool bIsNetAuthority = IsOwnerActorAuthoritative();
+	// Check if the Spec definition is valid
+	if (Spec.Def == nullptr)
+	{
+		ABILITY_LOG(Verbose, TEXT("UAbilitySystemComponent::ApplyGameplayEffectSpecToSelf called with a Null effect definition"));
+		return FActiveGameplayEffectHandle();
+	}
 
 	// Check Network Authority
 	if (!HasNetworkAuthorityToApplyGameplayEffect(PredictionKey))
