[MultiServerProxy] FPredictionKeys should always send full data when communicating with a proxy.

#Synced-CL 50444249

diff --git a/Engine/Plugins/Runtime/GameplayAbilities/Source/GameplayAbilities/Private/GameplayPrediction.cpp b/Engine/Plugins/Runtime/GameplayAbilities/Source/GameplayAbilities/Private/GameplayPrediction.cpp
--- a/Engine/Plugins/Runtime/GameplayAbilities/Source/GameplayAbilities/Private/GameplayPrediction.cpp
+++ b/Engine/Plugins/Runtime/GameplayAbilities/Source/GameplayAbilities/Private/GameplayPrediction.cpp
@@ -4,7 +4,9 @@
 #include "AbilitySystemComponent.h"
 #include "AbilitySystemLog.h"
 #include "Engine/World.h"
+#include "Engine/NetConnection.h"
 #include "Engine/NetDriver.h"
+#include "Engine/PackageMapClient.h"
 
 #include UE_INLINE_GENERATED_CPP_BY_NAME(GameplayPrediction)
 
@@ -125,6 +127,23 @@ bool FPredictionKey::NetSerialize(FArchive& Ar, class UPackageMap* Map, bool& bO
 		 *  or if this is a server initiated key (valid on all connections)
 		 */		
 		ValidKeyForConnection = (Current > 0) && (bIsServerInitiated || (PredictiveConnectionObjectKey == FObjectKey()) || (PredictiveConnectionObjectKey == FObjectKey(Map)));
+
+#if UE_WITH_REMOTE_OBJECT_HANDLE
+		// If we are a server communicating with a proxy (or vice versa) we are always valid
+		if (GIsServer && !ValidKeyForConnection)
+		{
+			if (UPackageMapClient* PackageMapClient = Cast<UPackageMapClient>(Map))
+			{
+				UNetConnection* NetConnection = PackageMapClient->GetConnection();
+				const bool bIsServerSendingToProxy = NetConnection && NetConnection->IsProxyConnection();
+				const bool bIsProxySendingToServer = NetConnection && (NetConnection->GetDriver()->GetProxyType() == UE::Net::EProxyType::Backend);
+				if (bIsServerSendingToProxy || bIsProxySendingToServer)
+				{
+					ValidKeyForConnection = true;
+				}
+			}
+		}
+#endif
 	}
 	Ar.SerializeBits(&ValidKeyForConnection, 1);
 
diff --git a/Engine/Plugins/Runtime/GameplayAbilities/Source/GameplayAbilities/Private/Serialization/PredictionKeyNetSerializer.cpp b/Engine/Plugins/Runtime/GameplayAbilities/Source/GameplayAbilities/Private/Serialization/PredictionKeyNetSerializer.cpp
--- a/Engine/Plugins/Runtime/GameplayAbilities/Source/GameplayAbilities/Private/Serialization/PredictionKeyNetSerializer.cpp
+++ b/Engine/Plugins/Runtime/GameplayAbilities/Source/GameplayAbilities/Private/Serialization/PredictionKeyNetSerializer.cpp
@@ -65,6 +65,20 @@ UE_NET_IMPLEMENT_SERIALIZER(FPredictionKeyNetSerializer);
 const FPredictionKeyNetSerializer::ConfigType FPredictionKeyNetSerializer::DefaultConfig;
 FPredictionKeyNetSerializer::FNetSerializerRegistryDelegates FPredictionKeyNetSerializer::NetSerializerRegistryDelegates;
 
+const UNetConnection* GetNetConnectionEnsured(const FNetSerializationContext& Context, uint32 ConnectionId)
+{
+	// We won't always have UserData e.g. we could be (de-)serializing a default state for delta serialization
+	if (UObject* UserData = Context.GetLocalConnectionUserData(ConnectionId))
+	{
+		const UNetConnection* NetConnection = Cast<UNetConnection>(UserData);
+		ensureMsgf(NetConnection, TEXT("Expected FNetSerializationContext::GetLocalConnectionUserData to always return a NetConnection. Instead we received %s"), *GetNameSafe(UserData));
+
+		return NetConnection;
+	}
+
+	return nullptr;
+}
+
 void FPredictionKeyNetSerializer::Serialize(FNetSerializationContext& Context, const FNetSerializeArgs& Args)
 {
 	const QuantizedType& Value = *reinterpret_cast<QuantizedType*>(Args.Source);
@@ -72,7 +86,11 @@ void FPredictionKeyNetSerializer::Serialize(FNetSerializationContext& Context, c
 	FNetBitStreamWriter* Writer = Context.GetBitStreamWriter();
 
 	Writer->WriteBool(Value.bIsServerInitiated);
-	const bool bShouldWriteKey = (Value.CurrentKey > 0) && (Value.bIsServerInitiated || (Value.OwningConnectionId == 0) || (Value.OwningConnectionId == Context.GetLocalConnectionId()));
+
+	const UNetConnection* NetConnection = GetNetConnectionEnsured(Context, Context.GetLocalConnectionId());
+	const bool bIsServerSendingToProxy = NetConnection && NetConnection->IsProxyConnection();
+	const bool bIsProxySendingToServer = NetConnection && (NetConnection->GetDriver()->GetProxyType() == UE::Net::EProxyType::Backend);
+	const bool bShouldWriteKey = (Value.CurrentKey > 0) && (Value.bIsServerInitiated || (Value.OwningConnectionId == 0) || (Value.OwningConnectionId == Context.GetLocalConnectionId()) || bIsServerSendingToProxy || bIsProxySendingToServer);
 	if (Writer->WriteBool(bShouldWriteKey))
 	{
 		Writer->WriteBits(Value.CurrentKey, 15U);
@@ -110,12 +128,17 @@ void FPredictionKeyNetSerializer::Quantize(FNetSerializationContext&, const FNet
 	{
 		Target.CurrentKey = uint16(Source.Current);
 
-		// The owning connection is only relevant if the key is valid.
+		// The owning connection is only relevant if the key is valid and we've _received_ this key previously (PredictiveConnectionObjectKey was set)
 		if (UPackageMapClient* PackageMap = Cast<UPackageMapClient>(Source.PredictiveConnectionObjectKey.ResolveObjectPtr()))
 		{
 			if (const UNetConnection* Connection = PackageMap->GetConnection())
 			{
-				Target.OwningConnectionId = Connection->GetConnectionHandle().GetParentConnectionId();
+				uint32 OwningConnectionId = Connection->GetConnectionHandle().GetParentConnectionId();
+
+				ensureMsgf(!Target.OwningConnectionId || Target.OwningConnectionId == OwningConnectionId,
+					TEXT("FPredictionKeyNetSerializer::Quantize is stomping an existing OwningConnectionId with a new one (this should never happen)"));
+
+				Target.OwningConnectionId = OwningConnectionId;
 			}
 		}
 	}
@@ -128,17 +151,14 @@ void FPredictionKeyNetSerializer::Dequantize(FNetSerializationContext& Context,
 
 	Target.bIsServerInitiated = Source.bIsServerInitiated;
 	Target.Current = int16(Source.CurrentKey);
+
+	// This is set in Deserialize if the key is valid.
 	if (Source.OwningConnectionId != 0)
 	{
-		UObject* UserData = Context.GetLocalConnectionUserData(Source.OwningConnectionId);
-		if (UNetConnection* NetConnection = Cast<UNetConnection>(UserData))
+		if (const UNetConnection* NetConnection = GetNetConnectionEnsured(Context, Source.OwningConnectionId))
 		{
 			Target.PredictiveConnectionObjectKey = FObjectKey(NetConnection->PackageMap);
 		}
-		else
-		{
-			checkf(UserData == nullptr, TEXT("Unexpected user data type %s"), ToCStr(UserData->GetClass()->GetFName().ToString()));
-		}
 	}
 }
 
diff --git a/Engine/Source/Runtime/Engine/Classes/Engine/NetDriver.h b/Engine/Source/Runtime/Engine/Classes/Engine/NetDriver.h
--- a/Engine/Source/Runtime/Engine/Classes/Engine/NetDriver.h
+++ b/Engine/Source/Runtime/Engine/Classes/Engine/NetDriver.h
@@ -1539,6 +1539,7 @@ public:
 
 	/** (Iris replication only) Set the proxy type of the replication system. */
 	ENGINE_API void SetProxyType(UE::Net::EProxyType InProxyType);
+	ENGINE_API UE::Net::EProxyType GetProxyType() const;
 
 	/** (Generic Replicaton Only) Return true if an actor's role can be downgraded during replication to a given connection.  */
 	virtual bool CanDowngradeActorRole(UNetConnection* Connection, AActor* Actor) const { return true; };
diff --git a/Engine/Source/Runtime/Engine/Private/NetDriver.cpp b/Engine/Source/Runtime/Engine/Private/NetDriver.cpp
--- a/Engine/Source/Runtime/Engine/Private/NetDriver.cpp
+++ b/Engine/Source/Runtime/Engine/Private/NetDriver.cpp
@@ -2472,6 +2472,11 @@ void UNetDriver::SetProxyType(UE::Net::EProxyType InProxyType)
 	ProxyType = InProxyType;
 }
 
+UE::Net::EProxyType UNetDriver::GetProxyType() const
+{
+	return ProxyType;
+}
+
 ENetMode UNetDriver::GetNetMode() const
 {
 	// Special case for PIE and game worlds - forcing dedicated server behavior if requested
diff --git a/Engine/Source/Runtime/Net/Iris/Private/Iris/Serialization/NetSerializationContext.cpp b/Engine/Source/Runtime/Net/Iris/Private/Iris/Serialization/NetSerializationContext.cpp
--- a/Engine/Source/Runtime/Net/Iris/Private/Iris/Serialization/NetSerializationContext.cpp
+++ b/Engine/Source/Runtime/Net/Iris/Private/Iris/Serialization/NetSerializationContext.cpp
@@ -58,7 +58,7 @@ void FNetSerializationContext::SetBitStreamOverflow()
 	}
 }
 
-UObject* FNetSerializationContext::GetLocalConnectionUserData(uint32 ConnectionId)
+UObject* FNetSerializationContext::GetLocalConnectionUserData(uint32 ConnectionId) const
 {
 	if (InternalContext == nullptr)
 	{
diff --git a/Engine/Source/Runtime/Net/Iris/Public/Iris/Serialization/NetSerializationContext.h b/Engine/Source/Runtime/Net/Iris/Public/Iris/Serialization/NetSerializationContext.h
--- a/Engine/Source/Runtime/Net/Iris/Public/Iris/Serialization/NetSerializationContext.h
+++ b/Engine/Source/Runtime/Net/Iris/Public/Iris/Serialization/NetSerializationContext.h
@@ -88,7 +88,7 @@ public:
 	 * @param ConnectionId Local connection ID.
 	 * @return The user data object associated with the connection.
 	 */
-	IRISCORE_API UObject* GetLocalConnectionUserData(uint32 ConnectionId);
+	IRISCORE_API UObject* GetLocalConnectionUserData(uint32 ConnectionId) const;
 
 	IRISCORE_API const UE::Net::FNetTokenStore* GetNetTokenStore() const;
 	IRISCORE_API UE::Net::FNetTokenStore* GetNetTokenStore();
