Animations on sim clients may time out due to poor network conditions; so if the client triggering an ability continues to do so and the montage is stopped for sim clients, we restart the montage on sim clients by checking if the animation is stopped locally but is not stopped on the server.
#rnx
#Synced-CL 50404472

diff --git a/Engine/Plugins/Runtime/GameplayAbilities/Source/GameplayAbilities/Private/AbilitySystemComponent_Abilities.cpp b/Engine/Plugins/Runtime/GameplayAbilities/Source/GameplayAbilities/Private/AbilitySystemComponent_Abilities.cpp
--- a/Engine/Plugins/Runtime/GameplayAbilities/Source/GameplayAbilities/Private/AbilitySystemComponent_Abilities.cpp
+++ b/Engine/Plugins/Runtime/GameplayAbilities/Source/GameplayAbilities/Private/AbilitySystemComponent_Abilities.cpp
@@ -53,6 +53,7 @@ static TAutoConsoleVariable<bool> CVarGasFixClientSideMontageBlendOutTime(TEXT("
 static TAutoConsoleVariable<bool> CVarUpdateMontageSectionIdToPlay(TEXT("AbilitySystem.UpdateMontageSectionIdToPlay"), true, TEXT("During tick, update the section ID that replicated montages should use"));
 static TAutoConsoleVariable<bool> CVarReplicateMontageNextSectionId(TEXT("AbilitySystem.ReplicateMontageNextSectionId"), true, TEXT("Apply the replicated next section Id to montages when skipping position replication"));
 static TAutoConsoleVariable<bool> CVarEnsureAbilitiesEndGracefully(TEXT("AbilitySystem.EnsureAbilitiesEndGracefully"), true, TEXT("When shutting down (during ClearAllAbilities) we should check if all GameplayAbilities gracefully ended. This should be disabled if you have NonInstanced abilities that are designed for multiple concurrent executions."));
+static TAutoConsoleVariable<bool> CVarAllowSimulatedMontageRestart(TEXT("AbilitySystem.AllowSimulatedMontageRestart"), true, TEXT("Due to bad lag times, montages may time out on Simulated clients. This variable allows simulated montages to be restarted in such scenarios."));
 
 void UAbilitySystemComponent::InitializeComponent()
 {
@@ -3307,8 +3308,14 @@ void UAbilitySystemComponent::OnRep_ReplicatedAnimMontage()
 		{
 			// New Montage to play
 			UAnimSequenceBase* LocalAnimation = LocalAnimMontageInfo.AnimMontage && LocalAnimMontageInfo.AnimMontage->IsDynamicMontage() ? LocalAnimMontageInfo.AnimMontage->GetFirstAnimReference() : LocalAnimMontageInfo.AnimMontage.Get();
+			
+			//Animations on sim clients may time out due to poor network conditions; so if the client triggering an ability continues to do so and the montage is stopped for sim clients, we restart the montage on sim clients.
+			const bool bSimulatedMontageRequiresRestart = LocalAnimMontageInfo.AnimMontage && !ConstRepAnimMontageInfo.IsStopped && AnimInstance->Montage_GetIsStopped(LocalAnimMontageInfo.AnimMontage) 
+				&& CVarAllowSimulatedMontageRestart.GetValueOnGameThread();
+
 			if ((LocalAnimation != ConstRepAnimMontageInfo.Animation) ||
-			    (LocalAnimMontageInfo.PlayInstanceId != ConstRepAnimMontageInfo.PlayInstanceId))
+				(LocalAnimMontageInfo.PlayInstanceId != ConstRepAnimMontageInfo.PlayInstanceId) ||
+				bSimulatedMontageRequiresRestart)
 			{
 				LocalAnimMontageInfo.PlayInstanceId = ConstRepAnimMontageInfo.PlayInstanceId;
 
