When tags are removed from the count container they need to be removed from the minimal lists, not decremented.

#Synced-CL 50614422

diff --git a/Engine/Plugins/Runtime/GameplayAbilities/Source/GameplayAbilities/Private/AbilitySystemComponent.cpp b/Engine/Plugins/Runtime/GameplayAbilities/Source/GameplayAbilities/Private/AbilitySystemComponent.cpp
--- a/Engine/Plugins/Runtime/GameplayAbilities/Source/GameplayAbilities/Private/AbilitySystemComponent.cpp
+++ b/Engine/Plugins/Runtime/GameplayAbilities/Source/GameplayAbilities/Private/AbilitySystemComponent.cpp
@@ -827,13 +827,13 @@ void UAbilitySystemComponent::UpdateTagMap_Internal(const FGameplayTagContainer&
 						FMinimalReplicationTagCountMap& MTagMap = GetMinimalReplicationTags_Mutable();
 						if (MTagMap.TagMap.Contains(Tag))
 						{
-							MTagMap.RemoveTag(Tag);
+							MTagMap.SetTagCount(Tag, 0);
 						}
 
 						FMinimalReplicationTagCountMap& RLTagMap = GetReplicatedLooseTags_Mutable();
 						if (RLTagMap.TagMap.Contains(Tag))
 						{
-							RLTagMap.RemoveTag(Tag);
+							RLTagMap.SetTagCount(Tag, 0);
 						}
 					}
 				}
