Fix UCancelAbilityTagsGameplayEffectComponent being unable to cancel abilities when providing an empty WithTags container, as it would fail to pass the tag check

#rnx
#Synced-CL 50763359

#Synced-CL 50763359

diff --git a/Engine/Plugins/Runtime/GameplayAbilities/Source/GameplayAbilities/Private/GameplayEffectComponents/CancelAbilityTagsGameplayEffectComponent.cpp b/Engine/Plugins/Runtime/GameplayAbilities/Source/GameplayAbilities/Private/GameplayEffectComponents/CancelAbilityTagsGameplayEffectComponent.cpp
--- a/Engine/Plugins/Runtime/GameplayAbilities/Source/GameplayAbilities/Private/GameplayEffectComponents/CancelAbilityTagsGameplayEffectComponent.cpp
+++ b/Engine/Plugins/Runtime/GameplayAbilities/Source/GameplayAbilities/Private/GameplayEffectComponents/CancelAbilityTagsGameplayEffectComponent.cpp
@@ -86,7 +86,10 @@ void UCancelAbilityTagsGameplayEffectComponent::CancelOwnerAbilities(const FActi
 		return;
 	}
 
-	ASC->CancelAbilities(&InheritableCancelAbilitiesWithTagsContainer.CombinedTags, &InheritableCancelAbilitiesWithoutTagsContainer.CombinedTags);
+	// Empty container will fail the HasAny check, so pass nullptr instead
+	const FGameplayTagContainer* WithTags = !InheritableCancelAbilitiesWithTagsContainer.CombinedTags.IsEmpty() ? &InheritableCancelAbilitiesWithTagsContainer.CombinedTags : nullptr;
+	const FGameplayTagContainer* WithoutTags = !InheritableCancelAbilitiesWithoutTagsContainer.CombinedTags.IsEmpty() ? &InheritableCancelAbilitiesWithoutTagsContainer.CombinedTags : nullptr;
+	ASC->CancelAbilities(WithTags, WithoutTags);
 }
 
 #if WITH_EDITOR
