RemoteObjects: Fixes that allow parts of the AbilitySystemComponent to work with object migrations & multi-server proxy.

#Synced-CL 50628995

#Synced-CL 50628995

diff --git a/Engine/Plugins/Runtime/GameplayAbilities/Source/GameplayAbilities/Private/AbilitySystemComponent.cpp b/Engine/Plugins/Runtime/GameplayAbilities/Source/GameplayAbilities/Private/AbilitySystemComponent.cpp
--- a/Engine/Plugins/Runtime/GameplayAbilities/Source/GameplayAbilities/Private/AbilitySystemComponent.cpp
+++ b/Engine/Plugins/Runtime/GameplayAbilities/Source/GameplayAbilities/Private/AbilitySystemComponent.cpp
@@ -3,6 +3,7 @@
 #include "AbilitySystemComponent.h"
 #include "AbilitySystemLog.h"
 #include "UObject/UObjectIterator.h"
+#include "UObject/UObjectMigrationContext.h"
 #include "Engine/Canvas.h"
 #include "DisplayDebugHelpers.h"
 #include "Engine/Engine.h"
@@ -237,6 +238,17 @@ void UAbilitySystemComponent::OnUnregister()
 {
 	Super::OnUnregister();
 
+#if UE_WITH_REMOTE_OBJECT_HANDLE
+	if (const FUObjectMigrationContext* MigrationContext = FUObjectMigrationContext::GetCurrentMigrationContext())
+	{
+		// Do not DestroyActiveState when migrating-out.  There's too many side effects which we should instead handle in PostMigrate.
+		if (MigrationContext->MigrationSide == EObjectMigrationSide::Send)
+		{
+			return;
+		}
+	}
+#endif
+
 	DestroyActiveState();
 }
 
@@ -254,6 +266,22 @@ void UAbilitySystemComponent::Serialize(FArchive& Ar)
 
 	if (Ar.IsMigratingRemoteObjects())
 	{
+		bool bHasAbilityActorInfo = AbilityActorInfo.IsValid();
+		Ar << bHasAbilityActorInfo;
+		if (bHasAbilityActorInfo)
+		{
+			// Can happen when loading (the sending side had an ability info we didn't yet allocate)
+			if (!AbilityActorInfo.IsValid())
+			{
+				AbilityActorInfo = TSharedPtr<FGameplayAbilityActorInfo>( UAbilitySystemGlobals::Get().AllocAbilityActorInfo() );
+			}
+
+			// TODO: This just uses the default serializer but does not account for polymorphic structs
+			ensureMsgf(sizeof(*AbilityActorInfo.Get()) == sizeof(FGameplayAbilityActorInfo), TEXT("We do not handle polymorphic AbilityActorInfo properly"));
+			FGameplayAbilityActorInfo DefaultStruct;
+			FGameplayAbilityActorInfo::StaticStruct()->SerializeItem(Ar, AbilityActorInfo.Get(), &DefaultStruct);
+		}
+
 		if (Ar.IsSaving())
 		{
 			auto SaveTagToGASpecMap = [&Ar](TMap<FGameplayTag, TArray<FGameplayAbilitySpecHandle>>& Map)
