GASDebugger takes DPI scale into account when creating the canvas width.

#13898

#Synced-CL 51023371

diff --git a/Engine/Plugins/Runtime/GameplayAbilities/Source/GameplayAbilities/Private/GameplayDebuggerCategory_Abilities.cpp b/Engine/Plugins/Runtime/GameplayAbilities/Source/GameplayAbilities/Private/GameplayDebuggerCategory_Abilities.cpp
--- a/Engine/Plugins/Runtime/GameplayAbilities/Source/GameplayAbilities/Private/GameplayDebuggerCategory_Abilities.cpp
+++ b/Engine/Plugins/Runtime/GameplayAbilities/Source/GameplayAbilities/Private/GameplayDebuggerCategory_Abilities.cpp
@@ -621,7 +621,7 @@ void FGameplayDebuggerCategory_Abilities::DrawGameplayAbilities(FGameplayDebugge
 {
 	using namespace UE::AbilitySystem::Debug;
 
-	const float CanvasWidth = CanvasContext.Canvas->SizeX;
+	const float CanvasWidth = CanvasContext.Canvas->SizeX / CanvasContext.Canvas->GetDPIScale();
 	Algo::Sort(DataPack.Abilities, [](const FRepData::FGameplayAbilityDebug& ItemOne, const FRepData::FGameplayAbilityDebug& ItemTwo) { return ItemOne.Ability < ItemTwo.Ability; });
 
 	int32 NumActive = 0;
