Fix OriginalTag on GameplayCues not updating properly
PR #14420:

#Synced-CL 51222906

diff --git a/Engine/Plugins/Runtime/GameplayAbilities/Source/GameplayAbilities/Private/GameplayCueManager.cpp b/Engine/Plugins/Runtime/GameplayAbilities/Source/GameplayAbilities/Private/GameplayCueManager.cpp
--- a/Engine/Plugins/Runtime/GameplayAbilities/Source/GameplayAbilities/Private/GameplayCueManager.cpp
+++ b/Engine/Plugins/Runtime/GameplayAbilities/Source/GameplayAbilities/Private/GameplayCueManager.cpp
@@ -133,6 +133,7 @@ void UGameplayCueManager::HandleGameplayCues(AActor* TargetActor, const FGamepla
 
 	for (auto It = GameplayCueTags.CreateConstIterator(); It; ++It)
 	{
+		Parameters.OriginalTag = *It;
 		HandleGameplayCue(TargetActor, *It, EventType, Parameters, Options);
 	}
 }
