Allow BP to edit GameplayEffect parameters before applying
* An overridable function allows BP to access the spec handle before it gets applied
* The main goal is for Force Power Throw to be able to override cooldown and damage values

#Synced-CL 51756957

diff --git a/Engine/Plugins/Runtime/GameplayAbilities/Source/GameplayAbilities/Private/Abilities/GameplayAbility.cpp b/Engine/Plugins/Runtime/GameplayAbilities/Source/GameplayAbilities/Private/Abilities/GameplayAbility.cpp
--- a/Engine/Plugins/Runtime/GameplayAbilities/Source/GameplayAbilities/Private/Abilities/GameplayAbility.cpp
+++ b/Engine/Plugins/Runtime/GameplayAbilities/Source/GameplayAbilities/Private/Abilities/GameplayAbility.cpp
@@ -1381,7 +1381,7 @@ FGameplayEffectSpecHandle UGameplayAbility::MakeOutgoingGameplayEffectSpec(const
 		{
 			NewHandle.Data->SetByCallerTagMagnitudes = AbilitySpec->SetByCallerTagMagnitudes;
 		}
-
+		BP_EditSpecValues(NewHandle);
 	}
 	return NewHandle;
 }
diff --git a/Engine/Plugins/Runtime/GameplayAbilities/Source/GameplayAbilities/Public/Abilities/GameplayAbility.h b/Engine/Plugins/Runtime/GameplayAbilities/Source/GameplayAbilities/Public/Abilities/GameplayAbility.h
--- a/Engine/Plugins/Runtime/GameplayAbilities/Source/GameplayAbilities/Public/Abilities/GameplayAbility.h
+++ b/Engine/Plugins/Runtime/GameplayAbilities/Source/GameplayAbilities/Public/Abilities/GameplayAbility.h
@@ -624,6 +624,10 @@ protected:
 
 	UE_API FActiveGameplayEffectHandle ApplyGameplayEffectSpecToOwner(const FGameplayAbilitySpecHandle AbilityHandle, const FGameplayAbilityActorInfo* ActorInfo, const FGameplayAbilityActivationInfo ActivationInfo, const FGameplayEffectSpecHandle SpecHandle) const;
 
+	/** Change Gameplay effect values before they get applied */
+	UFUNCTION(BlueprintImplementableEvent, Category = Ability, DisplayName = "EditSpecValues", meta = (ScriptName = "EditSpecValues"))
+	UE_API void BP_EditSpecValues(FGameplayEffectSpecHandle SpecHandle) const;
+	
 	// -------------------------------------
 	//  Apply Gameplay effects to Target
 	// -------------------------------------
