[GAS] When marking abilities as garbage, also mark any SubObjects of those Abilities as Garbage.  LyraGameplayAbilities can have SubObjects that must be MarkAsGarbage with their Outer.

#Synced-CL 51772012

diff --git a/Engine/Plugins/Runtime/GameplayAbilities/Source/GameplayAbilities/Private/AbilitySystemComponent_Abilities.cpp b/Engine/Plugins/Runtime/GameplayAbilities/Source/GameplayAbilities/Private/AbilitySystemComponent_Abilities.cpp
--- a/Engine/Plugins/Runtime/GameplayAbilities/Source/GameplayAbilities/Private/AbilitySystemComponent_Abilities.cpp
+++ b/Engine/Plugins/Runtime/GameplayAbilities/Source/GameplayAbilities/Private/AbilitySystemComponent_Abilities.cpp
@@ -55,6 +55,28 @@ static TAutoConsoleVariable<bool> CVarReplicateMontageNextSectionId(TEXT("Abilit
 static TAutoConsoleVariable<bool> CVarEnsureAbilitiesEndGracefully(TEXT("AbilitySystem.EnsureAbilitiesEndGracefully"), true, TEXT("When shutting down (during ClearAllAbilities) we should check if all GameplayAbilities gracefully ended. This should be disabled if you have NonInstanced abilities that are designed for multiple concurrent executions."));
 static TAutoConsoleVariable<bool> CVarAllowSimulatedMontageRestart(TEXT("AbilitySystem.AllowSimulatedMontageRestart"), true, TEXT("Due to bad lag times, montages may time out on Simulated clients. This variable allows simulated montages to be restarted in such scenarios."));
 
+/**
+ * Mark the passed-in GA (and any SubObjects) as Garbage
+ */
+void MarkAbilityAsGarbage(UGameplayAbility* GameplayAbility)
+{
+	if (GameplayAbility)
+	{
+		GameplayAbility->MarkAsGarbage();
+
+		// We must also mark all SubObjects as garbage, gather any ones not already marked.
+		TArray<UObject*> GASubObjects;
+		constexpr bool bIncludeNestedObjects = true;
+		constexpr EObjectFlags ExclusionFlags = EObjectFlags::RF_MirroredGarbage;
+		GetObjectsWithOuter(GameplayAbility, GASubObjects, bIncludeNestedObjects, ExclusionFlags);
+
+		for (UObject* GASubObject : GASubObjects)
+		{
+			GASubObject->MarkAsGarbage();
+		}
+	}
+}
+
 void UAbilitySystemComponent::InitializeComponent()
 {
 	Super::InitializeComponent();
@@ -641,7 +663,7 @@ void UAbilitySystemComponent::OnRemoveAbility(FGameplayAbilitySpec& Spec)
 				if (Instance->GetInstancingPolicy() == EGameplayAbilityInstancingPolicy::InstancedPerExecution)
 				{
 					ABILITY_LOG(Error, TEXT("%s was InActive, yet still instanced during OnRemove"), *Instance->GetName());
-					Instance->MarkAsGarbage();
+					MarkAbilityAsGarbage(Instance);
 				}
 			}
 		}
@@ -658,7 +680,7 @@ void UAbilitySystemComponent::OnRemoveAbility(FGameplayAbilitySpec& Spec)
 		{
 			RemoveReplicatedInstancedAbility(PrimaryInstance);
 		}
-		PrimaryInstance->MarkAsGarbage();
+		MarkAbilityAsGarbage(PrimaryInstance);
 	}
 	else
 	{
@@ -1190,6 +1212,10 @@ UGameplayAbility* UAbilitySystemComponent::CreateNewInstanceOfAbility(FGameplayA
 	UGameplayAbility * AbilityInstance = NewObject<UGameplayAbility>(Owner, Ability->GetClass());
 	check(AbilityInstance);
 
+#if UE_WITH_REMOTE_OBJECT_HANDLE
+	UE_LOGF(LogAbilitySystem, Verbose, "%s: %ls (%ls)", __FUNCTION__, *GetNameSafe(AbilityInstance), *FRemoteObjectId(AbilityInstance).ToString());
+#endif
+
 	// Add it to one of our instance lists so that it doesn't GC.
 	if (AbilityInstance->GetReplicationPolicy() != EGameplayAbilityReplicationPolicy::ReplicateNo)
 	{
@@ -1261,7 +1287,7 @@ void UAbilitySystemComponent::NotifyAbilityEnded(FGameplayAbilitySpecHandle Hand
 			Spec->NonReplicatedInstances.Remove(Ability);
 		}
 
-		Ability->MarkAsGarbage();
+		MarkAbilityAsGarbage(Ability);
 	}
 
 	if (OwnerRole == ROLE_Authority)
@@ -1400,10 +1426,7 @@ void UAbilitySystemComponent::DestroyActiveState()
 				TArray<UGameplayAbility*> AbilitiesToCancel = Spec.GetAbilityInstances();
 				for (UGameplayAbility* InstanceAbility : AbilitiesToCancel)
 				{
-					if (InstanceAbility)
-					{
-						InstanceAbility->MarkAsGarbage();
-					}
+					MarkAbilityAsGarbage(InstanceAbility);
 				}
 
 				Spec.ReplicatedInstances.Empty();
