Fix double delete call on gameplay container when head and head->next end up being the same pointer. Added Cvar to toggle.

#tests Inconsistent repo, server only, ensured compiled and preflight
#rnx

diff --git a/Engine/Plugins/Runtime/GameplayAbilities/Source/GameplayAbilities/Private/GameplayEffect.cpp b/Engine/Plugins/Runtime/GameplayAbilities/Source/GameplayAbilities/Private/GameplayEffect.cpp
--- a/Engine/Plugins/Runtime/GameplayAbilities/Source/GameplayAbilities/Private/GameplayEffect.cpp
+++ b/Engine/Plugins/Runtime/GameplayAbilities/Source/GameplayAbilities/Private/GameplayEffect.cpp
@@ -114,6 +114,12 @@ namespace UE::GameplayEffect
 		true,
 		TEXT("Default: True.  When enabled Gameplay Effect tag application usesthe same logic as Gameplay Cues so that application and removal works correctly for predicted effects."),
 		ECVF_Default);
+
+	TAutoConsoleVariable<bool> CVarFixPendingGameplayEffectHeadNextAtTeardown(
+		TEXT("AbilitySystem.Fix.PendingGameplayEffectHeadNextAtTeardown"),
+		true,
+		TEXT("Default: True.  When enabled we will prevent the Pending Gameplay Effect Head from pointing to itself through Next pointer during teardown to prevent a double call to delete."),
+		ECVF_Default);
 #if WITH_EDITOR
 	namespace EditorOnly
 	{
@@ -2958,7 +2964,16 @@ FActiveGameplayEffectsContainer::~FActiveGameplayEffectsContainer()
 {
 	using namespace UE::GameplayEffect;
 
-	ensureAlwaysMsgf(!PendingGameplayEffectHead || PendingGameplayEffectHead != PendingGameplayEffectHead->PendingNext, TEXT("Pending Gameplay Effect Head is pointing to the same location as Pending Next this shouldn't happen'"));
+	if (!ensureAlwaysMsgf(!PendingGameplayEffectHead || PendingGameplayEffectHead != PendingGameplayEffectHead->PendingNext, TEXT("Pending Gameplay Effect Head is pointing to the same location as Pending Next this shouldn't happen'")))
+	{
+		ensure(HasActiveGameplayEffectFix(EActiveGameplayEffectFix::OmitPendingNextOnCopy));
+
+		if (CVarFixPendingGameplayEffectHeadNextAtTeardown.GetValueOnAnyThread())
+		{
+			// setting next to null so that we don't call delete on the same memory twice in the code below
+			PendingGameplayEffectHead->PendingNext = nullptr;
+		}
+	}
 
 	if (HasActiveGameplayEffectFix(EActiveGameplayEffectFix::CleanupAllPendingActiveGEs))
 	{
