GAS: Add RemovalPolicy to OnApplicationGameplayEffects, incorporated in AdditionalEffectsGameplayEffectComponent

- Use case: Control life time of external applied GE through GE that applies it conditionally. I myself have noticed this before, and I've seen bugs in the past from designers not understanding that a GE applied through a GE controls its own lifetime by default, they implicitly assume it will be removed when the effect which applies ends.
- Added support for removal policy, following a very similiar pattern to Granting GAs through a gameplay effect component (You can cancel / end the GA when the GE granting the ability is removed)
- Added EGameplayEffectGrantedEffectRemovalPolicy enum (GrantedEffectControlsOwnLifetime, RemoveGrantedEffectOnEnd)
- Added stack count to remove field to conditional gameplay effect
- OnActiveGameplayEffectRemoved removes any on-application effects whose policy is RemoveGrantedEffectOnEnd
- Added content validation: error if RemoveGrantedEffectOnEnd is used with an Instant owning GE or an Instant granted GE

[RN] Updated AdditionalEffectsGameplayEffectComponent to support controlling lifetime of granted GEs
#Synced-CL 52071577

diff --git a/Engine/Plugins/Runtime/GameplayAbilities/Source/GameplayAbilities/Private/GameplayEffect.cpp b/Engine/Plugins/Runtime/GameplayAbilities/Source/GameplayAbilities/Private/GameplayEffect.cpp
--- a/Engine/Plugins/Runtime/GameplayAbilities/Source/GameplayAbilities/Private/GameplayEffect.cpp
+++ b/Engine/Plugins/Runtime/GameplayAbilities/Source/GameplayAbilities/Private/GameplayEffect.cpp
@@ -1505,7 +1505,7 @@ FGameplayEffectSpecHandle FConditionalGameplayEffect::CreateSpec(FGameplayEffect
 
 bool FConditionalGameplayEffect::operator==(const FConditionalGameplayEffect& Other) const
 {
-	return EffectClass == Other.EffectClass && RequiredSourceTags == Other.RequiredSourceTags;
+	return EffectClass == Other.EffectClass && RequiredSourceTags == Other.RequiredSourceTags && RemovalPolicy == Other.RemovalPolicy && StackCountToRemove == Other.StackCountToRemove;
 }
 
 bool FConditionalGameplayEffect::operator!=(const FConditionalGameplayEffect& Other) const
diff --git a/Engine/Plugins/Runtime/GameplayAbilities/Source/GameplayAbilities/Private/GameplayEffectComponents/AdditionalEffectsGameplayEffectComponent.cpp b/Engine/Plugins/Runtime/GameplayAbilities/Source/GameplayAbilities/Private/GameplayEffectComponents/AdditionalEffectsGameplayEffectComponent.cpp
--- a/Engine/Plugins/Runtime/GameplayAbilities/Source/GameplayAbilities/Private/GameplayEffectComponents/AdditionalEffectsGameplayEffectComponent.cpp
+++ b/Engine/Plugins/Runtime/GameplayAbilities/Source/GameplayAbilities/Private/GameplayEffectComponents/AdditionalEffectsGameplayEffectComponent.cpp
@@ -113,6 +113,18 @@ void UAdditionalEffectsGameplayEffectComponent::OnActiveGameplayEffectRemoved(co
 			ASC->ApplyGameplayEffectSpecToSelf(NewSpec);
 		}
 	}
+
+	// Remove any on-application effects that were configured with RemoveGrantedEffectOnEnd
+	for (const FConditionalGameplayEffect& ConditionalEffect : OnApplicationGameplayEffects)
+	{
+		if (ConditionalEffect.RemovalPolicy == EGameplayEffectGrantedEffectRemovalPolicy::RemoveGrantedEffectOnEnd
+			&& ConditionalEffect.EffectClass != nullptr)
+		{
+			FGameplayEffectQuery RemoveQuery;
+			RemoveQuery.EffectDefinition = ConditionalEffect.EffectClass;
+			ActiveGEContainer->RemoveActiveEffects(RemoveQuery, ConditionalEffect.StackCountToRemove);
+		}
+	}
 }
 
 #if WITH_EDITOR
@@ -120,7 +132,9 @@ EDataValidationResult UAdditionalEffectsGameplayEffectComponent::IsDataValid(FDa
 {
 	EDataValidationResult Result = Super::IsDataValid(Context);
 
-	if (GetOwner()->DurationPolicy == EGameplayEffectDurationType::Instant)
+	const bool bOwnerIsInstant = (GetOwner()->DurationPolicy == EGameplayEffectDurationType::Instant);
+
+	if (bOwnerIsInstant)
 	{
 		const bool bHasOnCompleteEffects = (OnCompleteAlways.Num() + OnCompleteNormal.Num() + OnCompletePrematurely.Num() > 0);
 		if (bHasOnCompleteEffects)
@@ -137,6 +151,35 @@ EDataValidationResult UAdditionalEffectsGameplayEffectComponent::IsDataValid(FDa
 		}
 	}
 
+	for (const FConditionalGameplayEffect& ConditionalEffect : OnApplicationGameplayEffects)
+	{
+		if (ConditionalEffect.RemovalPolicy != EGameplayEffectGrantedEffectRemovalPolicy::RemoveGrantedEffectOnEnd)
+		{
+			continue;
+		}
+
+		if (bOwnerIsInstant)
+		{
+			Context.AddError(LOCTEXT("RemoveGrantedEffectOnEndWithInstantOwner",
+				"RemoveGrantedEffectOnEnd cannot be used on an Instant owning GE. Instant effects are never added to the ActiveGE container, so the removal callback will never fire."));
+			Result = EDataValidationResult::Invalid;
+		}
+
+		// The granted GE must be non-Instant: instant effects are applied and immediately gone, there is nothing to remove later
+		if (ConditionalEffect.EffectClass != nullptr)
+		{
+			if (const UGameplayEffect* GrantedCDO = ConditionalEffect.EffectClass.GetDefaultObject())
+			{
+				if (GrantedCDO->DurationPolicy == EGameplayEffectDurationType::Instant)
+				{
+					Context.AddError(LOCTEXT("RemoveGrantedEffectOnEndWithInstantChild",
+							"RemoveGrantedEffectOnEnd cannot be used with Instant granted effect. Instant effects are applied and immediately removed; there is nothing to clean up when the owning GE ends."));
+					Result = EDataValidationResult::Invalid;
+				}
+			}
+		}
+	}
+
 	return Result;
 }
 #endif // WITH_EDITOR
diff --git a/Engine/Plugins/Runtime/GameplayAbilities/Source/GameplayAbilities/Public/GameplayEffect.h b/Engine/Plugins/Runtime/GameplayAbilities/Source/GameplayAbilities/Public/GameplayEffect.h
--- a/Engine/Plugins/Runtime/GameplayAbilities/Source/GameplayAbilities/Public/GameplayEffect.h
+++ b/Engine/Plugins/Runtime/GameplayAbilities/Source/GameplayAbilities/Public/GameplayEffect.h
@@ -474,6 +474,19 @@ struct FGameplayEffectExecutionScopedModifierInfo
 	FGameplayTagRequirements TargetTags;
 };
 
+/** Describes what happens to a Gameplay Effect that was applied on-application by another Gameplay Effect
+ *  (via UAdditionalEffectsGameplayEffectComponent) when the applying effect is removed. */
+UENUM()
+enum class EGameplayEffectGrantedEffectRemovalPolicy : uint8
+{
+	/** The applied Gameplay Effect manages its own lifetime independently. This is the default behavior. */
+	GrantedEffectControlsOwnLifetime,
+
+	/** The applied Gameplay Effect will be removed when the Gameplay Effect that applied it is removed.
+	 *  Cannot be used when the granted effect has an Instant duration (instant GEs are never persisted). */
+	RemoveGrantedEffectOnEnd,
+};
+
 /**
  * Struct for gameplay effects that apply only if another gameplay effect (or execution) was successfully applied.
  */
@@ -494,6 +507,16 @@ struct FConditionalGameplayEffect
 	UPROPERTY(EditDefaultsOnly, BlueprintReadOnly, Category = GameplayEffect)
 	FGameplayTagContainer RequiredSourceTags;
 
+	/** What happens to this applied effect when the Gameplay Effect that applied it is removed.
+	 *  Only relevant when this entry is used in UAdditionalEffectsGameplayEffectComponent::OnApplicationGameplayEffects. */
+	UPROPERTY(EditDefaultsOnly, BlueprintReadOnly, Category = GameplayEffect)
+	EGameplayEffectGrantedEffectRemovalPolicy RemovalPolicy = EGameplayEffectGrantedEffectRemovalPolicy::GrantedEffectControlsOwnLifetime;
+
+	/** Number of stacks to remove when RemovalPolicy is RemoveGrantedEffectOnEnd.
+	 *  INDEX_NONE (default) removes all stacks. */
+	UPROPERTY(EditDefaultsOnly, BlueprintReadOnly, Category = GameplayEffect, meta = (EditCondition = "RemovalPolicy == EGameplayEffectGrantedEffectRemovalPolicy::RemoveGrantedEffectOnEnd", EditConditionHides))
+	int32 StackCountToRemove = INDEX_NONE;
+
 	UE_API bool operator==(const FConditionalGameplayEffect& Other) const;
 	UE_API bool operator!=(const FConditionalGameplayEffect& Other) const;
 };
