Fix issue of only checking first deletion for if next pointer is pointing to self. Consolidate duplicate code. Add logging to find the ActiveGE that is pointing to itself.
Wrapped everything in a cvar and will remove old code in future release.

#tests preflight
#rnx

diff --git a/Engine/Plugins/Runtime/GameplayAbilities/Source/GameplayAbilities/Private/GameplayEffect.cpp b/Engine/Plugins/Runtime/GameplayAbilities/Source/GameplayAbilities/Private/GameplayEffect.cpp
--- a/Engine/Plugins/Runtime/GameplayAbilities/Source/GameplayAbilities/Private/GameplayEffect.cpp
+++ b/Engine/Plugins/Runtime/GameplayAbilities/Source/GameplayAbilities/Private/GameplayEffect.cpp
@@ -120,6 +120,13 @@ namespace UE::GameplayEffect
 		true,
 		TEXT("Default: True.  When enabled we will prevent the Pending Gameplay Effect Head from pointing to itself through Next pointer during teardown to prevent a double call to delete."),
 		ECVF_Default);
+
+	TAutoConsoleVariable<bool> CVarFixAllPendingGameplayEffectsHeadNextsAtTeardown(
+		TEXT("AbilitySystem.Fix.AllPendingGameplayEffectsHeadNextsAtTeardown"),
+		true,
+		TEXT("Default: True.  When enabled we will prevent all the Pending Gameplay Effects Head from pointing to itself through Next pointers during teardown to prevent any double call to delete."),
+		ECVF_Default);
+
 #if WITH_EDITOR
 	namespace EditorOnly
 	{
@@ -2964,38 +2971,81 @@ FActiveGameplayEffectsContainer::~FActiveGameplayEffectsContainer()
 {
 	using namespace UE::GameplayEffect;
 
-	if (!ensureAlwaysMsgf(!PendingGameplayEffectHead || PendingGameplayEffectHead != PendingGameplayEffectHead->PendingNext, TEXT("Pending Gameplay Effect Head is pointing to the same location as Pending Next this shouldn't happen'")))
+	if (CVarFixAllPendingGameplayEffectsHeadNextsAtTeardown.GetValueOnAnyThread())
 	{
-		ensure(HasActiveGameplayEffectFix(EActiveGameplayEffectFix::OmitPendingNextOnCopy));
-
-		if (CVarFixPendingGameplayEffectHeadNextAtTeardown.GetValueOnAnyThread())
+		if (HasActiveGameplayEffectFix(EActiveGameplayEffectFix::CleanupAllPendingActiveGEs))
 		{
-			// setting next to null so that we don't call delete on the same memory twice in the code below
-			PendingGameplayEffectHead->PendingNext = nullptr;
+			while (PendingGameplayEffectHead)
+			{
+				DeletePendingActiveGameplayEffect();
+			}
 		}
-	}
-
-	if (HasActiveGameplayEffectFix(EActiveGameplayEffectFix::CleanupAllPendingActiveGEs))
-	{
-		while (PendingGameplayEffectHead)
+		else
 		{
-			FActiveGameplayEffect* Next = PendingGameplayEffectHead->PendingNext;
-			delete PendingGameplayEffectHead;
-			PendingGameplayEffectHead = Next;
+			DeletePendingActiveGameplayEffect();
 		}
 	}
 	else
 	{
-		if (PendingGameplayEffectHead)
+		if (!ensureAlwaysMsgf(!PendingGameplayEffectHead || PendingGameplayEffectHead != PendingGameplayEffectHead->PendingNext, TEXT("Pending Gameplay Effect Head is pointing to the same location as Pending Next this shouldn't happen'")))
+		{
+			ensure(HasActiveGameplayEffectFix(EActiveGameplayEffectFix::OmitPendingNextOnCopy));
+
+			if (CVarFixPendingGameplayEffectHeadNextAtTeardown.GetValueOnAnyThread())
+			{
+				// setting next to null so that we don't call delete on the same memory twice in the code below
+				PendingGameplayEffectHead->PendingNext = nullptr;
+			}
+		}
+
+		if (HasActiveGameplayEffectFix(EActiveGameplayEffectFix::CleanupAllPendingActiveGEs))
 		{
-			FActiveGameplayEffect* Next = PendingGameplayEffectHead->PendingNext;
-			delete PendingGameplayEffectHead;
-			PendingGameplayEffectHead = Next;
+			while (PendingGameplayEffectHead)
+			{
+				FActiveGameplayEffect* Next = PendingGameplayEffectHead->PendingNext;
+				delete PendingGameplayEffectHead;
+				PendingGameplayEffectHead = Next;
+			}
+		}
+		else
+		{
+			if (PendingGameplayEffectHead)
+			{
+				FActiveGameplayEffect* Next = PendingGameplayEffectHead->PendingNext;
+				delete PendingGameplayEffectHead;
+				PendingGameplayEffectHead = Next;
+			}
 		}
 	}
 }
 PRAGMA_ENABLE_DEPRECATION_WARNINGS
 
+void FActiveGameplayEffectsContainer::DeletePendingActiveGameplayEffect()
+{
+	using namespace UE::GameplayEffect;
+
+	if (!PendingGameplayEffectHead)
+	{
+		return;
+	}
+
+	if (!ensureAlwaysMsgf(PendingGameplayEffectHead != PendingGameplayEffectHead->PendingNext, TEXT("Pending Gameplay Effect Head is pointing to the same location as Pending Next this shouldn't happen'")))
+	{
+		ensure(HasActiveGameplayEffectFix(EActiveGameplayEffectFix::OmitPendingNextOnCopy));
+
+		if (CVarFixPendingGameplayEffectHeadNextAtTeardown.GetValueOnAnyThread())
+		{
+			UE_LOG(LogGameplayEffects, Warning, TEXT("Setting PendingNext to null for pointing to self on Pending ActiveGE: %s"), *PendingGameplayEffectHead->GetDebugString());
+			// setting next to null so that we don't call delete on the same memory twice in the code below
+			PendingGameplayEffectHead->PendingNext = nullptr;
+		}
+	}
+
+	FActiveGameplayEffect* Next = PendingGameplayEffectHead->PendingNext;
+	delete PendingGameplayEffectHead;
+	PendingGameplayEffectHead = Next;
+}
+
 void FActiveGameplayEffectsContainer::RegisterWithOwner(UAbilitySystemComponent* InOwner)
 {
 	if (Owner != InOwner && InOwner != nullptr)
diff --git a/Engine/Plugins/Runtime/GameplayAbilities/Source/GameplayAbilities/Public/GameplayEffect.h b/Engine/Plugins/Runtime/GameplayAbilities/Source/GameplayAbilities/Public/GameplayEffect.h
--- a/Engine/Plugins/Runtime/GameplayAbilities/Source/GameplayAbilities/Public/GameplayEffect.h
+++ b/Engine/Plugins/Runtime/GameplayAbilities/Source/GameplayAbilities/Public/GameplayEffect.h
@@ -2002,6 +2002,8 @@ private:
 
 	UE_API bool ShouldUseMinimalReplication();
 
+	void DeletePendingActiveGameplayEffect();
+
 	mutable int32 ScopedLockCount;
 	int32 PendingRemoves;
 	uint32 NumConsecutiveUnmappedReferencesDebug = 0;
