Fixed indeterministic asset registry queries during cook in GameplayCueManager

#Synced-CL 52466916

diff --git a/Engine/Plugins/Runtime/GameplayAbilities/Source/GameplayAbilities/Private/GameplayCueManager.cpp b/Engine/Plugins/Runtime/GameplayAbilities/Source/GameplayAbilities/Private/GameplayCueManager.cpp
--- a/Engine/Plugins/Runtime/GameplayAbilities/Source/GameplayAbilities/Private/GameplayCueManager.cpp
+++ b/Engine/Plugins/Runtime/GameplayAbilities/Source/GameplayAbilities/Private/GameplayCueManager.cpp
@@ -820,12 +820,13 @@ TSharedPtr<FStreamableHandle> UGameplayCueManager::InitObjectLibrary(FGameplayCu
 	DECLARE_SCOPE_CYCLE_COUNTER(TEXT("Loading Library"), STAT_ObjectLibrary, STATGROUP_LoadTime);
 
 	// Instantiate the UObjectLibraries if they aren't there already
+	// When running in editor, query in-memory assets so any unsaved assets are discovered. When cooking, only query on-disk assets to avoid indeterministic results
 	if (!Lib.StaticObjectLibrary)
 	{
 		Lib.StaticObjectLibrary = UObjectLibrary::CreateLibrary(UGameplayCueNotify_Static::StaticClass(), true, GIsEditor && !IsRunningCommandlet());
 		if (GIsEditor)
 		{
-			Lib.StaticObjectLibrary->bIncludeOnlyOnDiskAssets = false;
+			Lib.StaticObjectLibrary->bIncludeOnlyOnDiskAssets = IsRunningCookCommandlet();
 		}
 	}
 	if (!Lib.ActorObjectLibrary)
@@ -833,7 +834,7 @@ TSharedPtr<FStreamableHandle> UGameplayCueManager::InitObjectLibrary(FGameplayCu
 		Lib.ActorObjectLibrary = UObjectLibrary::CreateLibrary(AGameplayCueNotify_Actor::StaticClass(), true, GIsEditor && !IsRunningCommandlet());
 		if (GIsEditor)
 		{
-			Lib.ActorObjectLibrary->bIncludeOnlyOnDiskAssets = false;
+			Lib.ActorObjectLibrary->bIncludeOnlyOnDiskAssets = IsRunningCookCommandlet();
 		}
 	}	
 
