[ContentBrowser] New Add Menu Gameplay Menu

#Synced-CL 52439714

diff --git a/Engine/Plugins/Cameras/GameplayCameras/Source/GameplayCamerasEditor/Private/AssetTools/AssetDefinition_CameraAsset.cpp b/Engine/Plugins/Cameras/GameplayCameras/Source/GameplayCamerasEditor/Private/AssetTools/AssetDefinition_CameraAsset.cpp
--- a/Engine/Plugins/Cameras/GameplayCameras/Source/GameplayCamerasEditor/Private/AssetTools/AssetDefinition_CameraAsset.cpp
+++ b/Engine/Plugins/Cameras/GameplayCameras/Source/GameplayCamerasEditor/Private/AssetTools/AssetDefinition_CameraAsset.cpp
@@ -29,7 +29,7 @@ TSoftClassPtr<UObject> UAssetDefinition_CameraAsset::GetAssetClass() const
 
 TConstArrayView<FAssetCategoryPath> UAssetDefinition_CameraAsset::GetAssetCategories() const
 {
-	static const auto Categories = { FAssetCategoryPath(EAssetCategoryPaths::Gameplay) };
+	static const auto Categories = { FAssetCategoryPath(EAssetCategoryPaths::Gameplay, LOCTEXT("CameraAsset_SubMenu", "Camera"), ECategoryMenuType::Section) };
 	return Categories;
 }
 
diff --git a/Engine/Plugins/Cameras/GameplayCameras/Source/GameplayCamerasEditor/Private/AssetTools/AssetDefinition_CameraRigAsset.cpp b/Engine/Plugins/Cameras/GameplayCameras/Source/GameplayCamerasEditor/Private/AssetTools/AssetDefinition_CameraRigAsset.cpp
--- a/Engine/Plugins/Cameras/GameplayCameras/Source/GameplayCamerasEditor/Private/AssetTools/AssetDefinition_CameraRigAsset.cpp
+++ b/Engine/Plugins/Cameras/GameplayCameras/Source/GameplayCamerasEditor/Private/AssetTools/AssetDefinition_CameraRigAsset.cpp
@@ -13,7 +13,7 @@
 
 FText UAssetDefinition_CameraRigAsset::GetAssetDisplayName() const
 {
-	return LOCTEXT("AssetDisplayName", "Camera Rig");
+	return LOCTEXT("CameraRigAsset_DisplayName", "Camera Rig Prefab");
 }
 
 FLinearColor UAssetDefinition_CameraRigAsset::GetAssetColor() const
@@ -28,7 +28,7 @@ TSoftClassPtr<UObject> UAssetDefinition_CameraRigAsset::GetAssetClass() const
 
 TConstArrayView<FAssetCategoryPath> UAssetDefinition_CameraRigAsset::GetAssetCategories() const
 {
-	static const auto Categories = { FAssetCategoryPath(EAssetCategoryPaths::Gameplay) };
+	static const auto Categories = { FAssetCategoryPath(EAssetCategoryPaths::Gameplay, LOCTEXT("CameraRigAsset_SubMenu", "Camera"), ECategoryMenuType::Section) };
 	return Categories;
 }
 
diff --git a/Engine/Plugins/Cameras/GameplayCameras/Source/GameplayCamerasEditor/Private/AssetTools/AssetDefinition_CameraRigProxyAsset.cpp b/Engine/Plugins/Cameras/GameplayCameras/Source/GameplayCamerasEditor/Private/AssetTools/AssetDefinition_CameraRigProxyAsset.cpp
--- a/Engine/Plugins/Cameras/GameplayCameras/Source/GameplayCamerasEditor/Private/AssetTools/AssetDefinition_CameraRigProxyAsset.cpp
+++ b/Engine/Plugins/Cameras/GameplayCameras/Source/GameplayCamerasEditor/Private/AssetTools/AssetDefinition_CameraRigProxyAsset.cpp
@@ -27,7 +27,7 @@ TSoftClassPtr<UObject> UAssetDefinition_CameraRigProxyAsset::GetAssetClass() con
 
 TConstArrayView<FAssetCategoryPath> UAssetDefinition_CameraRigProxyAsset::GetAssetCategories() const
 {
-	static const auto Categories = { FAssetCategoryPath(EAssetCategoryPaths::Gameplay) };
+	static const auto Categories = { FAssetCategoryPath(EAssetCategoryPaths::Gameplay, LOCTEXT("CameraRigProxyAsset_SubMenu", "Camera"), ECategoryMenuType::Section) };
 	return Categories;
 }
 
diff --git a/Engine/Plugins/Cameras/GameplayCameras/Source/GameplayCamerasEditor/Private/AssetTools/AssetDefinition_CameraShakeAsset.cpp b/Engine/Plugins/Cameras/GameplayCameras/Source/GameplayCamerasEditor/Private/AssetTools/AssetDefinition_CameraShakeAsset.cpp
--- a/Engine/Plugins/Cameras/GameplayCameras/Source/GameplayCamerasEditor/Private/AssetTools/AssetDefinition_CameraShakeAsset.cpp
+++ b/Engine/Plugins/Cameras/GameplayCameras/Source/GameplayCamerasEditor/Private/AssetTools/AssetDefinition_CameraShakeAsset.cpp
@@ -28,7 +28,7 @@ TSoftClassPtr<UObject> UAssetDefinition_CameraShakeAsset::GetAssetClass() const
 
 TConstArrayView<FAssetCategoryPath> UAssetDefinition_CameraShakeAsset::GetAssetCategories() const
 {
-	static const auto Categories = { FAssetCategoryPath(EAssetCategoryPaths::Gameplay) };
+	static const auto Categories = { FAssetCategoryPath(EAssetCategoryPaths::Gameplay, LOCTEXT("CameraAsset_SubMenu", "Camera"), ECategoryMenuType::Section) };
 	return Categories;
 }
 
diff --git a/Engine/Plugins/Cameras/GameplayCameras/Source/GameplayCamerasEditor/Private/AssetTools/AssetDefinition_CameraVariableCollection.cpp b/Engine/Plugins/Cameras/GameplayCameras/Source/GameplayCamerasEditor/Private/AssetTools/AssetDefinition_CameraVariableCollection.cpp
--- a/Engine/Plugins/Cameras/GameplayCameras/Source/GameplayCamerasEditor/Private/AssetTools/AssetDefinition_CameraVariableCollection.cpp
+++ b/Engine/Plugins/Cameras/GameplayCameras/Source/GameplayCamerasEditor/Private/AssetTools/AssetDefinition_CameraVariableCollection.cpp
@@ -13,7 +13,7 @@
 
 TConstArrayView<FAssetCategoryPath> UAssetDefinition_CameraVariableCollection::StaticMenuCategories()
 {
-	static const auto Categories = { FAssetCategoryPath(EAssetCategoryPaths::Gameplay) };
+	static const auto Categories = { FAssetCategoryPath(EAssetCategoryPaths::Gameplay, LOCTEXT("CameraVariableCollection_SubMenu", "Camera"), ECategoryMenuType::Section) };
 	return Categories;
 }
 
diff --git a/Engine/Plugins/Editor/EngineAssetDefinitions/Source/Private/AssetDefinition_LightWeightInstance.h b/Engine/Plugins/Editor/EngineAssetDefinitions/Source/Private/AssetDefinition_LightWeightInstance.h
--- a/Engine/Plugins/Editor/EngineAssetDefinitions/Source/Private/AssetDefinition_LightWeightInstance.h
+++ b/Engine/Plugins/Editor/EngineAssetDefinitions/Source/Private/AssetDefinition_LightWeightInstance.h
@@ -19,7 +19,7 @@ public:
 	virtual TSoftClassPtr<UObject> GetAssetClass() const override { return ADEPRECATED_LightWeightInstanceManager::StaticClass(); }
 	virtual TConstArrayView<FAssetCategoryPath> GetAssetCategories() const override
 	{
-		static const auto Categories = { EAssetCategoryPaths::Gameplay };
+		static const auto Categories = { FAssetCategoryPath(EAssetCategoryPaths::Gameplay, NSLOCTEXT("AssetDefinition", "LightWeightInstance_SubMenu", "Other")) };
 		return Categories;
 	}
 	// UAssetDefinition End
diff --git a/Engine/Plugins/Runtime/GameplayAbilities/GameplayAbilities.uplugin b/Engine/Plugins/Runtime/GameplayAbilities/GameplayAbilities.uplugin
--- a/Engine/Plugins/Runtime/GameplayAbilities/GameplayAbilities.uplugin
+++ b/Engine/Plugins/Runtime/GameplayAbilities/GameplayAbilities.uplugin
@@ -28,6 +28,10 @@
 		}
 	],
 	"Plugins": [
+		{
+			"Name": "EngineAssetDefinitions",
+			"Enabled": true
+		},
 		{
 			"Name": "GameplayTagsEditor",
 			"Enabled": true
diff --git a/Engine/Plugins/Runtime/GameplayAbilities/Source/GameplayAbilitiesEditor/GameplayAbilitiesEditor.Build.cs b/Engine/Plugins/Runtime/GameplayAbilities/Source/GameplayAbilitiesEditor/GameplayAbilitiesEditor.Build.cs
--- a/Engine/Plugins/Runtime/GameplayAbilities/Source/GameplayAbilitiesEditor/GameplayAbilitiesEditor.Build.cs
+++ b/Engine/Plugins/Runtime/GameplayAbilities/Source/GameplayAbilitiesEditor/GameplayAbilitiesEditor.Build.cs
@@ -12,9 +12,11 @@ namespace UnrealBuildTool.Rules
 				new string[]
 				{
 					// ... add private dependencies that you statically link with here ...
+					"AssetDefinition",
 					"Core",
 					"CoreUObject",
 					"Engine",
+					"EngineAssetDefinitions",
 					"AssetTools",
 					"ClassViewer",
 					"GameplayTags",
@@ -44,7 +46,7 @@ namespace UnrealBuildTool.Rules
 					"DataRegistry",
 					"DataRegistryEditor",
 					"ToolMenus",
-					"ApplicationCore"
+					"ApplicationCore",
 				}
 			);
 		}
diff --git a/Engine/Plugins/Runtime/GameplayAbilities/Source/GameplayAbilitiesEditor/Private/AssetDefinition_GameplayAbilitiesBlueprint.cpp b/Engine/Plugins/Runtime/GameplayAbilities/Source/GameplayAbilitiesEditor/Private/AssetDefinition_GameplayAbilitiesBlueprint.cpp
new file mode 100644
--- /dev/null
+++ b/Engine/Plugins/Runtime/GameplayAbilities/Source/GameplayAbilitiesEditor/Private/AssetDefinition_GameplayAbilitiesBlueprint.cpp
@@ -0,0 +1,64 @@
+// Copyright Epic Games, Inc. All Rights Reserved.
+
+#include "AssetDefinition_GameplayAbilitiesBlueprint.h"
+#include "GameplayAbilitiesBlueprintFactory.h"
+#include "GameplayAbilitiesEditor.h"
+#include "Kismet2/BlueprintEditorUtils.h"
+#include "Misc/MessageDialog.h"
+
+#define LOCTEXT_NAMESPACE "AssetDefinition_GameplayAbilitiesBlueprint"
+
+EAssetCommandResult UAssetDefinition_GameplayAbilitiesBlueprint::OpenAssets(const FAssetOpenArgs& OpenArgs) const
+{
+	const EToolkitMode::Type Mode = OpenArgs.ToolkitHost.IsValid() ? EToolkitMode::WorldCentric : EToolkitMode::Standalone;
+
+	const TArray<UObject*>& Objects = OpenArgs.LoadObjects<UObject>();
+	for (auto ObjIt = Objects.CreateConstIterator(); ObjIt; ++ObjIt)
+	{
+		UBlueprint* Blueprint = Cast<UBlueprint>(*ObjIt);
+		if (Blueprint)
+		{
+			bool bLetOpen = true;
+			if (!Blueprint->ParentClass)
+			{
+				bLetOpen = EAppReturnType::Yes == FMessageDialog::Open(
+					EAppMsgType::YesNo, 
+					LOCTEXT("FailedToLoadAbilityBlueprintWithContinue", "Gameplay Ability Blueprint could not be loaded because it derives from an invalid class.  Check to make sure the parent class for this blueprint hasn't been removed! Do you want to continue (it can crash the editor)?")
+				);
+			}
+		
+			if (bLetOpen)
+			{
+				const TSharedRef<FGameplayAbilitiesEditor> NewEditor(new FGameplayAbilitiesEditor());
+
+				TArray<UBlueprint*> Blueprints;
+				Blueprints.Add(Blueprint);
+
+				NewEditor->InitGameplayAbilitiesEditor(Mode, OpenArgs.ToolkitHost, Blueprints, ShouldUseDataOnlyEditor(Blueprint));
+			}
+		}
+		else
+		{
+			FMessageDialog::Open(EAppMsgType::Ok, LOCTEXT("FailedToLoadAbilityBlueprint", "Gameplay Ability Blueprint could not be loaded because it derives from an invalid class.  Check to make sure the parent class for this blueprint hasn't been removed!"));
+		}
+	}
+	return EAssetCommandResult::Handled;
+}
+
+UFactory* UAssetDefinition_GameplayAbilitiesBlueprint::GetFactoryForBlueprintType(UBlueprint* InBlueprint) const
+{
+	UGameplayAbilitiesBlueprintFactory* GameplayAbilitiesBlueprintFactory = NewObject<UGameplayAbilitiesBlueprintFactory>();
+	GameplayAbilitiesBlueprintFactory->ParentClass = TSubclassOf<UGameplayAbility>(*InBlueprint->GeneratedClass);
+	return GameplayAbilitiesBlueprintFactory;
+}
+
+bool UAssetDefinition_GameplayAbilitiesBlueprint::ShouldUseDataOnlyEditor(const UBlueprint* Blueprint) const
+{
+	return FBlueprintEditorUtils::IsDataOnlyBlueprint(Blueprint)
+		&& !FBlueprintEditorUtils::IsLevelScriptBlueprint(Blueprint)
+		&& !FBlueprintEditorUtils::IsInterfaceBlueprint(Blueprint)
+		&& !Blueprint->bForceFullEditor
+		&& !Blueprint->bIsNewlyCreated;
+}
+
+#undef LOCTEXT_NAMESPACE
diff --git a/Engine/Plugins/Runtime/GameplayAbilities/Source/GameplayAbilitiesEditor/Private/AssetDefinition_GameplayAbilitiesBlueprint.h b/Engine/Plugins/Runtime/GameplayAbilities/Source/GameplayAbilitiesEditor/Private/AssetDefinition_GameplayAbilitiesBlueprint.h
new file mode 100644
--- /dev/null
+++ b/Engine/Plugins/Runtime/GameplayAbilities/Source/GameplayAbilitiesEditor/Private/AssetDefinition_GameplayAbilitiesBlueprint.h
@@ -0,0 +1,33 @@
+// Copyright Epic Games, Inc. All Rights Reserved.
+
+#pragma once
+#include "GameplayAbilityBlueprint.h"
+#include "Script/AssetDefinition_Blueprint.h"
+#include "AssetDefinition_GameplayAbilitiesBlueprint.generated.h"
+
+UCLASS()
+class UAssetDefinition_GameplayAbilitiesBlueprint : public UAssetDefinition_Blueprint
+{
+	GENERATED_BODY()
+public:
+	// UAssetDefinition Begin
+	virtual FText GetAssetDisplayName() const override { return NSLOCTEXT("AssetDefinition", "AssetDefinition_GameplayAbilitiesBlueprint", "Gameplay Ability Blueprint"); }
+	virtual FLinearColor GetAssetColor() const override { return FColor(0, 96, 128); }
+	virtual TSoftClassPtr<UObject> GetAssetClass() const override { return UGameplayAbilityBlueprint::StaticClass(); }
+	virtual TConstArrayView<FAssetCategoryPath> GetAssetCategories() const override
+	{
+		static const auto Categories = 
+		{
+			FAssetCategoryPath(EAssetCategoryPaths::Gameplay, NSLOCTEXT("AssetDefinition", "AssetDefinition_GameplayAbilitiesBlueprintSubMenu", "Other"), ECategoryMenuType::Section)
+		};
+
+		return Categories;
+	}
+	virtual EAssetCommandResult OpenAssets(const FAssetOpenArgs& OpenArgs) const override;
+	virtual UFactory* GetFactoryForBlueprintType(UBlueprint* InBlueprint) const override;
+	// UAssetDefinition End
+
+private:
+	/** Returns true if the blueprint is data only */
+	bool ShouldUseDataOnlyEditor(const UBlueprint* Blueprint) const;
+};
diff --git a/Engine/Plugins/Runtime/GameplayAbilities/Source/GameplayAbilitiesEditor/Private/AssetTypeActions_GameplayAbilitiesBlueprint.cpp b/Engine/Plugins/Runtime/GameplayAbilities/Source/GameplayAbilitiesEditor/Private/AssetTypeActions_GameplayAbilitiesBlueprint.cpp
deleted file mode 100644
diff --git a/Engine/Plugins/Runtime/GameplayAbilities/Source/GameplayAbilitiesEditor/Private/AssetTypeActions_GameplayAbilitiesBlueprint.h b/Engine/Plugins/Runtime/GameplayAbilities/Source/GameplayAbilitiesEditor/Private/AssetTypeActions_GameplayAbilitiesBlueprint.h
deleted file mode 100644
diff --git a/Engine/Plugins/Runtime/GameplayAbilities/Source/GameplayAbilitiesEditor/Private/GameplayAbilitiesEditorModule.cpp b/Engine/Plugins/Runtime/GameplayAbilities/Source/GameplayAbilitiesEditor/Private/GameplayAbilitiesEditorModule.cpp
--- a/Engine/Plugins/Runtime/GameplayAbilities/Source/GameplayAbilitiesEditor/Private/GameplayAbilitiesEditorModule.cpp
+++ b/Engine/Plugins/Runtime/GameplayAbilities/Source/GameplayAbilitiesEditor/Private/GameplayAbilitiesEditorModule.cpp
@@ -33,7 +33,6 @@
 
 #include "AssetToolsModule.h"
 #include "IAssetTypeActions.h"
-#include "AssetTypeActions_GameplayAbilitiesBlueprint.h"
 #include "EdGraphUtilities.h"
 #include "GameplayAbilitiesGraphPanelPinFactory.h"
 #include "GameplayCueTagDetails.h"
@@ -90,7 +89,6 @@ class FGameplayAbilitiesEditorModule : public IGameplayAbilitiesEditorModule
 	}
 
 protected:
-	void RegisterAssetTypeAction(class IAssetTools& AssetTools, TSharedRef<IAssetTypeActions> Action);
 	static void GameplayTagTreeChanged();
 
 private:
@@ -98,9 +96,6 @@ private:
 	/** Helper function to apply the gameplay mod evaluation channel aliases as display name data to the enum */
 	void ApplyGameplayModEvaluationChannelAliasesToEnumMetadata();
 
-	/** All created asset type actions.  Cached here so that we can unregister it during shutdown. */
-	TArray< TSharedPtr<IAssetTypeActions> > CreatedAssetTypeActions;
-
 	/** Pin factory for abilities graph; Cached so it can be unregistered */
 	TSharedPtr<FGameplayAbilitiesGraphPanelPinFactory> GameplayAbilitiesGraphPanelPinFactory;
 
@@ -185,11 +180,6 @@ void FGameplayAbilitiesEditorModule::StartupModule()
 	PropertyModule.RegisterCustomClassLayout( "AttributeSet", FOnGetDetailCustomizationInstance::CreateStatic( &FAttributeDetails::MakeInstance ) );
 	PropertyModule.RegisterCustomClassLayout( "GameplayEffect", FOnGetDetailCustomizationInstance::CreateStatic( &FGameplayEffectDetails::MakeInstance ) );
 
-	// Register asset types
-	IAssetTools& AssetTools = FModuleManager::LoadModuleChecked<FAssetToolsModule>("AssetTools").Get();
-	TSharedRef<IAssetTypeActions> GABAction = MakeShareable(new FAssetTypeActions_GameplayAbilitiesBlueprint());
-	RegisterAssetTypeAction(AssetTools, GABAction);
-
 	if (ISettingsModule* SettingsModule = FModuleManager::GetModulePtr<ISettingsModule>("Settings"))
 	{
 		SettingsModule->RegisterSettings("Project", "Project", "Gameplay Effect Parents",
@@ -286,12 +276,6 @@ void FGameplayAbilitiesEditorModule::RegisterDebuggingCallbacks()
 	});
 }
 
-void FGameplayAbilitiesEditorModule::RegisterAssetTypeAction(IAssetTools& AssetTools, TSharedRef<IAssetTypeActions> Action)
-{
-	AssetTools.RegisterAssetTypeActions(Action);
-	CreatedAssetTypeActions.Add(Action);
-}
-
 void FGameplayAbilitiesEditorModule::GameplayTagTreeChanged()
 {
 	if (FBlueprintActionDatabase* BAD = FBlueprintActionDatabase::TryGet())
@@ -330,20 +314,6 @@ void FGameplayAbilitiesEditorModule::ShutdownModule()
 		PropertyModule.UnregisterCustomPropertyTypeLayout("GameplayAttribute");
 	}
 
-	// Unregister asset type actions
-	if (FModuleManager::Get().IsModuleLoaded("AssetTools"))
-	{
-		IAssetTools& AssetToolsModule = FModuleManager::GetModuleChecked<FAssetToolsModule>("AssetTools").Get();
-		for (auto& AssetTypeAction : CreatedAssetTypeActions)
-		{
-			if (AssetTypeAction.IsValid())
-			{
-				AssetToolsModule.UnregisterAssetTypeActions(AssetTypeAction.ToSharedRef());
-			}
-		}
-	}
-	CreatedAssetTypeActions.Empty();
-
 	// Unregister graph factories
 	if (GameplayAbilitiesGraphPanelPinFactory.IsValid())
 	{
