[GameplayTags][GameplayCue] Performance improvements to tag children lookup and cue set management

Add an LRU cache to RequestGameplayTagChildren (size configurable via
GameplayTags.LruCacheSize, default 3) to avoid redundant tree traversals
on repeated calls for the same tag.

Convert AddChildrenTags from recursive to iterative (queue-based) to
eliminate call stack growth proportional to tag hierarchy depth.

Replace the O(n) linear duplicate scan in UGameplayCueSet::AddCues with
an O(1) lookup against the existing GameplayCueDataMap, and implement the
previously empty RemoveCuesByTags using iterator-based removal followed
by an acceleration map rebuild.

Change BuildCuesToAddToGlobalSet to accept TConstArrayView instead of
TArray to avoid unnecessary copies at call sites.

#Synced-CL 52557824

diff --git a/Engine/Plugins/Runtime/GameplayAbilities/Source/GameplayAbilities/Private/GameplayCueManager.cpp b/Engine/Plugins/Runtime/GameplayAbilities/Source/GameplayAbilities/Private/GameplayCueManager.cpp
--- a/Engine/Plugins/Runtime/GameplayAbilities/Source/GameplayAbilities/Private/GameplayCueManager.cpp
+++ b/Engine/Plugins/Runtime/GameplayAbilities/Source/GameplayAbilities/Private/GameplayCueManager.cpp
@@ -773,11 +773,15 @@ TArray<FString> UGameplayCueManager::GetAlwaysLoadedGameplayCuePaths()
 void UGameplayCueManager::AddGameplayCueNotifyPath(const FString& InPath, const bool bShouldRescanCueAssets /* = true */)
 {
 	UAbilitySystemGlobals::Get().AddGameplayCueNotifyPath(InPath);
-	const int32 NumAdded = RuntimeGameplayCueObjectLibrary.Paths.AddUnique(InPath);
-	
-	if(bShouldRescanCueAssets && NumAdded != INDEX_NONE)
+	const bool bAlreadyExists = RuntimeGameplayCueObjectLibrary.Paths.Find(InPath) != INDEX_NONE;
+	if (!bAlreadyExists)
 	{
-		InitializeRuntimeObjectLibrary();		
+		RuntimeGameplayCueObjectLibrary.Paths.Add(InPath);
+
+		if(bShouldRescanCueAssets)
+		{
+			InitializeRuntimeObjectLibrary();
+		}
 	}
 }
 
@@ -935,7 +939,7 @@ TSharedPtr<FStreamableHandle> UGameplayCueManager::InitObjectLibrary(FGameplayCu
 
 static FAutoConsoleVariable CVarGameplyCueAddToGlobalSetDebug(TEXT("GameplayCue.AddToGlobalSet.DebugTag"), TEXT(""), TEXT("Debug Tag adding to global set"), ECVF_Default	);
 
-void UGameplayCueManager::BuildCuesToAddToGlobalSet(const TArray<FAssetData>& AssetDataList, FName TagPropertyName, TArray<FGameplayCueReferencePair>& OutCuesToAdd, TArray<FSoftObjectPath>& OutAssetsToLoad, FShouldLoadGCNotifyDelegate ShouldLoad)
+void UGameplayCueManager::BuildCuesToAddToGlobalSet(TConstArrayView<FAssetData> AssetDataList, FName TagPropertyName, TArray<FGameplayCueReferencePair>& OutCuesToAdd, TArray<FSoftObjectPath>& OutAssetsToLoad, FShouldLoadGCNotifyDelegate ShouldLoad)
 {
 	UGameplayTagsManager& Manager = UGameplayTagsManager::Get();
 
diff --git a/Engine/Plugins/Runtime/GameplayAbilities/Source/GameplayAbilities/Private/GameplayCueSet.cpp b/Engine/Plugins/Runtime/GameplayAbilities/Source/GameplayAbilities/Private/GameplayCueSet.cpp
--- a/Engine/Plugins/Runtime/GameplayAbilities/Source/GameplayAbilities/Private/GameplayCueSet.cpp
+++ b/Engine/Plugins/Runtime/GameplayAbilities/Source/GameplayAbilities/Private/GameplayCueSet.cpp
@@ -8,7 +8,9 @@
 #include "GameplayCueNotify_Actor.h"
 #include "GameplayCueNotify_Static.h"
 #include "GameplayCueManager.h"
+#include "Misc/CoreMiscDefines.h"
 #include "NativeGameplayTags.h"
+#include "ProfilingDebugging/CpuProfilerTrace.h"
 
 #include UE_INLINE_GENERATED_CPP_BY_NAME(GameplayCueSet)
 
@@ -94,36 +96,48 @@ void UGameplayCueSet::AddCues(const TArray<FGameplayCueReferencePair>& CuesToAdd
 {
 	if (CuesToAdd.Num() > 0)
 	{
+		GameplayCueData.Reserve(GameplayCueData.Num() + CuesToAdd.Num());
+		GameplayCueDataMap.Reserve(GameplayCueDataMap.Num() + CuesToAdd.Num());
 		for (const FGameplayCueReferencePair& CueRefPair : CuesToAdd)
 		{
 			const FGameplayTag& GameplayCueTag = CueRefPair.GameplayCueTag;
 			const FSoftObjectPath& StringRef = CueRefPair.StringRef;
 
-			// Check for duplicates: we may want to remove this eventually (allow multiple GC notifies to handle same event)
-			bool bDupe = false;
-			for (FGameplayCueNotifyData& Data : GameplayCueData)
+			const int32* ExistingId = GameplayCueDataMap.Find(GameplayCueTag);
+			if (ExistingId && *ExistingId != INDEX_NONE)
 			{
-				if (Data.GameplayCueTag == GameplayCueTag)
+				const FGameplayCueNotifyData& ExistingData = GameplayCueData[*ExistingId];
+				// GameplayCueDataMap also have a redirection for child tags, but we allow to override those
+				if (ExistingData.GameplayCueTag == GameplayCueTag)
 				{
-					if (StringRef.ToString() != Data.GameplayCueNotifyObj.ToString())
+					TStringBuilder<256> AddingCueName;
+					TStringBuilder<256> ExistingCueName;
+					StringRef.ToString(AddingCueName);
+					ExistingData.GameplayCueNotifyObj.ToString(ExistingCueName);
+					if (AddingCueName.ToView() != ExistingCueName.ToView())
 					{
-						ABILITY_LOG(Warning, TEXT("Multiple gameplay cues are triggered using tag '%s'. No two gameplay cues may be triggered by the same tag. Cue '%s' and '%s' are triggered on this tag. Skipping."), *GameplayCueTag.ToString(), *Data.GameplayCueNotifyObj.ToString(), *StringRef.ToString());
+						ABILITY_LOG(Warning, 
+							TEXT("Multiple gameplay cues are triggered using tag '%s'. ")
+							TEXT("No two gameplay cues may be triggered by the same tag. ")
+							TEXT("Cue '%s' and '%s' are triggered on this tag. Skipping."), 
+							*GameplayCueTag.ToString(), 
+							*ExistingCueName, 
+							*AddingCueName);
 					}
-					bDupe = true;
-					break;
+					continue;
 				}
 			}
 
-			if (bDupe)
-			{
-				continue;
-			}
-
 			FGameplayCueNotifyData NewData;
 			NewData.GameplayCueNotifyObj = StringRef;
 			NewData.GameplayCueTag = GameplayCueTag;
 
 			GameplayCueData.Add(NewData);
+
+			// GameplayCueDataMap will be rebuilt after the loop,
+			// so it does not matter if child tags incorrectly point to the parent.
+			// We will use the map only as a duplicate lookup during the loop.
+			GameplayCueDataMap.Add(GameplayCueTag, GameplayCueData.Num() - 1);
 		}
 
 		BuildAccelerationMap_Internal();
@@ -132,6 +146,19 @@ void UGameplayCueSet::AddCues(const TArray<FGameplayCueReferencePair>& CuesToAdd
 
 void UGameplayCueSet::RemoveCuesByTags(const FGameplayTagContainer& TagsToRemove)
 {
+	bool bAnyRemoved = false;
+	for (auto It = GameplayCueData.CreateIterator(); It; ++It)
+	{
+		if (TagsToRemove.HasAnyExact(FGameplayTagContainer{It->GameplayCueTag}))
+		{
+			bAnyRemoved = true;
+			It.RemoveCurrentSwap();
+		}
+	}
+	if (bAnyRemoved)
+	{
+		BuildAccelerationMap_Internal();
+	}
 }
 
 void UGameplayCueSet::RemoveCuesByStringRefs(const TArray<FSoftObjectPath>& CuesToRemove)
@@ -244,10 +271,14 @@ void UGameplayCueSet::PrintCues() const
 
 	for (FGameplayTag ThisGameplayCueTag : AllGameplayCueTags)
 	{
-		int32 idx = GameplayCueDataMap.FindChecked(ThisGameplayCueTag);
-		if (idx != INDEX_NONE)
+		const int32* idx = GameplayCueDataMap.Find(ThisGameplayCueTag);
+		if (!idx)
 		{
-			ABILITY_LOG(Warning, TEXT("   %s -> %d"), *ThisGameplayCueTag.ToString(), idx);
+			ABILITY_LOG(Warning, TEXT("   %s -> index not found"), *ThisGameplayCueTag.ToString());
+		}
+		else if (*idx != INDEX_NONE)
+		{
+			ABILITY_LOG(Warning, TEXT("   %s -> %d"), *ThisGameplayCueTag.ToString(), *idx);
 		}
 		else
 		{
@@ -365,13 +396,14 @@ bool UGameplayCueSet::HandleGameplayCueNotify_Internal(AActor* TargetActor, int3
 
 void UGameplayCueSet::BuildAccelerationMap_Internal()
 {
+	TRACE_CPUPROFILER_EVENT_SCOPE(UGameplayCueSet_BuildAccelerationMap);
 	// ---------------------------------------------------------
 	//	Build up the rest of the acceleration map: every GameplayCue tag should have an entry in the map that points to the index into GameplayCueData to use when it is invoked.
 	//	(or to -1 if no GameplayCueNotify is associated with that tag)
 	// 
 	// ---------------------------------------------------------
 
-	GameplayCueDataMap.Empty();
+	GameplayCueDataMap.Reset();
 	GameplayCueDataMap.Add(BaseGameplayCueTag()) = INDEX_NONE;
 
 	for (int32 idx = 0; idx < GameplayCueData.Num(); ++idx)
@@ -396,6 +428,7 @@ void UGameplayCueSet::BuildAccelerationMap_Internal()
 		int32 ParentValue = GameplayCueDataMap.FindChecked(Parent);
 		GameplayCueDataMap.Add(ThisGameplayCueTag, ParentValue);
 	}
+	GameplayCueDataMap.Shrink();
 
 
 	// Build up parentIdx on each item in GameplayCUeData
diff --git a/Engine/Plugins/Runtime/GameplayAbilities/Source/GameplayAbilities/Public/GameplayCueManager.h b/Engine/Plugins/Runtime/GameplayAbilities/Source/GameplayAbilities/Public/GameplayCueManager.h
--- a/Engine/Plugins/Runtime/GameplayAbilities/Source/GameplayAbilities/Public/GameplayCueManager.h
+++ b/Engine/Plugins/Runtime/GameplayAbilities/Source/GameplayAbilities/Public/GameplayCueManager.h
@@ -278,7 +278,7 @@ protected:
 	UE_API void RefreshObjectLibraries();
 
 	/** Internal function to actually init the FGameplayCueObjectLibrary.  Returns StreamableHandle when asset async loading is requested. */
-	UE_API TSharedPtr<FStreamableHandle> InitObjectLibrary(FGameplayCueObjectLibrary& Library);
+	UE_API virtual TSharedPtr<FStreamableHandle> InitObjectLibrary(FGameplayCueObjectLibrary& Library);
 
 	UE_API virtual TArray<FString> GetAlwaysLoadedGameplayCuePaths();
 
@@ -396,7 +396,7 @@ protected:
 	UE_API virtual void ReloadObjectLibrary(UWorld* World, const UWorld::InitializationValues IVS);
 #endif
 
-	UE_API void BuildCuesToAddToGlobalSet(const TArray<FAssetData>& AssetDataList, FName TagPropertyName, TArray<struct FGameplayCueReferencePair>& OutCuesToAdd, TArray<FSoftObjectPath>& OutAssetsToLoad, FShouldLoadGCNotifyDelegate = FShouldLoadGCNotifyDelegate());
+	UE_API void BuildCuesToAddToGlobalSet(TConstArrayView<FAssetData> AssetDataList, FName TagPropertyName, TArray<struct FGameplayCueReferencePair>& OutCuesToAdd, TArray<FSoftObjectPath>& OutAssetsToLoad, FShouldLoadGCNotifyDelegate = FShouldLoadGCNotifyDelegate());
 
 	/** The cue manager has a tendency to produce a lot of RPCs. This logs out when we are attempting to fire more RPCs than will actually go off */
 	UE_API void CheckForTooManyRPCs(FName FuncName, const FGameplayCuePendingExecute& PendingCue, const FString& CueID, const FGameplayEffectContext* EffectContext);
diff --git a/Engine/Source/Runtime/GameplayTags/Classes/GameplayTagsManager.h b/Engine/Source/Runtime/GameplayTags/Classes/GameplayTagsManager.h
--- a/Engine/Source/Runtime/GameplayTags/Classes/GameplayTagsManager.h
+++ b/Engine/Source/Runtime/GameplayTags/Classes/GameplayTagsManager.h
@@ -3,6 +3,7 @@
 #pragma once
 
 #include "AssetRegistry/AssetData.h"
+#include "Containers/LruCache.h"
 #include "CoreMinimal.h"
 #include "Misc/EnumRange.h"
 #include "Stats/Stats.h"
@@ -1051,4 +1052,6 @@ private:
 	GAMEPLAYTAGS_API const static FName NAME_GameplayTagFilter;
 
 	friend class UGameplayTagsManagerIncrementalCookFunctions;
+
+	mutable TLruCache<FGameplayTag, FGameplayTagContainer> ChildNodesCache;
 };
diff --git a/Engine/Source/Runtime/GameplayTags/Private/GameplayTagsManager.cpp b/Engine/Source/Runtime/GameplayTags/Private/GameplayTagsManager.cpp
--- a/Engine/Source/Runtime/GameplayTags/Private/GameplayTagsManager.cpp
+++ b/Engine/Source/Runtime/GameplayTags/Private/GameplayTagsManager.cpp
@@ -1,6 +1,8 @@
 // Copyright Epic Games, Inc. All Rights Reserved.
 
 #include "GameplayTagsManager.h"
+#include "Containers/AllowShrinking.h"
+#include "Containers/Array.h"
 #include "Engine/Engine.h"
 #include "GameplayTagContainer.h"
 #include "HAL/PlatformFileManager.h"
@@ -9,6 +11,7 @@
 #include "Logging/StructuredLog.h"
 #include "Misc/Paths.h"
 #include "Misc/ScopeLock.h"
+#include "ProfilingDebugging/CpuProfilerTrace.h"
 #include "Stats/StatsMisc.h"
 #include "Misc/ConfigCacheIni.h"
 #include "UObject/UObjectArray.h"
@@ -99,6 +102,12 @@ static FAutoConsoleCommand CMD_DumpGameplayTagSources(
 );
 #endif
 
+static TAutoConsoleVariable<int32> CVarGameplayTagsLruCacheSize(
+	TEXT("GameplayTags.LruCacheSize"),
+	3,
+	TEXT("Accelerate consequent RequestGameplayTagChildren searches at the cost of the memory usage.")
+);
+
 struct FCompareFGameplayTagNodeByTag
 {
 	FORCEINLINE bool operator()(const TSharedPtr<FGameplayTagNode>& A, const TSharedPtr<FGameplayTagNode>& B) const
@@ -1195,6 +1204,8 @@ UE_AUTORTFM_ALWAYS_OPEN void UGameplayTagsManager::InitializeManager()
 	SingletonManager = NewObject<UGameplayTagsManager>(GetTransientPackage(), NAME_None);
 	SingletonManager->AddToRoot();
 
+	SingletonManager->ChildNodesCache.Empty(CVarGameplayTagsLruCacheSize.GetValueOnAnyThread());
+
 	//This is always going to be a synchronous load this early in init, so save some time by not attempting anything async
 	SingletonManager->LoadGameplayTagTables(false);
 	SingletonManager->ConstructGameplayTagTree();
@@ -1467,6 +1478,8 @@ void UGameplayTagsManager::DestroyGameplayTagTree()
 	{
 		Pair.Value.bWasAddedToTree = false;
 	}
+
+	ChildNodesCache.Empty(ChildNodesCache.Max());
 }
 
 int32 UGameplayTagsManager::InsertTagIntoNodeArray(FName Tag, FName FullTag, TSharedPtr<FGameplayTagNode> ParentNode, TArray< TSharedPtr<FGameplayTagNode> >& NodeArray, FName SourceName, const FString& DevComment, bool bIsExplicitTag, bool bIsRestrictedTag, bool bAllowNonRestrictedChildren)
@@ -1541,6 +1554,9 @@ int32 UGameplayTagsManager::InsertTagIntoNodeArray(FName Tag, FName FullTag, TSh
 			// This function is not generically threadsafe.
 			UE::TConditionalScopeLock Lock(GameplayTagMapCritical, !bIsConstructingGameplayTagTree);
 			GameplayTagNodeMap.Add(GameplayTag, TagNode);
+
+			// Invalidate the cache behind the lock
+			ChildNodesCache.Empty(ChildNodesCache.Max());
 		}
 	}
 
@@ -2830,12 +2846,26 @@ void UGameplayTagsManager::RequestAllGameplayTags(FGameplayTagContainer& TagCont
 
 FGameplayTagContainer UGameplayTagsManager::RequestGameplayTagChildren(const FGameplayTag& GameplayTag) const
 {
+	TRACE_CPUPROFILER_EVENT_SCOPE(UGameplayTagsManager_RequestGameplayTagChildren);
+	UE::TScopeLock Lock(GameplayTagMapCritical);
+
 	FGameplayTagContainer TagContainer;
 	// Note this purposefully does not include the passed in GameplayTag in the container.
 	TSharedPtr<FGameplayTagNode> GameplayTagNode = FindTagNode(GameplayTag);
 	if (GameplayTagNode.IsValid())
 	{
+		if (FGameplayTagContainer* Cached = ChildNodesCache.FindAndTouch(GameplayTag))
+		{
+			return *Cached;
+		}
+
 		AddChildrenTags(TagContainer, GameplayTagNode, true, false);
+
+		ChildNodesCache.SetMaxNumElements(CVarGameplayTagsLruCacheSize.GetValueOnAnyThread());
+		if (ChildNodesCache.Max() > 0)
+		{
+			ChildNodesCache.Add(GameplayTag, TagContainer);
+		}
 	}
 	return TagContainer;
 }
@@ -2856,35 +2886,41 @@ FGameplayTag UGameplayTagsManager::RequestGameplayTagDirectParent(const FGamepla
 
 void UGameplayTagsManager::AddChildrenTags(FGameplayTagContainer& TagContainer, TSharedPtr<FGameplayTagNode> GameplayTagNode, bool RecurseAll, bool OnlyIncludeDictionaryTags) const
 {
-	if (GameplayTagNode.IsValid())
+	TRACE_CPUPROFILER_EVENT_SCOPE(UGameplayTagsManager_AddChildrenTags);
+
+	TArray<FGameplayTag> TagsToAdd;
+	TArray<TSharedPtr<FGameplayTagNode>> Queue;
+	Queue.Push(GameplayTagNode);
+	bool bFirstOne = true;
+	while (!Queue.IsEmpty())
 	{
-		TArray< TSharedPtr<FGameplayTagNode> >& ChildrenNodes = GameplayTagNode->GetChildTagNodes();
-		for (TSharedPtr<FGameplayTagNode> ChildNode : ChildrenNodes)
+		TSharedPtr<FGameplayTagNode> CurrentNode = Queue.Pop(EAllowShrinking::No);
+		if (CurrentNode.IsValid())
 		{
-			if (ChildNode.IsValid())
-			{
-				bool bShouldInclude = true;
+			bool bShouldInclude = !bFirstOne;
 
 #if WITH_EDITORONLY_DATA
-				if (OnlyIncludeDictionaryTags && !ChildNode->IsExplicitTag())
-				{
-					// Only have info to do this in editor builds
-					bShouldInclude = false;
-				}
-#endif	
-				if (bShouldInclude)
-				{
-					TagContainer.AddTag(ChildNode->GetCompleteTag());
-				}
+			if (OnlyIncludeDictionaryTags && !CurrentNode->IsExplicitTag())
+			{
+				// Only have info to do this in editor builds
+				bShouldInclude = false;
+			}
+#endif
 
-				if (RecurseAll)
-				{
-					AddChildrenTags(TagContainer, ChildNode, true, OnlyIncludeDictionaryTags);
-				}
+			if (bShouldInclude)
+			{
+				TagsToAdd.Add(CurrentNode->GetCompleteTag());
+			}
+
+			if (bFirstOne || RecurseAll)
+			{
+				Queue.Append(CurrentNode->GetChildTagNodes());
 			}
 
+			bFirstOne = false;
 		}
 	}
+	TagContainer.AppendTags(FGameplayTagContainer::CreateFromArray(TagsToAdd));
 }
 
 void UGameplayTagsManager::SplitGameplayTagFName(const FGameplayTag& Tag, TArray<FName>& OutNames) const
