Added logging on deletion of gameplay effect to help track down rare crash. Logging can be disabled with cvar.

#rnx
#tests Log added, verified compiled and hit
#Synced-CL 52672847

diff --git a/Engine/Plugins/Runtime/GameplayAbilities/Source/GameplayAbilities/Private/GameplayEffect.cpp b/Engine/Plugins/Runtime/GameplayAbilities/Source/GameplayAbilities/Private/GameplayEffect.cpp
--- a/Engine/Plugins/Runtime/GameplayAbilities/Source/GameplayAbilities/Private/GameplayEffect.cpp
+++ b/Engine/Plugins/Runtime/GameplayAbilities/Source/GameplayAbilities/Private/GameplayEffect.cpp
@@ -127,6 +127,12 @@ namespace UE::GameplayEffect
 		TEXT("Default: True.  When enabled we will prevent all the Pending Gameplay Effects Head from pointing to itself through Next pointers during teardown to prevent any double call to delete."),
 		ECVF_Default);
 
+	TAutoConsoleVariable<bool> CVarFixVerboseLogDeletionOfPendingGameplayEffectsHead(
+		TEXT("AbilitySystem.Fix.VerboseLogDeletionOfPendingGameplayEffectsHead"),
+		true,
+		TEXT("Default: True.  When enabled we will verbose log each gameplay effect deletion."),
+		ECVF_Default);
+
 #if WITH_EDITOR
 	namespace EditorOnly
 	{
@@ -3041,6 +3047,11 @@ void FActiveGameplayEffectsContainer::DeletePendingActiveGameplayEffect()
 		}
 	}
 
+	if (CVarFixVerboseLogDeletionOfPendingGameplayEffectsHead.GetValueOnAnyThread())
+	{
+		UE_LOG(LogGameplayEffects, Log, TEXT("Deleting pending gameplay effect: %s. Pending next effect: %s. Owner: %s."), *PendingGameplayEffectHead->GetDebugString(), PendingGameplayEffectHead->PendingNext ? *PendingGameplayEffectHead->PendingNext->GetDebugString() : TEXT("Next is null"), Owner ? *Owner->GetFullName() : TEXT("Owner is null"));
+	}
+
 	FActiveGameplayEffect* Next = PendingGameplayEffectHead->PendingNext;
 	delete PendingGameplayEffectHead;
 	PendingGameplayEffectHead = Next;
